[I/O Dispatcher] Commandlet for creating container files.
- moved ZenCreator to IoStoreUtilities
- updated staging automation script to run IoStore commandlet when running with [-iostore]

diff --git a/Engine/Source/Developer/IoStoreUtilities/IoStoreUtilities.Build.cs b/Engine/Source/Developer/IoStoreUtilities/IoStoreUtilities.Build.cs
new file mode 100644
--- /dev/null
+++ b/Engine/Source/Developer/IoStoreUtilities/IoStoreUtilities.Build.cs
@@ -0,0 +1,22 @@
+// Copyright 1998-2019 Epic Games, Inc. All Rights Reserved.
+
+using UnrealBuildTool;
+
+public class IoStoreUtilities : ModuleRules
+{
+	public IoStoreUtilities (ReadOnlyTargetRules Target) : base(Target)
+	{
+        PrivateIncludePathModuleNames.AddRange(new string[] {
+			"TargetPlatform",
+		});
+
+		PrivateDependencyModuleNames.AddRange(new string[] {
+            "Core",
+            "CoreUObject",
+            "Projects",
+        });
+        PrivateDependencyModuleNames.Add("PakFile");
+        PrivateDependencyModuleNames.Add("Json");
+        PrivateDependencyModuleNames.Add("RSA");
+	}
+}
diff --git a/Engine/Source/Developer/IoStoreUtilities/Private/IoStoreUtilities.cpp b/Engine/Source/Developer/IoStoreUtilities/Private/IoStoreUtilities.cpp
new file mode 100644
--- /dev/null
+++ b/Engine/Source/Developer/IoStoreUtilities/Private/IoStoreUtilities.cpp
@@ -0,0 +1,1821 @@
+// Copyright 1998-2019 Epic Games, Inc. All Rights Reserved.
+
+#include "IoStoreUtilities.h"
+#include "HAL/FileManager.h"
+#include "Hash/CityHash.h"
+#include "Interfaces/IPluginManager.h"
+#include "Interfaces/ITargetPlatformManagerModule.h"
+#include "Interfaces/ITargetPlatform.h"
+#include "IO/IoDispatcher.h"
+#include "Misc/App.h"
+#include "Misc/CommandLine.h"
+#include "Misc/PackageName.h"
+#include "Misc/Paths.h"
+#include "Misc/PackageName.h"
+#include "Misc/Base64.h"
+#include "Modules/ModuleManager.h"
+#include "Serialization/Archive.h"
+#include "Serialization/JsonReader.h"
+#include "Serialization/JsonSerializer.h"
+#include "Serialization/BufferWriter.h"
+#include "UObject/PackageFileSummary.h"
+#include "UObject/ObjectResource.h"
+#include "UObject/Package.h"
+#include "UObject/SavePackage.h"
+
+#include "Algo/Find.h"
+//PRAGMA_DISABLE_OPTIMIZATION
+
+IMPLEMENT_MODULE(FDefaultModuleImpl, IoStoreUtilities);
+
+DEFINE_LOG_CATEGORY_STATIC(LogIoStore, Log, All);
+
+#define OUTPUT_CHUNKID_DIRECTORY 0
+#define SKIP_WRITE_CONTAINER 0
+#define SKIP_BULKDATA 0
+
+struct FContainerTarget 
+{
+	ITargetPlatform* TargetPlatform;
+	FString CookedDirectory;
+	FString OutputDirectory;
+};
+
+class FZenGlobalNameMap
+{
+public:
+	void Load(const FString& FilePath)
+	{
+		UE_LOG(LogIoStore, Display, TEXT("Loading global name map from '%s' for container files..."), *FilePath);
+
+		TUniquePtr<FArchive> Ar(IFileManager::Get().CreateFileReader(*FilePath));
+		check(Ar);
+
+		int32 NameCount;
+		*Ar << NameCount;
+
+		DisplayEntries.Reserve(NameCount);
+		ComparisonEntries.Reserve(NameCount);
+
+		FNameEntrySerialized SerializedNameEntry(ENAME_LinkerConstructor);
+
+		for (int32 I = 0; I < NameCount; ++I)
+		{
+			*Ar << SerializedNameEntry;
+			FName Name(SerializedNameEntry);
+
+			DisplayEntries.Emplace(Name.GetDisplayIndex());
+			DisplayEntryToIndex.Emplace(DisplayEntries[I], I);
+
+			ComparisonEntries.Emplace(Name.GetComparisonIndex());
+			ComparisonEntryToIndex.Emplace(ComparisonEntries[I], I);
+		}
+	}
+
+	void Save(const FString& FilePath)
+	{
+		UE_LOG(LogIoStore, Display, TEXT("Saving Container name map to '%s' with '%d' additional names"),
+			*FilePath,
+			ComparisonEntries.Num() - DisplayEntries.Num());
+
+		TUniquePtr<FArchive> Ar(IFileManager::Get().CreateFileWriter(*FilePath));
+		check(Ar);
+
+		int32 NameCount = ComparisonEntries.Num();
+		*Ar << NameCount;
+		
+		for (int32 I = 0; I < NameCount; ++I)
+		{
+			FName::GetEntry(ComparisonEntries[I])->Write(*Ar);
+		}
+	}
+
+	FName GetNameFromDisplayIndex(const uint32 DisplayIndex, const uint32 NameNumber) const
+	{
+		FNameEntryId DisplayEntry = DisplayEntries[DisplayIndex];
+		return FName::CreateFromDisplayId(DisplayEntry, NameNumber);
+	}
+
+	const int32* GetComparisonIndex(const FName& Name) const
+	{
+		return ComparisonEntryToIndex.Find(Name.GetComparisonIndex());
+	}
+
+	int32 GetOrCreateComparisonIndex(const FName& Name)
+	{
+		if (const int32* ExistingIndex = ComparisonEntryToIndex.Find(Name.GetComparisonIndex()))
+		{
+			return *ExistingIndex;
+		}
+		else
+		{
+			int32 NewIndex = ComparisonEntries.Add(Name.GetComparisonIndex());
+			ComparisonEntryToIndex.Add(ComparisonEntries[NewIndex], NewIndex);
+			return NewIndex;
+		}
+	}
+private:
+	TArray<FNameEntryId> DisplayEntries;
+	TArray<FNameEntryId> ComparisonEntries;
+	TMap<FNameEntryId, int32> DisplayEntryToIndex;
+	TMap<FNameEntryId, int32> ComparisonEntryToIndex;
+};
+
+#if OUTPUT_CHUNKID_DIRECTORY
+class FChunkIdCsv
+{
+public:
+
+	~FChunkIdCsv()
+	{
+		if (OutputArchive)
+		{
+			OutputArchive->Flush();
+		}
+	}
+
+	void CreateOutputFile(const FString& RootPath)
+	{
+		const FString OutputFilename = RootPath / TEXT("chunkid_directory.csv");
+		OutputArchive.Reset(IFileManager::Get().CreateFileWriter(*OutputFilename));
+		if (OutputArchive)
+		{
+			const ANSICHAR* Output = "NameIndex,NameNumber,ChunkIndex,ChunkType,ChunkIdHash,DebugString\n";
+			OutputArchive->Serialize((void*)Output, FPlatformString::Strlen(Output));
+		}
+	}
+
+	void AddChunk(uint32 NameIndex, uint32 NameNumber, uint16 ChunkIndex, uint8 ChunkType, uint32 ChunkIdHash, const TCHAR* DebugString)
+	{
+		ANSICHAR Buffer[MAX_SPRINTF + 1] = { 0 };
+		int32 NumOfCharacters = FCStringAnsi::Sprintf(Buffer, "%u,%u,%u,%u,%u,%s\n", NameIndex, NameNumber, ChunkIndex, ChunkType, ChunkIdHash, TCHAR_TO_ANSI(DebugString));
+		OutputArchive->Serialize(Buffer, NumOfCharacters);
+	}	
+
+private:
+	TUniquePtr<FArchive> OutputArchive;
+};
+FChunkIdCsv ChunkIdCsv;
+
+#endif
+
+enum class EZenChunkType : uint8
+{
+	None,
+	PackageSummary,
+	ExportData,
+	ExportBundleData,
+	BulkData
+};
+
+static FIoChunkId CreateZenChunkId(uint32 NameIndex, uint32 NameNumber, uint16 ChunkIndex, EZenChunkType ChunkType, const TCHAR* DebugString)
+{
+	uint8 Data[12] = {0};
+
+	*reinterpret_cast<uint32*>(&Data[0]) = NameIndex;
+	*reinterpret_cast<int32*>(&Data[4]) = NameNumber;
+	*reinterpret_cast<uint16*>(&Data[8]) = ChunkIndex;
+	*reinterpret_cast<uint8*>(&Data[10]) = static_cast<uint8>(ChunkType);
+
+	FIoChunkId ChunkId;
+	ChunkId.Set(Data, 12);
+	
+#if OUTPUT_CHUNKID_DIRECTORY
+	ChunkIdCsv.AddChunk(NameIndex, NameNumber, ChunkIndex, (uint8)ChunkType, GetTypeHash(ChunkId), DebugString);
+#endif
+
+	return ChunkId;
+}
+
+struct FZenPackageSummary
+{
+	uint32 PackageFlags;
+	int32 ExportMapOffset;
+	int32 ExportBundlesOffset;
+	int32 GraphDataOffset;
+	int32 GraphDataSize;
+	int32 BulkDataStartOffset;
+};
+
+enum EPreloadDependencyType
+{
+	PreloadDependencyType_Create,
+	PreloadDependencyType_Serialize,
+};
+
+enum EEventLoadNode2 : uint8
+{
+	Package_CreateLinker,
+	Package_ProcessSummary,
+	Package_ExportsSerialized,
+	Package_StartPostLoad,
+	Package_Tick,
+	Package_Delete,
+	Package_NumPhases,
+
+	ExportBundle_StartIo = 0,
+	ExportBundle_Process,
+	ExportBundle_NumPhases,
+};
+
+struct FArc
+{
+	uint32 FromNodeIndex;
+	uint32 ToNodeIndex;
+
+	bool operator==(const FArc& Other) const
+	{
+		return ToNodeIndex == Other.ToNodeIndex && FromNodeIndex == Other.FromNodeIndex;
+	}
+};
+
+struct FExportBundleEntry
+{
+	enum EExportCommandType
+	{
+		ExportCommandType_Create,
+		ExportCommandType_Serialize
+	};
+	uint32 LocalExportIndex;
+	EExportCommandType CommandType;
+};
+
+struct FExportGraphNode;
+
+struct FExportBundle
+{
+	TArray<FExportGraphNode*> Nodes;
+	TSet<uint32> AddedScriptDependencies;
+};
+
+struct FPackage;
+
+struct FPackageGraphNode
+{
+	FPackage* Package = nullptr;
+	bool bTemporaryMark = false;
+	bool bPermanentMark = false;
+};
+
+class FPackageGraph
+{
+public:
+	FPackageGraph()
+	{
+
+	}
+
+	~FPackageGraph()
+	{
+		for (FPackageGraphNode* Node : Nodes)
+		{
+			delete Node;
+		}
+	}
+
+	FPackageGraphNode* AddNode(FPackage* Package)
+	{
+		FPackageGraphNode* Node = new FPackageGraphNode();
+		Node->Package = Package;
+		Nodes.Add(Node);
+		return Node;
+	}
+
+	void AddImportDependency(FPackageGraphNode* FromNode, FPackageGraphNode* ToNode)
+	{
+		Edges.Add(FromNode, ToNode);
+	}
+
+	TArray<FPackage*> TopologicalSort() const;
+
+private:
+	TArray<FPackageGraphNode*> Nodes;
+	TMultiMap<FPackageGraphNode*, FPackageGraphNode*> Edges;
+};
+
+struct FExportGraphNode
+{
+	FPackage* Package;
+	FExportBundleEntry BundleEntry;
+	TSet<FExportGraphNode*> ExternalDependencies;
+	TSet<uint32> ScriptDependencies;
+	uint64 NodeIndex;
+};
+
+class FExportGraph
+{
+public:
+	FExportGraph()
+	{
+
+	}
+
+	~FExportGraph()
+	{
+		for (FExportGraphNode* Node : Nodes)
+		{
+			delete Node;
+		}
+	}
+
+	FExportGraphNode* AddNode(FPackage* Package, const FExportBundleEntry& BundleEntry)
+	{
+		FExportGraphNode* Node = new FExportGraphNode();
+		Node->Package = Package;
+		Node->BundleEntry = BundleEntry;
+		Node->NodeIndex = Nodes.Num();
+		Nodes.Add(Node);
+		return Node;
+	}
+
+	void AddInternalDependency(FExportGraphNode* FromNode, FExportGraphNode* ToNode)
+	{
+		AddEdge(FromNode, ToNode);
+	}
+
+	void AddExternalDependency(FExportGraphNode* FromNode, FExportGraphNode* ToNode)
+	{
+		AddEdge(FromNode, ToNode);
+		ToNode->ExternalDependencies.Add(FromNode);
+	}
+
+	TArray<FExportGraphNode*> ComputeLoadOrder(TMap<FName, FPackage*>& PackageMap) const;
+
+private:
+	void AddEdge(FExportGraphNode* FromNode, FExportGraphNode* ToNode)
+	{
+		Edges.Add(FromNode, ToNode);
+	}
+
+	TArray<FExportGraphNode*> TopologicalSort() const;
+
+	TArray<FExportGraphNode*> Nodes;
+	TMultiMap<FExportGraphNode*, FExportGraphNode*> Edges;
+};
+
+struct FPackage
+{
+	FName Name;
+	FString FileName;
+	FString RelativeFileName;
+	int32 GlobalPackageId = 0;
+	uint32 PackageFlags = 0;
+	int32 ImportCount = 0;
+	int32 ImportOffset = 0;
+	int32 ScriptArcsCount = 0;
+	int32 ScriptArcsOffset = 0;
+	int32 SlimportCount = 0;
+	int32 SlimportOffset = -1;
+	int32 ExportCount = 0;
+	int32 ExportIndexOffset = -1;
+	int32 PreloadIndexOffset = -1;
+	int64 BulkDataStartOffset = -1;
+	int64 UExpSize = 0;
+	int64 UAssetSize = 0;
+	int64 SummarySize = 0;
+	int64 UGraphSize = 0;
+	int64 ExportMapSize = 0;
+	int64 ExportBundlesSize = 0;
+
+	bool bHasCircularImportDependencies = false;
+	bool bHasCircularPostLoadDependencies = false;
+
+	TSet<FName> ImportedPackageNames;
+	TArray<FPackage*> ImportedPackages;
+	TSet<FPackage*> AllReachablePackages;
+	TSet<FPackage*> ImportedPreloadPackages;
+
+	TArray<int32> Imports;
+	TArray<int32> Exports;
+	TArray<FArc> InternalArcs;
+	TMap<FName, TArray<FArc>> ExternalArcs;
+	TArray<FArc> ScriptArcs;
+	
+	TArray<FExportBundle> ExportBundles;
+	TMap<FExportGraphNode*, uint32> ExportBundleMap;
+
+	TArray<FExportGraphNode*> CreateExportNodes;
+	TArray<FExportGraphNode*> SerializeExportNodes;
+
+	uint64 LoadOrder = uint64(-1);
+	TArray<FExportGraphNode*> NodesWithNoIncomingEdges;
+	FPackageGraphNode* Node = nullptr;
+};
+
+struct FCircularImportChain
+{
+	TArray<FName> SortedNames;
+	TArray<FPackage*> Packages;
+	uint32 Hash = 0;
+
+	FCircularImportChain()
+	{
+		Packages.Reserve(128);
+	}
+
+	void Add(FPackage* Package)
+	{
+		Packages.Add(Package);
+	}
+
+	void Pop()
+	{
+		Packages.Pop();
+	}
+
+	int32 Num()
+	{
+		return Packages.Num();
+	}
+
+	void SortAndGenerateHash()
+	{
+		SortedNames.Empty(Packages.Num());
+		for (FPackage* Package : Packages)
+		{
+			SortedNames.Emplace(Package->Name);
+		}
+		SortedNames.Sort(FNameLexicalLess());
+		Hash = CityHash32((char*)SortedNames.GetData(), SortedNames.Num() * SortedNames.GetTypeSize());
+	}
+
+	FString ToString()
+	{
+		FString Result = FString::Printf(TEXT("%d:%u: "), SortedNames.Num(), Hash);
+		for (const FName& Name : SortedNames)
+		{
+			Result.Append(Name.ToString());
+			Result.Append(TEXT(" -> "));
+		}
+		Result.Append(SortedNames[0].ToString());
+		return Result;
+	}
+
+	bool operator==(const FCircularImportChain& Other) const
+	{
+		return Hash == Other.Hash && SortedNames == Other.SortedNames;
+	}
+
+	friend FORCEINLINE uint32 GetTypeHash(const FCircularImportChain& In)
+	{
+		return In.Hash;
+	}
+};
+
+TArray<FPackage*> FPackageGraph::TopologicalSort() const
+{
+	TMultiMap<FPackageGraphNode*, FPackageGraphNode*> EdgesCopy = Edges;
+	TArray<FPackage*> Result;
+	Result.Reserve(Nodes.Num());
+	
+	TSet<FPackageGraphNode*> UnmarkedNodes;
+	UnmarkedNodes.Append(Nodes);
+
+	struct
+	{
+		void Visit(FPackageGraphNode* Node)
+		{
+			if (Node->bPermanentMark)
+			{
+				return;
+			}
+			if (Node->bTemporaryMark)
+			{
+				return;
+			}
+			Node->bTemporaryMark = true;
+			for (auto EdgeIt = Edges.CreateKeyIterator(Node); EdgeIt; ++EdgeIt)
+			{
+				FPackageGraphNode* ToNode = EdgeIt.Value();
+				Visit(ToNode);
+			}
+			Node->bTemporaryMark = false;
+			Node->bPermanentMark = true;
+			UnmarkedNodes.Remove(Node);
+			Result.Insert(Node->Package, 0);
+		}
+
+		TSet<FPackageGraphNode*>& UnmarkedNodes;
+		TMultiMap<FPackageGraphNode*, FPackageGraphNode*>& Edges;
+		TArray<FPackage*>& Result;
+
+	} Visitor{ UnmarkedNodes, EdgesCopy, Result };
+
+	while (Result.Num() < Nodes.Num())
+	{
+		auto It = UnmarkedNodes.CreateIterator();
+		FPackageGraphNode* UnmarkedNode = *It;
+		It.RemoveCurrent();
+		Visitor.Visit(UnmarkedNode);
+	}
+
+	return Result;
+}
+
+TArray<FExportGraphNode*> FExportGraph::ComputeLoadOrder(TMap<FName, FPackage*>& PackageMap) const
+{
+	FPackageGraph PackageGraph;
+	for (auto& KV : PackageMap)
+	{
+		FPackage* Package = KV.Value;
+		Package->Node = PackageGraph.AddNode(Package);
+	}
+	for (auto& KV : PackageMap)
+	{
+		FPackage* Package = KV.Value;
+		
+		for (FPackage* ImportedPackage : Package->ImportedPackages)
+		{
+			PackageGraph.AddImportDependency(ImportedPackage->Node, Package->Node);
+		}
+	}
+
+	TArray<FPackage*> SortedPackages = PackageGraph.TopologicalSort();
+	
+	int32 NodeCount = Nodes.Num();
+	TArray<uint32> NodesIncomingEdgeCount;
+	NodesIncomingEdgeCount.AddZeroed(NodeCount);
+	TMultiMap<FExportGraphNode*, FExportGraphNode*> EdgesCopy = Edges;
+	for (auto& KV : EdgesCopy)
+	{
+		FExportGraphNode* ToNode = KV.Value;
+		++NodesIncomingEdgeCount[ToNode->NodeIndex];
+	}
+
+	TArray<FExportGraphNode*> LoadOrder;
+	LoadOrder.Reserve(NodeCount);
+	
+	for (int32 NodeIndex = 0; NodeIndex < NodeCount; ++NodeIndex)
+	{
+		if (NodesIncomingEdgeCount[NodeIndex] == 0)
+		{
+			FExportGraphNode* Node = Nodes[NodeIndex];
+			Node->Package->NodesWithNoIncomingEdges.Push(Node);
+		}
+	}
+	while (LoadOrder.Num() < NodeCount)
+	{
+		for (FPackage* Package : SortedPackages)
+		{
+			while (Package->NodesWithNoIncomingEdges.Num() > 0)
+			{
+				FExportGraphNode* RemovedNode = Package->NodesWithNoIncomingEdges.Pop();
+				LoadOrder.Add(RemovedNode);
+				for (auto EdgeIt = EdgesCopy.CreateKeyIterator(RemovedNode); EdgeIt; ++EdgeIt)
+				{
+					FExportGraphNode* ToNode = EdgeIt.Value();
+					if (--NodesIncomingEdgeCount[ToNode->NodeIndex] == 0)
+					{
+						ToNode->Package->NodesWithNoIncomingEdges.Push(ToNode);
+					}
+					EdgeIt.RemoveCurrent();
+				}
+			}
+		}
+	}
+
+	return LoadOrder;
+}
+
+static void AddInternalExportArc(FExportGraph& ExportGraph, FPackage& Package, uint32 FromExportIndex, EPreloadDependencyType FromPhase, uint32 ToExportIndex, EPreloadDependencyType ToPhase)
+{
+	FExportGraphNode* FromNode = FromPhase == PreloadDependencyType_Create ? Package.CreateExportNodes[FromExportIndex] : Package.SerializeExportNodes[FromExportIndex];
+	FExportGraphNode* ToNode = ToPhase == PreloadDependencyType_Create ? Package.CreateExportNodes[ToExportIndex] : Package.SerializeExportNodes[ToExportIndex];
+	ExportGraph.AddInternalDependency(FromNode, ToNode);
+}
+
+static void AddExternalExportArc(FExportGraph& ExportGraph, FPackage& FromPackage, uint32 FromExportIndex, EPreloadDependencyType FromPhase, FPackage& ToPackage, uint32 ToExportIndex, EPreloadDependencyType ToPhase)
+{
+	FExportGraphNode* FromNode = FromPhase == PreloadDependencyType_Create ? FromPackage.CreateExportNodes[FromExportIndex] : FromPackage.SerializeExportNodes[FromExportIndex];
+	FExportGraphNode* ToNode = ToPhase == PreloadDependencyType_Create ? ToPackage.CreateExportNodes[ToExportIndex] : ToPackage.SerializeExportNodes[ToExportIndex];
+	ExportGraph.AddExternalDependency(FromNode, ToNode);
+}
+
+static void AddScriptArc(FPackage& Package, uint32 GlobalImportIndex, uint32 ExportIndex, EPreloadDependencyType Phase)
+{
+	FExportGraphNode* Node = Phase == PreloadDependencyType_Create ? Package.CreateExportNodes[ExportIndex] : Package.SerializeExportNodes[ExportIndex];
+	Node->ScriptDependencies.Add(GlobalImportIndex);
+}
+
+static void AddPostLoadArc(FPackage& FromPackage, FPackage& ToPackage)
+{
+	// TArray<FArc>& ExternalArcs = ToPackage.ExternalArcs.FindOrAdd(FromPackage.Name);
+	// ExternalArcs.Add({ EEventLoadNode2::Package_ExportsSerialized, EEventLoadNode2::Package_Tick });
+}
+
+static void AddStartPostLoadArc(FPackage& FromPackage, FPackage& ToPackage)
+{
+	TArray<FArc>& ExternalArcs = ToPackage.ExternalArcs.FindOrAdd(FromPackage.Name);
+	ensure(!ExternalArcs.Contains(FArc({EEventLoadNode2::Package_ExportsSerialized, EEventLoadNode2::Package_StartPostLoad})));
+	ensure(!ExternalArcs.Contains(FArc({EEventLoadNode2::Package_StartPostLoad, EEventLoadNode2::Package_StartPostLoad})));
+	ExternalArcs.Add({ EEventLoadNode2::Package_StartPostLoad, EEventLoadNode2::Package_StartPostLoad });
+}
+
+static void AddExportsDoneArc(FPackage& FromPackage, FPackage& ToPackage)
+{
+	TArray<FArc>& ExternalArcs = ToPackage.ExternalArcs.FindOrAdd(FromPackage.Name);
+	ensure(!ExternalArcs.Contains(FArc({EEventLoadNode2::Package_ExportsSerialized, EEventLoadNode2::Package_StartPostLoad})));
+	ensure(!ExternalArcs.Contains(FArc({EEventLoadNode2::Package_StartPostLoad, EEventLoadNode2::Package_StartPostLoad})));
+	ExternalArcs.Add({ EEventLoadNode2::Package_ExportsSerialized, EEventLoadNode2::Package_StartPostLoad });
+}
+
+static void AddInternalBundleArc(FPackage& Package, uint32 FromBundleIndex, uint32 ToBundleIndex)
+{
+	uint32 FromNodeIndex = EEventLoadNode2::Package_NumPhases + FromBundleIndex * EEventLoadNode2::ExportBundle_NumPhases + EEventLoadNode2::ExportBundle_Process;
+	uint32 ToNodeIndex = EEventLoadNode2::Package_NumPhases + ToBundleIndex * EEventLoadNode2::ExportBundle_NumPhases + EEventLoadNode2::ExportBundle_StartIo;
+	Package.InternalArcs.Add({ FromNodeIndex, ToNodeIndex });
+}
+
+static void AddExternalBundleArc(FPackage& FromPackage, uint32 FromBundleIndex, FPackage& ToPackage, uint32 ToBundleIndex)
+{
+	uint32 FromNodeIndex = EEventLoadNode2::Package_NumPhases + FromBundleIndex * EEventLoadNode2::ExportBundle_NumPhases + EEventLoadNode2::ExportBundle_Process;
+	uint32 ToNodeIndex = EEventLoadNode2::Package_NumPhases + ToBundleIndex * EEventLoadNode2::ExportBundle_NumPhases + EEventLoadNode2::ExportBundle_StartIo;
+	TArray<FArc>& ExternalArcs = ToPackage.ExternalArcs.FindOrAdd(FromPackage.Name);
+	ExternalArcs.Add({ FromNodeIndex, ToNodeIndex });
+}
+
+static void AddScriptBundleArc(FPackage& Package, uint32 GlobalImportIndex, uint32 BundleIndex)
+{
+	uint32 NodeIndex = EEventLoadNode2::Package_NumPhases + BundleIndex * EEventLoadNode2::ExportBundle_NumPhases + EEventLoadNode2::ExportBundle_StartIo;
+	Package.ScriptArcs.Add({ GlobalImportIndex, NodeIndex });
+}
+
+static void AddPostLoadDependenciesRecursive(FPackage& Package, FPackage& ImportedPackage, TSet<FPackage*>& Visited)
+{
+	if (&ImportedPackage == &Package)
+	{
+		Package.bHasCircularPostLoadDependencies = true;
+		return;
+	}
+
+	bool bIsVisited = false;
+	Visited.Add(&ImportedPackage, &bIsVisited);
+	if (bIsVisited)
+	{
+		return;
+	}
+
+	AddPostLoadArc(ImportedPackage, Package);
+	
+	for (FPackage* DependentPackage : ImportedPackage.ImportedPackages)
+	{
+		AddPostLoadDependenciesRecursive(Package, *DependentPackage, Visited);
+	}
+}
+
+static void AddReachablePackagesRecursive(FPackage& Package, FPackage& PackageWithImports, TSet<FPackage*>& Visited, bool bFirst)
+{
+	if (!bFirst)
+	{
+		bool bIsVisited = false;
+		Visited.Add(&PackageWithImports, &bIsVisited);
+		if (bIsVisited)
+		{
+			return;
+		}
+
+		if (&PackageWithImports == &Package)
+		{
+			return;
+		}
+	}
+
+	if (PackageWithImports.AllReachablePackages.Num() > 0)
+	{
+		Visited.Append(PackageWithImports.AllReachablePackages);
+		
+	}
+	else
+	{
+		for (FPackage* ImportedPackage : PackageWithImports.ImportedPackages)
+		{
+			AddReachablePackagesRecursive(Package, *ImportedPackage, Visited, false);
+		}
+	}
+}
+
+static bool FindNewCircularImportChains(
+	FPackage& Package,
+	FPackage& ImportedPackage,
+	TSet<FPackage*>& Visited,
+	TSet<FCircularImportChain>& CircularChains,
+	FCircularImportChain& CurrentChain)
+{
+	if (&ImportedPackage == &Package)
+	{
+		CurrentChain.SortAndGenerateHash();
+		bool bAlreadyFound = true;
+		CircularChains.AddByHash(CurrentChain.Hash, CurrentChain, &bAlreadyFound);
+
+		if (bAlreadyFound)
+		{
+			// UE_LOG(LogIoStore, Display, TEXT("OLD-IsCircular: %s with %s"), *Package.Name.ToString(), *CurrentChain.ToString());
+			return false;
+		}
+		else
+		{
+			// UE_LOG(LogIoStore, Display, TEXT("NEW-IsCircular: %s with %s"), *Package.Name.ToString(), *CurrentChain.ToString());
+			return true;
+		}
+	}
+
+	bool bIsVisited = false;
+	Visited.Add(&ImportedPackage, &bIsVisited);
+	if (bIsVisited)
+	{
+		return false;
+	}
+
+	bool bFoundNew = false;
+	for (FPackage* DependentPackage : ImportedPackage.ImportedPackages)
+	{
+		CurrentChain.Add(DependentPackage);
+		bFoundNew |= FindNewCircularImportChains(Package, *DependentPackage, Visited, CircularChains, CurrentChain);
+		CurrentChain.Pop();
+	}
+
+	return bFoundNew;
+}
+
+static void AddPostLoadDependencies(
+	FPackage& Package,
+	TSet<FPackage*>& Visited,
+	TSet<FCircularImportChain>& CircularChains)
+{
+	TSet<FPackage*> DependentPackages;
+
+	for (FPackage* ImportedPackage : Package.ImportedPackages)
+	{
+		Visited.Reset();
+		FCircularImportChain CurrentChain;
+		CurrentChain.Add(ImportedPackage);
+		if (FindNewCircularImportChains(Package, *ImportedPackage, Visited, CircularChains, CurrentChain))
+		{
+			DependentPackages.Append(MoveTemp(Visited));
+		}
+	}
+
+	for (FPackage* ImportedPackage : Package.ImportedPackages)
+	{
+		if (!DependentPackages.Contains(ImportedPackage))
+		{
+			AddStartPostLoadArc(*ImportedPackage, Package);
+		}
+	}
+
+	DependentPackages.Remove(&Package);
+	for (FPackage* DependentPackage : DependentPackages)
+	{
+		AddExportsDoneArc(*DependentPackage, Package);
+	}
+
+	/*
+	if (Package.bHasCircularImportDependencies)
+	{
+		int32 diff = Package.AllReachablePackages.Num() - 1 - DependentPackages.Num();
+		if (DependentPackages.Num() == 0)
+		{
+			UE_LOG(LogIoStore, Display, TEXT("OPT-ALL: %s: Skipping %d/%d arcs"), *Package.Name.ToString(),
+				diff,
+				Package.AllReachablePackages.Num() - 1);
+		}
+		else if (diff > 0)
+		{
+			UE_LOG(LogIoStore, Display, TEXT("OPT: %s: Skipping %d/%d arcs"), *Package.Name.ToString(),
+				diff,
+				Package.AllReachablePackages.Num() - 1);
+		}
+		else
+		{
+			UE_LOG(LogIoStore, Display, TEXT("NOP: %s: Skipping %d/%d arcs"), *Package.Name.ToString(),
+				0,
+				Package.AllReachablePackages.Num() - 1);
+		}
+	}
+	*/
+}
+
+static void BuildBundles(FExportGraph& ExportGraph, TMap<FName, FPackage*>& PackageMap)
+{
+	TArray<FExportGraphNode*> ExportLoadOrder = ExportGraph.ComputeLoadOrder(PackageMap);
+	FPackage* LastPackage = nullptr;
+	for (FExportGraphNode* Node : ExportLoadOrder)
+	{
+		FPackage* Package = Node->Package;
+		check(Package);
+
+		if (Package != LastPackage)
+		{
+			Package->ExportBundles.AddDefaulted();
+			if (Package->ExportBundles.Num() > 1)
+			{
+				AddInternalBundleArc(*Package, Package->ExportBundles.Num() - 2, Package->ExportBundles.Num() - 1);
+			}
+		}
+		LastPackage = Package;
+
+		uint32 BundleIndex = Package->ExportBundles.Num() - 1;
+		for (FExportGraphNode* ExternalDependency : Node->ExternalDependencies)
+		{
+			uint32* FindDependentBundleIndex = ExternalDependency->Package->ExportBundleMap.Find(ExternalDependency);
+			check(FindDependentBundleIndex);
+			AddExternalBundleArc(*ExternalDependency->Package, *FindDependentBundleIndex, *Package, BundleIndex);
+		}
+		for (uint32 ScriptDependencyGlobalImportIndex : Node->ScriptDependencies)
+		{
+			bool bAlreadyInSet;
+			Package->ExportBundles[BundleIndex].AddedScriptDependencies.Add(ScriptDependencyGlobalImportIndex, &bAlreadyInSet);
+			if (!bAlreadyInSet)
+			{
+				AddScriptBundleArc(*Package, ScriptDependencyGlobalImportIndex, BundleIndex);
+			}
+		}
+
+		Package->ExportBundles[BundleIndex].Nodes.Add(Node);
+		Package->ExportBundleMap.Add(Node, BundleIndex);
+	}
+}
+
+struct FImportData
+{
+	int32 GlobalIndex = -1;
+	int32 OuterIndex = -1;
+	int32 OutermostIndex = -1;
+	int32 RefCount = 0;
+	FName ObjectName;
+	bool bIsPackage = false;
+	bool bIsScript = false;
+	FString FullName;
+};
+
+struct FExportData
+{
+	int32 GlobalIndex = -1;
+	FName SourcePackageName;
+	FName ObjectName;
+	int32 SourceIndex = -1;
+	int32 GlobalImportIndex = -1;
+	FString FullName;
+
+	FExportGraphNode* CreateNode = nullptr;
+	FExportGraphNode* SerializeNode = nullptr;
+};
+
+static void FindImport(TArray<FImportData>& GlobalImports, TMap<FString, int32>& GlobalImportsByFullName, TArray<FString>& TempFullNames, FObjectImport* ImportMap, int32 LocalImportIndex)
+{
+	FObjectImport* Import = &ImportMap[LocalImportIndex];
+	FString& FullName = TempFullNames[LocalImportIndex];
+	if (FullName.Len() == 0)
+	{
+		if (Import->OuterIndex.IsNull())
+		{
+			Import->ObjectName.AppendString(FullName);
+			int32* FindGlobalImport = GlobalImportsByFullName.Find(FullName);
+			if (!FindGlobalImport)
+			{
+				// first time, assign global index for this root package
+				int32 GlobalImportIndex = GlobalImports.Num();
+				GlobalImportsByFullName.Add(FullName, GlobalImportIndex);
+				FImportData& GlobalImport = GlobalImports.AddDefaulted_GetRef();
+				GlobalImport.GlobalIndex = GlobalImportIndex;
+				GlobalImport.OutermostIndex = GlobalImportIndex;
+				GlobalImport.OuterIndex = -1;
+				GlobalImport.ObjectName = Import->ObjectName;
+				GlobalImport.bIsPackage = true;
+				GlobalImport.bIsScript = FullName.StartsWith(TEXT("/Script/"));
+				GlobalImport.FullName = FullName;
+				GlobalImport.RefCount = 1;
+			}
+			else
+			{
+				++GlobalImports[*FindGlobalImport].RefCount;
+			}
+		}
+		else
+		{
+			int32 LocalOuterIndex = Import->OuterIndex.ToImport();
+			FindImport(GlobalImports, GlobalImportsByFullName, TempFullNames, ImportMap, LocalOuterIndex);
+			FString& OuterName = TempFullNames[LocalOuterIndex];
+			ensure(OuterName.Len() > 0);
+
+			FullName.Append(OuterName);
+			FullName.AppendChar(TEXT('/'));
+			Import->ObjectName.AppendString(FullName);
+
+			int32* FindGlobalImport = GlobalImportsByFullName.Find(FullName);
+			if (!FindGlobalImport)
+			{
+				// first time, assign global index for this intermediate import
+				int32 GlobalImportIndex = GlobalImports.Num();
+				GlobalImportsByFullName.Add(FullName, GlobalImportIndex);
+				FImportData& GlobalImport = GlobalImports.AddDefaulted_GetRef();
+				int32* FindOuterGlobalImport = GlobalImportsByFullName.Find(OuterName);
+				check(FindOuterGlobalImport);
+				const FImportData& OuterGlobalImport = GlobalImports[*FindOuterGlobalImport];
+				GlobalImport.GlobalIndex = GlobalImportIndex;
+				GlobalImport.OutermostIndex = OuterGlobalImport.OutermostIndex;
+				GlobalImport.OuterIndex = OuterGlobalImport.GlobalIndex;
+				GlobalImport.ObjectName = Import->ObjectName;
+				GlobalImport.bIsScript = OuterGlobalImport.bIsScript;
+				GlobalImport.FullName = FullName;
+				GlobalImport.RefCount = 1;
+			}
+			else
+			{
+				++GlobalImports[*FindGlobalImport].RefCount;
+			}
+		}
+	}
+};
+
+static void FindExport(TArray<FExportData>& GlobalExports, TMap<FString, int32>& GlobalExportsByFullName, TArray<FString>& TempFullNames, const FObjectExport* ExportMap, int32 LocalExportIndex, const FName& PackageName)
+{
+	const FObjectExport* Export = ExportMap + LocalExportIndex;
+	FString& FullName = TempFullNames[LocalExportIndex];
+	if (FullName.Len() == 0)
+	{
+		if (Export->OuterIndex.IsNull())
+		{
+			PackageName.AppendString(FullName);
+			FullName.AppendChar(TEXT('/'));
+			Export->ObjectName.AppendString(FullName);
+		}
+		else
+		{
+			check(Export->OuterIndex.IsExport());
+
+			FindExport(GlobalExports, GlobalExportsByFullName, TempFullNames, ExportMap, Export->OuterIndex.ToExport(), PackageName);
+			FString& OuterName = TempFullNames[Export->OuterIndex.ToExport()];
+			check(OuterName.Len() > 0);
+
+			FullName.Append(OuterName);
+			FullName.AppendChar(TEXT('/'));
+			Export->ObjectName.AppendString(FullName);
+		}
+		check(!GlobalExportsByFullName.Contains(FullName));
+		int32 GlobalExportIndex = GlobalExports.Num();
+		GlobalExportsByFullName.Add(FullName, GlobalExportIndex);
+		FExportData& ExportData = GlobalExports.AddDefaulted_GetRef();
+		ExportData.GlobalIndex = GlobalExportIndex;
+		ExportData.SourcePackageName = PackageName;
+		ExportData.ObjectName = Export->ObjectName;
+		ExportData.SourceIndex = LocalExportIndex;
+		ExportData.FullName = FullName;
+	}
+};
+
+int32 CreateTarget(const FContainerTarget& Target)
+{
+	TGuardValue<int32> GuardAllowUnversionedContentInEditor(GAllowUnversionedContentInEditor, 1);
+
+	FString CookedDir = Target.CookedDirectory;
+	FString OutputDir = Target.OutputDirectory;
+	FString RelativePrefixForLegacyFilename = TEXT("../../../");
+
+	FZenGlobalNameMap GlobalNameMap;
+	FString GlobalNameMapPath = CookedDir / TEXT("megafile.unamemap");
+
+	GlobalNameMap.Load(GlobalNameMapPath);
+
+	FPackageStoreBulkDataManifest BulkDataManifest(CookedDir);
+	const bool bWithBulkDataManifest = BulkDataManifest.Load();
+	// if (!BulkDataManifest.Load())
+	// {
+	// 	UE_LOG(LogIoStore, Error, TEXT("Must cook with -packagestore on the command line (%s not found)"), *BulkDataManifest.GetFilename());
+	// 	RequestEngineExit(TEXT("ZenCreator failed"));
+	// 	return 1;
+	// }
+
+#if OUTPUT_CHUNKID_DIRECTORY
+	ChunkIdCsv.CreateOutputFile(CookedDir);
+#endif
+
+	TArray<FString> FileNames;
+	UE_LOG(LogIoStore, Display, TEXT("Searching for .uasset and .umap files..."));
+	IFileManager::Get().FindFilesRecursive(FileNames, *CookedDir, TEXT("*.uasset"), true, false, false);
+	IFileManager::Get().FindFilesRecursive(FileNames, *CookedDir, TEXT("*.umap"), true, false, false);
+	UE_LOG(LogIoStore, Display, TEXT("Found '%d' files"), FileNames.Num());
+
+	auto DeserializeName = [&](FArchive& A, FName& N)
+	{
+		int32 NameIndex, NameNumber;
+		A << NameIndex << NameNumber;
+		N = GlobalNameMap.GetNameFromDisplayIndex(NameIndex, NameNumber);
+	};
+
+	auto SerializeName = [&](FArchive& A, const FName& N)
+	{
+		if (const int32* NameIndex = GlobalNameMap.GetComparisonIndex(N))
+		{
+			uint32 NameIndexNumber[] { uint32(*NameIndex), uint32(N.GetNumber()) };
+			A << NameIndexNumber[0] << NameIndexNumber[1];
+		}
+		else
+		{
+			UE_LOG(LogIoStore, Display, TEXT("FName '%s' has no valid name index in global name map"), *N.ToString());
+			check(false);
+		}
+	};
+
+	TArray<FName> Names;
+	TSet<FName> UniqueImports;
+	uint64 NameSize = 0;
+	TArray<FObjectImport> Imports;
+	TArray<FObjectExport> Exports;
+	TArray<FImportData> GlobalImports;
+	TArray<FExportData> GlobalExports;
+	TMap<FString, int32> GlobalImportsByFullName;
+	TMap<FString, int32> GlobalExportsByFullName;
+	TArray<FString> TempFullNames;
+	TArray<FPackageIndex> PreloadDependencies;
+	uint64 UniqueImportPackages = 0;
+	uint64 UPackageImports = 0;
+	TArray<FPackageFileSummary> Summaries;
+	TArray<int32> ImportPreloadCounts;
+	TArray<int32> ExportPreloadCounts;
+	Summaries.AddDefaulted(FileNames.Num());
+	ImportPreloadCounts.AddDefaulted(FileNames.Num());
+	ExportPreloadCounts.AddDefaulted(FileNames.Num());
+	uint64 ImportPreloadCount = 0;
+	uint64 ExportPreloadCount = 0;
+
+	TUniquePtr<FArchive> StoreTocArchive(IFileManager::Get().CreateFileWriter(*(OutputDir / TEXT("megafile.ustoretoc"))));
+	TUniquePtr<FArchive> GlimportArchive(IFileManager::Get().CreateFileWriter(*(OutputDir / TEXT("megafile.uglimport"))));
+	TUniquePtr<FArchive> SlimportArchive(IFileManager::Get().CreateFileWriter(*(OutputDir / TEXT("megafile.uslimport"))));
+	TUniquePtr<FArchive> ScriptArcsArchive(IFileManager::Get().CreateFileWriter(*(OutputDir / TEXT("megafile.uscriptarcs"))));
+
+	FIoStoreEnvironment IoStoreEnv;
+#if !SKIP_WRITE_CONTAINER
+	TUniquePtr<FIoStoreWriter> IoStoreWriter;
+
+	{
+		IoStoreEnv.InitializeFileEnvironment(OutputDir);
+		IoStoreWriter = MakeUnique<FIoStoreWriter>(IoStoreEnv);
+		FIoStatus IoStatus = IoStoreWriter->Initialize();
+		check(IoStatus.IsOk());
+	}
+#endif
+
+	TMap<FName, FPackage*> PackageMap;
+	FExportGraph ExportGraph;
+
+	for (int FileIndex = 0; FileIndex < FileNames.Num(); ++FileIndex)
+	{
+		const FString& FileName = FileNames[FileIndex];
+		FPackageFileSummary& Summary = Summaries[FileIndex];
+		TUniquePtr<FArchive> Ar(IFileManager::Get().CreateFileReader(*FileName));
+		check(Ar);
+
+		UE_CLOG(FileIndex % 1000 == 0, LogIoStore, Display, TEXT("Parsing %d: '%s'"), FileIndex, *FileName);
+
+		FString RelativeFileName = FileName;
+		RelativeFileName.RemoveFromStart(*CookedDir);
+		RelativeFileName.RemoveFromStart(TEXT("/"));
+		RelativeFileName = TEXT("../../../") / RelativeFileName;
+
+		FString PackageName;
+		FString ErrorMessage;
+		if (!FPackageName::TryConvertFilenameToLongPackageName(RelativeFileName, PackageName, &ErrorMessage))
+		{
+			UE_LOG(LogIoStore, Error, TEXT("Failed to convert file name from file name '%s'"), *ErrorMessage);
+			return -1;
+		}
+
+		FName PackageFName = *PackageName;
+
+		uint64 SummaryStartPos = Ar->Tell();
+		*Ar << Summary;
+
+		FPackage* NewPackage = new FPackage();
+		PackageMap.Add(PackageFName, NewPackage);
+		FPackage& Package = *NewPackage;
+		Package.Name = PackageFName;
+		Package.FileName = FileName;
+		Package.GlobalPackageId = FileIndex;
+		Package.UAssetSize = Ar->TotalSize();
+		Package.SummarySize = Ar->Tell() - SummaryStartPos;
+		Package.ImportCount = Summary.ImportCount;
+		Package.ExportCount = Summary.ExportCount;
+		Package.PackageFlags = Summary.PackageFlags;
+		Package.BulkDataStartOffset = Summary.BulkDataStartOffset;
+
+		Package.RelativeFileName = RelativePrefixForLegacyFilename;
+		Package.RelativeFileName.Append(*FileName + CookedDir.Len());
+
+		if (Summary.ImportCount > 0)
+		{
+			Ar->Seek(Summary.ImportOffset);
+
+			int32 NumPackages = 0;
+			int32 BaseIndex = Imports.Num();
+			Imports.AddUninitialized(Summary.ImportCount);
+			TArray<FString> ImportNames;
+			ImportNames.Reserve(Summary.ImportCount);
+			for (int32 I = 0; I < Summary.ImportCount; ++I)
+			{
+				FObjectImport& ObjectImport = Imports[BaseIndex + I];
+				DeserializeName(*Ar, ObjectImport.ClassPackage);
+				DeserializeName(*Ar, ObjectImport.ClassName);
+				*Ar << ObjectImport.OuterIndex;
+				DeserializeName(*Ar, ObjectImport.ObjectName);
+
+				if (ObjectImport.OuterIndex.IsNull())
+				{
+					++NumPackages;
+				}
+
+				ImportNames.Emplace(ObjectImport.ObjectName.ToString());
+			}
+
+			UPackageImports += NumPackages;
+
+			Package.SlimportCount = Summary.ImportCount;
+			Package.SlimportOffset = SlimportArchive->Tell();
+			TempFullNames.Reset();
+			TempFullNames.SetNum(Summary.ImportCount, false);
+			for (int32 I = 0; I < Summary.ImportCount; ++I)
+			{
+				FindImport(GlobalImports, GlobalImportsByFullName, TempFullNames, Imports.GetData() + BaseIndex, I);
+
+				FImportData& ImportData = GlobalImports[*GlobalImportsByFullName.Find(TempFullNames[I])];
+				*SlimportArchive << ImportData.GlobalIndex;
+
+				if (ImportData.bIsPackage)
+				{
+					Package.ImportedPackageNames.Add(ImportData.ObjectName);
+				}
+				Package.Imports.Add(ImportData.GlobalIndex);
+			}
+		}
+
+		Package.PreloadIndexOffset = PreloadDependencies.Num();
+		int32 PreloadDependenciesBaseIndex = -1;
+		if (Summary.PreloadDependencyCount > 0)
+		{
+			Ar->Seek(Summary.PreloadDependencyOffset);
+			PreloadDependenciesBaseIndex = PreloadDependencies.Num();
+			PreloadDependencies.AddUninitialized(Summary.PreloadDependencyCount);
+			for (int32 I = 0; I < Summary.PreloadDependencyCount; ++I)
+			{
+				FPackageIndex& Index = PreloadDependencies[PreloadDependenciesBaseIndex + I];
+				*Ar << Index;
+				if (Index.IsImport())
+				{
+					++ImportPreloadCounts[FileIndex];
+					++ImportPreloadCount;
+				}
+				else
+				{
+					++ExportPreloadCounts[FileIndex];
+					++ExportPreloadCount;
+				}
+			}
+		}
+
+		Package.ExportIndexOffset = Exports.Num();
+		if (Summary.ExportCount > 0)
+		{
+			Ar->Seek(Summary.ExportOffset);
+
+			int32 BaseIndex = Exports.Num();
+			Exports.AddUninitialized(Summary.ExportCount);
+			for (int32 I = 0; I < Summary.ExportCount; ++I)
+			{
+				FObjectExport& ObjectExport = Exports[BaseIndex + I];
+				*Ar << ObjectExport.ClassIndex;
+				*Ar << ObjectExport.SuperIndex;
+				*Ar << ObjectExport.TemplateIndex;
+				*Ar << ObjectExport.OuterIndex;
+				DeserializeName(*Ar, ObjectExport.ObjectName);
+				uint32 ObjectFlags;
+				*Ar << ObjectFlags;
+				ObjectExport.ObjectFlags = (EObjectFlags)ObjectFlags;
+				*Ar << ObjectExport.SerialSize;
+				*Ar << ObjectExport.SerialOffset;
+				*Ar << ObjectExport.bForcedExport;
+				*Ar << ObjectExport.bNotForClient;
+				*Ar << ObjectExport.bNotForServer;
+				*Ar << ObjectExport.PackageGuid;
+				*Ar << ObjectExport.PackageFlags;
+				*Ar << ObjectExport.bNotAlwaysLoadedForEditorGame;
+				*Ar << ObjectExport.bIsAsset;
+				*Ar << ObjectExport.FirstExportDependency;
+				*Ar << ObjectExport.SerializationBeforeSerializationDependencies;
+				*Ar << ObjectExport.CreateBeforeSerializationDependencies;
+				*Ar << ObjectExport.SerializationBeforeCreateDependencies;
+				*Ar << ObjectExport.CreateBeforeCreateDependencies;
+			}
+
+			TempFullNames.Reset();
+			TempFullNames.SetNum(Summary.ExportCount, false);
+			for (int32 I = 0; I < Summary.ExportCount; ++I)
+			{
+				FindExport(GlobalExports, GlobalExportsByFullName, TempFullNames, Exports.GetData() + BaseIndex, I, PackageFName);
+
+				FExportData& ExportData = GlobalExports[*GlobalExportsByFullName.Find(TempFullNames[I])];
+				Package.Exports.Add(ExportData.GlobalIndex);
+				ExportData.CreateNode = ExportGraph.AddNode(&Package, { uint32(I), FExportBundleEntry::ExportCommandType_Create });
+				ExportData.SerializeNode = ExportGraph.AddNode(&Package, { uint32(I), FExportBundleEntry::ExportCommandType_Serialize });
+				Package.CreateExportNodes.Add(ExportData.CreateNode);
+				Package.SerializeExportNodes.Add(ExportData.SerializeNode);
+				ExportGraph.AddInternalDependency(ExportData.CreateNode, ExportData.SerializeNode);
+			}
+		}
+	
+		Ar->Close();
+	}
+
+	for (FExportData& GlobalExport : GlobalExports)
+	{
+		int32* FindGlobalImport = GlobalImportsByFullName.Find(GlobalExport.FullName);
+		if (FindGlobalImport)
+		{
+			GlobalExport.GlobalImportIndex = *FindGlobalImport;
+		}
+	}
+
+	uint64 ImportSize = Imports.Num() * 28;
+	uint64 ExportSize = Exports.Num() * 104;
+	uint64 PreloadDependenciesSize = PreloadDependencies.Num() * 4;
+
+	FString CsvFilePath = SlimportArchive->GetArchiveName();
+	CsvFilePath.Append(TEXT(".csv"));
+	TUniquePtr<FArchive> CsvArchive(IFileManager::Get().CreateFileWriter(*CsvFilePath));
+	if (CsvArchive)
+	{
+		ANSICHAR Line[MAX_SPRINTF + FName::StringBufferSize];
+		ANSICHAR Header[] = "Count\tOuter\tOutermost\tImportName\n";
+		CsvArchive->Serialize(Header, sizeof(Header) - 1);
+		for (const FImportData& ImportData : GlobalImports)
+		{
+			FCStringAnsi::Sprintf(Line, "%d\t%d\t%d\t",
+				ImportData.RefCount, ImportData.OuterIndex, ImportData.OutermostIndex);
+			ANSICHAR* L = Line + FCStringAnsi::Strlen(Line);
+			const TCHAR* N = *ImportData.FullName;
+			while (*N)
+			{
+				*L++ = CharCast<ANSICHAR,TCHAR>(*N++);
+			}
+			*L++ = '\n';
+			CsvArchive.Get()->Serialize(Line, L - Line);
+		}
+	}
+	if (GlimportArchive)
+	{
+		int32 Pad = 0;
+		for (const FImportData& ImportData : GlobalImports)
+		{
+			UniqueImportPackages += (ImportData.OuterIndex == 0);
+			SerializeName(*GlimportArchive, ImportData.ObjectName);
+			FPackageIndex Index;
+			Index = FPackageIndex::FromImport(ImportData.GlobalIndex);
+			*GlimportArchive << Index;
+			Index = ImportData.OuterIndex >= 0 ? FPackageIndex::FromImport(ImportData.OuterIndex) :
+				FPackageIndex();
+			*GlimportArchive << Index;
+			Index = FPackageIndex::FromImport(ImportData.OutermostIndex);
+			*GlimportArchive << Index;
+			*GlimportArchive << Pad;
+		}
+	}
+
+	uint64 SummaryChunkCount = 0;
+	uint64 ExportChunkCount = 0;
+	uint64 BulkChunkCount = 0;
+	uint64 BulkPartialChunkCount = 0;
+
+	TSet<FString> MissingExports;
+	TSet<FPackage*> Visited;
+	TSet<FCircularImportChain> CircularChains;
+	int32 PackageIndex = 0;
+
+	// Lookup package pointers for all imports before running postload code
+	UE_LOG(LogIoStore, Display, TEXT("Looking up import packages..."));
+	for (auto& PackageKV : PackageMap)
+	{
+		FPackage& Package = *PackageKV.Value;
+
+		Package.ImportedPackages.Reserve(Package.ImportedPackageNames.Num());
+		for (const FName& ImportedPackageName : Package.ImportedPackageNames)
+		{
+			FPackage* FindPackage = PackageMap.FindRef(ImportedPackageName);
+			if (FindPackage)
+			{
+				Package.ImportedPackages.Add(FindPackage);
+			}
+		}
+	}
+
+	UE_LOG(LogIoStore, Display, TEXT("Simulating postload dependencies for stats..."));
+	for (auto& PackageKV : PackageMap)
+	{
+		FPackage& Package = *PackageKV.Value;
+
+		Visited.Reset();
+		for (FPackage* DependentPackage : Package.ImportedPackages)
+		{
+			AddPostLoadDependenciesRecursive(Package, *DependentPackage, Visited);
+		}
+	}
+
+	UE_LOG(LogIoStore, Display, TEXT("Adding reachable packages for stats..."));
+	for (auto& PackageKV : PackageMap)
+	{
+		FPackage& Package = *PackageKV.Value;
+
+		AddReachablePackagesRecursive(Package, Package, Package.AllReachablePackages, /*bFirst*/ true);
+		Package.bHasCircularImportDependencies = Package.AllReachablePackages.Contains(&Package);
+		// if (Package.bHasCircularImportDependencies)
+		// {
+		// 	UE_LOG(LogIoStore, Display, TEXT("Circular %s: %d reachable packages"), *Package.Name.ToString(),
+		// 		Package.AllReachablePackages.Num());
+		// }
+	}
+
+	UE_LOG(LogIoStore, Display, TEXT("Adding optimized postload dependencies..."));
+	for (auto& PackageKV : PackageMap)
+	{
+		FPackage& Package = *PackageKV.Value;
+
+		Visited.Reset();
+		Visited.Add(&Package);
+		AddPostLoadDependencies(Package, Visited, CircularChains);
+	}
+		
+	UE_LOG(LogIoStore, Display, TEXT("Adding preload dependencies..."));
+	for (auto& PackageKV : PackageMap)
+	{
+		FPackage& Package = *PackageKV.Value;
+
+		// Convert PreloadDependencies to arcs
+		for (int32 I = 0; I < Package.ExportCount; ++I)
+		{
+			FObjectExport& ObjectExport = Exports[Package.ExportIndexOffset + I];
+			FExportData& ExportData = GlobalExports[Package.Exports[I]];
+			int32 PreloadDependenciesBaseIndex = Package.PreloadIndexOffset;
+
+			FPackageIndex ExportPackageIndex = FPackageIndex::FromExport(I);
+
+			auto AddPreloadArc = [&](FPackageIndex Dep, EPreloadDependencyType PhaseFrom, EPreloadDependencyType PhaseTo)
+			{
+				if (Dep.IsExport())
+				{
+					AddInternalExportArc(ExportGraph, Package, Dep.ToExport(), PhaseFrom, I, PhaseTo);
+				}
+				else
+				{
+					FImportData& Import = GlobalImports[Package.Imports[Dep.ToImport()]];
+					ensure(!Import.bIsPackage);
+					int32* FindGlobalExport = GlobalExportsByFullName.Find(Import.FullName);
+					if (FindGlobalExport)
+					{
+						FExportData& Export = GlobalExports[*FindGlobalExport];
+						FPackage* SourcePackage = PackageMap.FindRef(Export.SourcePackageName);
+						check(SourcePackage);
+						FPackageIndex SourceDep = FPackageIndex::FromExport(Export.SourceIndex);
+						AddExternalExportArc(ExportGraph, *SourcePackage, Export.SourceIndex, PhaseFrom, Package, I, PhaseTo);
+						Package.ImportedPreloadPackages.Add(SourcePackage);
+					}
+					else
+					{
+						// Add script arc with null package and global import index as node index
+						AddScriptArc(Package, Import.GlobalIndex, I, PhaseTo);
+					}
+				}
+			};
+
+			if (PreloadDependenciesBaseIndex >= 0 && ObjectExport.FirstExportDependency >= 0)
+			{
+				int32 RunningIndex = PreloadDependenciesBaseIndex + ObjectExport.FirstExportDependency;
+				for (int32 Index = ObjectExport.SerializationBeforeSerializationDependencies; Index > 0; Index--)
+				{
+					FPackageIndex Dep = PreloadDependencies[RunningIndex++];
+					check(!Dep.IsNull());
+					AddPreloadArc(Dep, PreloadDependencyType_Serialize, PreloadDependencyType_Serialize);
+				}
+
+				for (int32 Index = ObjectExport.CreateBeforeSerializationDependencies; Index > 0; Index--)
+				{
+					FPackageIndex Dep = PreloadDependencies[RunningIndex++];
+					check(!Dep.IsNull());
+					AddPreloadArc(Dep, PreloadDependencyType_Create, PreloadDependencyType_Serialize);
+				}
+
+				for (int32 Index = ObjectExport.SerializationBeforeCreateDependencies; Index > 0; Index--)
+				{
+					FPackageIndex Dep = PreloadDependencies[RunningIndex++];
+					check(!Dep.IsNull());
+					AddPreloadArc(Dep, PreloadDependencyType_Serialize, PreloadDependencyType_Create);
+				}
+
+				for (int32 Index = ObjectExport.CreateBeforeCreateDependencies; Index > 0; Index--)
+				{
+					FPackageIndex Dep = PreloadDependencies[RunningIndex++];
+					check(!Dep.IsNull());
+					// can't create this export until these things are created
+					AddPreloadArc(Dep, PreloadDependencyType_Create, PreloadDependencyType_Create);
+				}
+			}
+		}
+	}
+
+	UE_LOG(LogIoStore, Display, TEXT("Building bundles..."));
+	BuildBundles(ExportGraph, PackageMap);
+
+	UE_LOG(LogIoStore, Display, TEXT("Serializing..."));
+	for (auto& PackageKV : PackageMap)
+	{
+		FPackage& Package = *PackageKV.Value;
+
+		// Separate file for script arcs that are only required during initial loading
+		{
+			Package.ScriptArcsOffset = ScriptArcsArchive->Tell();
+			Package.ScriptArcsCount = Package.ScriptArcs.Num();
+			for (FArc& ScriptArc : Package.ScriptArcs)
+			{
+				*ScriptArcsArchive << ScriptArc.FromNodeIndex;
+				*ScriptArcsArchive << ScriptArc.ToNodeIndex;
+			}
+		}
+
+		// Temporary Archive for serializing EDL graph data
+		FBufferWriter ZenGraphArchive(nullptr, 0, EBufferWriterFlags::AllowResize | EBufferWriterFlags::TakeOwnership );
+
+		int32 InternalArcCount = Package.InternalArcs.Num();
+		ZenGraphArchive << InternalArcCount;
+		for (FArc& InternalArc : Package.InternalArcs)
+		{
+			ZenGraphArchive << InternalArc.FromNodeIndex;
+			ZenGraphArchive << InternalArc.ToNodeIndex;
+		}
+
+		int32 ImportedPackagesCount = Package.ExternalArcs.Num();
+		ZenGraphArchive << ImportedPackagesCount;
+		for (auto& KV : Package.ExternalArcs)
+		{
+			const FName& ImportedPackageName = KV.Key;
+			SerializeName(ZenGraphArchive, ImportedPackageName);
+
+			TArray<FArc>& Arcs = KV.Value;
+
+			int32 ExternalArcCount = Arcs.Num();
+			ZenGraphArchive << ExternalArcCount;
+			for (FArc& ExternalArc : Arcs)
+			{
+				ZenGraphArchive << ExternalArc.FromNodeIndex;
+				ZenGraphArchive << ExternalArc.ToNodeIndex;
+			}
+		}
+		Package.UGraphSize = ZenGraphArchive.Tell();
+
+		// Temporary Archive for serializing export map data
+		FBufferWriter ZenExportMapArchive(nullptr, 0, EBufferWriterFlags::AllowResize | EBufferWriterFlags::TakeOwnership );
+		for (int32 I = 0; I < Package.ExportCount; ++I)
+		{
+			FObjectExport& ObjectExport = Exports[Package.ExportIndexOffset + I];
+			FExportData& ExportData = GlobalExports[Package.Exports[I]];
+
+			ZenExportMapArchive << ObjectExport.SerialSize;
+			SerializeName(ZenExportMapArchive, ObjectExport.ObjectName);
+			ZenExportMapArchive << ObjectExport.OuterIndex;
+			ZenExportMapArchive << ObjectExport.ClassIndex;
+			ZenExportMapArchive << ObjectExport.SuperIndex;
+			ZenExportMapArchive << ObjectExport.TemplateIndex;
+			ZenExportMapArchive << ExportData.GlobalImportIndex;
+			ZenExportMapArchive << (uint32&)ObjectExport.ObjectFlags;
+		}
+		Package.ExportMapSize = ZenExportMapArchive.Tell();
+
+		// Temporary archive for serializing export bundle data
+		FBufferWriter ZenExportBundlesArchive(nullptr, 0, EBufferWriterFlags::AllowResize | EBufferWriterFlags::TakeOwnership);
+		int32 ExportBundleEntryIndex = 0;
+		for (FExportBundle& ExportBundle : Package.ExportBundles)
+		{
+			ZenExportBundlesArchive << ExportBundleEntryIndex;
+			int32 EntryCount = ExportBundle.Nodes.Num();
+			ZenExportBundlesArchive << EntryCount;
+			ExportBundleEntryIndex += ExportBundle.Nodes.Num();
+		}
+		for (FExportBundle& ExportBundle : Package.ExportBundles)
+		{
+			for (FExportGraphNode* ExportNode : ExportBundle.Nodes)
+			{
+				uint32 CommandType = uint32(ExportNode->BundleEntry.CommandType);
+				ZenExportBundlesArchive << CommandType;
+				ZenExportBundlesArchive << ExportNode->BundleEntry.LocalExportIndex;
+			}
+		}
+		Package.ExportBundlesSize = ZenExportBundlesArchive.Tell();
+
+		int32 PackageNameIndex = GlobalNameMap.GetOrCreateComparisonIndex(Package.Name);
+		int32 PackageNameNumber = Package.Name.GetNumber();
+
+		FName RelativeFileName(*Package.RelativeFileName);
+		int32 FileNameIndex = GlobalNameMap.GetOrCreateComparisonIndex(RelativeFileName);
+		int32 FileNameNumber = Package.Name.GetNumber();
+
+		ensure(Package.GlobalPackageId == PackageIndex++);
+		{
+			*StoreTocArchive << PackageNameIndex << PackageNameNumber;
+			*StoreTocArchive << FileNameIndex << FileNameNumber;
+			*StoreTocArchive << Package.ImportCount;
+			*StoreTocArchive << Package.SlimportCount;
+			*StoreTocArchive << Package.SlimportOffset;
+			*StoreTocArchive << Package.ExportCount;
+			int32 ExportBundleCount = Package.ExportBundles.Num();
+			*StoreTocArchive << ExportBundleCount;
+			*StoreTocArchive << Package.ScriptArcsOffset;
+			*StoreTocArchive << Package.ScriptArcsCount;
+		}
+
+		{
+			const uint64 ZenSummarySize = 
+						sizeof(FZenPackageSummary)
+						+ Package.ExportMapSize
+						+ Package.ExportBundlesSize
+						+ Package.UGraphSize;
+
+			uint8* ZenSummaryBuffer = static_cast<uint8*>(FMemory::Malloc(ZenSummarySize));
+			FZenPackageSummary* ZenSummary = reinterpret_cast<FZenPackageSummary*>(ZenSummaryBuffer);
+
+			ZenSummary->PackageFlags = Package.PackageFlags;
+			ZenSummary->GraphDataSize = Package.UGraphSize;
+			ZenSummary->BulkDataStartOffset = Package.BulkDataStartOffset;
+
+			FBufferWriter ZenAr(ZenSummaryBuffer, ZenSummarySize); 
+			ZenAr.Seek(sizeof(FZenPackageSummary));
+
+			// ExportMap data
+			{
+				check(ZenExportMapArchive.Tell() == Package.ExportMapSize);	
+				ZenSummary->ExportMapOffset = ZenAr.Tell();
+				ZenAr.Serialize(ZenExportMapArchive.GetWriterData(), ZenExportMapArchive.Tell());
+			}
+
+			// ExportBundle data
+			{
+				check(ZenExportBundlesArchive.Tell() == Package.ExportBundlesSize);
+				ZenSummary->ExportBundlesOffset = ZenAr.Tell();
+				ZenAr.Serialize(ZenExportBundlesArchive.GetWriterData(), ZenExportBundlesArchive.Tell());
+			}
+
+			// Graph data
+			{
+				check(ZenGraphArchive.Tell() == Package.UGraphSize);	
+				ZenSummary->GraphDataOffset = ZenAr.Tell();
+				ZenAr.Serialize(ZenGraphArchive.GetWriterData(), ZenGraphArchive.Tell());
+			}
+
+			// Package summary chunk
+			{
+				FIoBuffer IoBuffer(FIoBuffer::AssumeOwnership, ZenSummaryBuffer, ZenSummarySize);
+#if !SKIP_WRITE_CONTAINER
+				IoStoreWriter->Append(CreateZenChunkId(PackageNameIndex, PackageNameNumber, 0, EZenChunkType::PackageSummary, TEXT("PackageSummary")), IoBuffer);
+#endif
+				++SummaryChunkCount;
+			}
+
+			// Export chunks
+			if (Package.ExportCount)
+			{
+				FString UExpFileName = FPaths::ChangeExtension(Package.FileName, TEXT(".uexp"));
+				TUniquePtr<FArchive> ExpAr(IFileManager::Get().CreateFileReader(*UExpFileName));
+				const int64 TotalExportsSize = ExpAr->TotalSize();
+#if !SKIP_WRITE_CONTAINER
+				uint8* ExportsBuffer = static_cast<uint8*>(FMemory::Malloc(TotalExportsSize));
+				ExpAr->Serialize(ExportsBuffer, TotalExportsSize);
+#endif
+				Package.UExpSize = TotalExportsSize;
+
+				for (int32 I = 0; I < Package.ExportBundles.Num(); ++I)
+				{
+					check(I < UINT16_MAX);
+					FExportBundle& ExportBundle = Package.ExportBundles[I];
+					uint64 BundleBufferSize = 0;
+					for (FExportGraphNode* Node : ExportBundle.Nodes)
+					{
+						if (Node->BundleEntry.CommandType == FExportBundleEntry::ExportCommandType_Serialize)
+						{
+							FObjectExport& ObjectExport = Exports[Package.ExportIndexOffset + Node->BundleEntry.LocalExportIndex];
+							BundleBufferSize += ObjectExport.SerialSize;
+						}
+					}
+					uint8* BundleBuffer = static_cast<uint8*>(FMemory::Malloc(BundleBufferSize));
+					uint64 BundleBufferOffset = 0;
+					for (FExportGraphNode* Node : ExportBundle.Nodes)
+					{
+						if (Node->BundleEntry.CommandType == FExportBundleEntry::ExportCommandType_Serialize)
+						{
+							FObjectExport& ObjectExport = Exports[Package.ExportIndexOffset + Node->BundleEntry.LocalExportIndex];
+							const int64 Offset = ObjectExport.SerialOffset - Package.UAssetSize;
+							FMemory::Memcpy(BundleBuffer + BundleBufferOffset, ExportsBuffer + Offset, ObjectExport.SerialSize);
+							BundleBufferOffset += ObjectExport.SerialSize;
+						}
+					}
+#if !SKIP_WRITE_CONTAINER
+					FObjectExport& ObjectExport = Exports[Package.ExportIndexOffset + I];
+					const int64 Offset = ObjectExport.SerialOffset - Package.UAssetSize;
+					FIoBuffer IoBuffer(FIoBuffer::Wrap, BundleBuffer, BundleBufferSize);
+					IoStoreWriter->Append(CreateZenChunkId(PackageNameIndex, PackageNameNumber, I, EZenChunkType::ExportBundleData, *Package.FileName), IoBuffer);
+#endif
+					++ExportChunkCount;
+					FMemory::Free(BundleBuffer);
+				}
+
+#if !SKIP_WRITE_CONTAINER
+				FMemory::Free(ExportsBuffer);
+				ExpAr->Close();
+#endif
+			}
+
+#if !SKIP_BULKDATA
+			// Bulk chunks
+			if (bWithBulkDataManifest)
+			{
+				FString UBulkFileName = FPaths::ChangeExtension(Package.FileName, TEXT(".ubulk"));
+				FPaths::NormalizeFilename(UBulkFileName);
+
+				TUniquePtr<FArchive> BulkAr(IFileManager::Get().CreateFileReader(*UBulkFileName));
+				if (BulkAr != nullptr)
+				{
+					const FIoChunkId BulkDataChunkId = CreateZenChunkId(PackageNameIndex, PackageNameNumber, 0, EZenChunkType::BulkData, *Package.FileName);
+#if !SKIP_WRITE_CONTAINER
+					uint8* BulkBuffer = static_cast<uint8*>(FMemory::Malloc(BulkAr->TotalSize()));
+					BulkAr->Serialize(BulkBuffer, BulkAr->TotalSize());
+					FIoBuffer IoBuffer(FIoBuffer::AssumeOwnership, BulkBuffer, BulkAr->TotalSize());
+					const FIoStatus AppendResult = IoStoreWriter->Append(BulkDataChunkId, IoBuffer);
+					if (!AppendResult.IsOk())
+					{
+						UE_LOG(LogIoStore, Error, TEXT("Failed to append bulkdata for '%s' due to: %s"), *Package.FileName, *AppendResult.ToString());
+						RequestEngineExit(TEXT("ZenCreator failed"));
+						return 1;
+					}
+
+					BulkAr->Close();
+#endif
+					++BulkChunkCount;
+
+					// Create additional mapping chunks as needed
+					const FPackageStoreBulkDataManifest::PackageDesc* PackageDesc = BulkDataManifest.Find(Package.FileName);
+					if (PackageDesc != nullptr)
+					{
+						for (const FPackageStoreBulkDataManifest::PackageDesc::BulkDataDesc& BulkDataDesc : PackageDesc->GetDataArray())
+						{
+#if !SKIP_WRITE_CONTAINER
+							const FIoChunkId AccessChunkId = CreateZenChunkId( PackageNameIndex, PackageNameNumber, BulkDataDesc.Index, EZenChunkType::BulkData, *Package.FileName);
+							const FIoStatus PartialResult = IoStoreWriter->MapPartialRange(BulkDataChunkId, BulkDataDesc.Offset, BulkDataDesc.Size, AccessChunkId);
+							if (!PartialResult.IsOk())
+							{
+								UE_LOG(LogIoStore, Error, TEXT("Failed to map partial range for '%s' due to: %s"), *Package.FileName, *PartialResult.ToString());
+								RequestEngineExit(TEXT("ZenCreator failed"));
+								return 1;
+							}
+#endif
+							++BulkPartialChunkCount;
+						}
+					}
+				}
+			}
+#endif
+		}
+	}
+
+	GlobalNameMap.Save(OutputDir / TEXT("Container.namemap"));
+	UE_LOG(LogIoStore, Display, TEXT("Saving StoreToc..."));
+	StoreTocArchive->Close();
+
+	UE_LOG(LogIoStore, Display, TEXT("Calculating stats..."));
+	uint64 UExpSize = 0;
+	uint64 UAssetSize = 0;
+	uint64 SummarySize = 0;
+	uint64 UGraphSize = 0;
+	uint64 ExportMapSize = 0;
+	uint64 ZenPackageSummarySize = FileNames.Num() * sizeof(FZenPackageSummary);
+	uint64 StoreTocSize = StoreTocArchive->Tell();
+	uint64 GlimportSize = GlimportArchive->Tell();
+	uint64 SlimportSize = SlimportArchive->Tell();
+	uint64 ScriptArcsSize = ScriptArcsArchive->Tell();
+	uint64 ScriptArcsCount = 0;
+	uint64 CircularPackagesCount = 0;
+	uint64 TotalInternalArcCount = 0;
+	uint64 TotalExternalArcCount = 0;
+
+	uint64 PackagesWithCircularDependenciesCount = 0;
+	uint64 PackagesWithoutImportDependenciesCount = 0;
+	uint64 PackagesWithoutPreloadDependenciesCount = 0;
+	uint64 BundleCount = 0;
+	uint64 BundleEntryCount = 0;
+
+	for (auto& PackageKV : PackageMap)
+	{
+		FPackage& Package = *PackageKV.Value;
+
+		UExpSize += Package.UExpSize;
+		UAssetSize += Package.UAssetSize;
+		SummarySize += Package.SummarySize;
+		UGraphSize += Package.UGraphSize;
+		ExportMapSize += Package.ExportMapSize;
+		ScriptArcsCount += Package.ScriptArcsCount;
+		CircularPackagesCount += Package.bHasCircularImportDependencies;
+		TotalInternalArcCount += Package.InternalArcs.Num();
+		PackagesWithCircularDependenciesCount += Package.bHasCircularPostLoadDependencies;
+		PackagesWithoutPreloadDependenciesCount += Package.ImportedPreloadPackages.Num() == 0;
+		PackagesWithoutImportDependenciesCount += Package.ImportedPackages.Num() == 0;
+
+		for (auto& KV : Package.ExternalArcs)
+		{
+			TArray<FArc>& Arcs = KV.Value;
+			TotalExternalArcCount += Arcs.Num();
+		}
+
+		for (FExportBundle& Bundle : Package.ExportBundles)
+		{
+			++BundleCount;
+			BundleEntryCount += Bundle.Nodes.Num();
+		}
+	}
+
+	UE_LOG(LogIoStore, Display, TEXT(""));
+	UE_LOG(LogIoStore, Display, TEXT("-------------------- IoStore Summary: %s --------------------"), *Target.TargetPlatform->PlatformName());
+
+	UE_LOG(LogIoStore, Display, TEXT("Packages: %5d total, %d circular dependencies, %d no preload dependencies, %d no import dependencies"),
+		PackageMap.Num(), PackagesWithCircularDependenciesCount, PackagesWithoutPreloadDependenciesCount, PackagesWithoutImportDependenciesCount);
+	UE_LOG(LogIoStore, Display, TEXT("Bundles:  %5d total, %d entries"), BundleCount, BundleEntryCount);
+	UE_LOG(LogIoStore, Display, TEXT("Exports:  %5d total"), GlobalExports.Num());
+	
+	UE_LOG(LogIoStore, Display, TEXT("Pak: %12.2f MB UExp, %d files"), (double)UExpSize / 1024.0 / 1024.0, FileNames.Num());
+	UE_LOG(LogIoStore, Display, TEXT("Pak: %12.2f MB UAsset/UMap, %d files"), (double)UAssetSize / 1024.0 / 1024.0, FileNames.Num());
+	UE_LOG(LogIoStore, Display, TEXT("Pak: %12.2f MB FPackageFileSummary"), (double)SummarySize / 1024.0 / 1024.0);
+	UE_LOG(LogIoStore, Display, TEXT("Pak: %12.2f MB ExportMap, %d exports"), (double)ExportSize / 1024.0 / 1024.0, Exports.Num());
+	UE_LOG(LogIoStore, Display, TEXT("Pak: %12.2f MB ImportMap, %d imports (%d UPackages)"), (double)ImportSize / 1024.0 / 1024.0, Imports.Num(), UPackageImports);
+	UE_LOG(LogIoStore, Display, TEXT("Pak: %12.2f MB PreloadDependencies, %d imports + %d exports = %d"),
+		(double)PreloadDependenciesSize / 1024.0 / 1024.0, ImportPreloadCount, ExportPreloadCount, PreloadDependencies.Num());
+	UE_LOG(LogIoStore, Display, TEXT(" "));
+	UE_LOG(LogIoStore, Display, TEXT("IoStore: %8.2f MB StoreToc"), (double)StoreTocSize / 1024.0 / 1024.0);
+	UE_LOG(LogIoStore, Display, TEXT("IoStore: %8.2f MB StoreGlimports"), (double)GlimportSize / 1024.0 / 1024.0);
+	UE_LOG(LogIoStore, Display, TEXT("IoStore: %8.2f MB StoreSlimports"), (double)SlimportSize / 1024.0 / 1024.0);
+	UE_LOG(LogIoStore, Display, TEXT("IoStore: %8.2f MB PackageSummary"), (double)ZenPackageSummarySize / 1024.0 / 1024.0);
+	UE_LOG(LogIoStore, Display, TEXT("IoStore: %8.2f MB ExportMap"), (double)ExportMapSize / 1024.0 / 1024.0);
+	UE_LOG(LogIoStore, Display, TEXT("IoStore: %8.2f MB UGraph, %d internal arcs, %d external arcs"), (double)UGraphSize / 1024.0 / 1024.0, TotalInternalArcCount, TotalExternalArcCount);
+	UE_LOG(LogIoStore, Display, TEXT("IoStore: %8.2f MB ScriptArcs, %d arcs"), (double)ScriptArcsSize / 1024.0 / 1024.0, ScriptArcsCount);
+	UE_LOG(LogIoStore, Display, TEXT("IoStore: %8d unique imports, %d unique packages"), GlobalImportsByFullName.Num(), UniqueImportPackages);
+	UE_LOG(LogIoStore, Display, TEXT("IoStore: %8d circular packages, %d unique chains resolved"), CircularPackagesCount, CircularChains.Num());
+	UE_LOG(LogIoStore, Display, TEXT("IoStore: %8d summary chunks, %d export chunks, %d bulk data chunks, %d partial bulk data chunks"), SummaryChunkCount, ExportChunkCount, BulkChunkCount, BulkPartialChunkCount);
+
+	return 0;
+}
+
+int32 CreateIoStoreContainerFiles(const TCHAR* CmdLine)
+{
+	UE_LOG(LogIoStore, Display, TEXT("==================== IoStore Utils ===================="));
+
+	FString TargetPlatformName;
+	if (FParse::Value(FCommandLine::Get(), TEXT("TargetPlatform="), TargetPlatformName))
+	{
+		UE_LOG(LogIoStore, Display, TEXT("Using target platform: '%s'"), *TargetPlatformName);
+	}
+	else
+	{
+		UE_LOG(LogIoStore, Error, TEXT("No target platform specified [-targetplatform]"));
+		return -1;
+	}
+
+	FString OutputDirectory;
+	if (FParse::Value(FCommandLine::Get(), TEXT("OutputDirectory="), OutputDirectory))
+	{
+		UE_LOG(LogIoStore, Display, TEXT("Using output directory: '%s'"), *OutputDirectory);
+	}
+	else
+	{
+		UE_LOG(LogIoStore, Display, TEXT("No output directory specified, using project's cooked folder"));
+	}
+
+	ITargetPlatformManagerModule& TargetPlatformManager = *GetTargetPlatformManager();
+
+	if (ITargetPlatform* TargetPlatform = TargetPlatformManager.FindTargetPlatform(TargetPlatformName))
+	{
+		FString TargetCookedDirectory = FPaths::ProjectSavedDir() / TEXT("Cooked") / TargetPlatform->PlatformName();
+		FString TargetOutputDirectory = OutputDirectory.Len() > 0
+			? OutputDirectory
+			: TargetCookedDirectory / FApp::GetProjectName();
+
+		FString GlobalNameMapFilePath = TargetCookedDirectory / TEXT("megafile.unamemap");
+
+		if (IFileManager::Get().FileExists(*GlobalNameMapFilePath))
+		{
+			FContainerTarget Target { TargetPlatform, TargetCookedDirectory, TargetOutputDirectory };
+
+			UE_LOG(LogIoStore, Display, TEXT("Creating target: '%s' using output directory: '%s'"), *Target.TargetPlatform->PlatformName(), *Target.OutputDirectory);
+
+			CreateTarget(Target);
+		}
+		else
+		{
+			UE_LOG(LogIoStore, Error,
+				TEXT("Expected to find global name map '%s' for target '%s'"),
+				*GlobalNameMapFilePath, *TargetPlatform->PlatformName());
+		}
+	}
+
+	return 0;
+}
diff --git a/Engine/Source/Developer/IoStoreUtilities/Public/IoStoreUtilities.h b/Engine/Source/Developer/IoStoreUtilities/Public/IoStoreUtilities.h
new file mode 100644
--- /dev/null
+++ b/Engine/Source/Developer/IoStoreUtilities/Public/IoStoreUtilities.h
@@ -0,0 +1,7 @@
+// Copyright 1998-2019 Epic Games, Inc. All Rights Reserved.
+
+#pragma once
+
+#include "CoreMinimal.h"
+
+IOSTOREUTILITIES_API int32 CreateIoStoreContainerFiles(const TCHAR* CmdLine);
diff --git a/Engine/Source/Editor/UnrealEd/Classes/Commandlets/IoStoreCommandlet.h b/Engine/Source/Editor/UnrealEd/Classes/Commandlets/IoStoreCommandlet.h
new file mode 100644
--- /dev/null
+++ b/Engine/Source/Editor/UnrealEd/Classes/Commandlets/IoStoreCommandlet.h
@@ -0,0 +1,24 @@
+// Copyright 1998-2019 Epic Games, Inc. All Rights Reserved.
+
+#pragma once
+
+#include "CoreMinimal.h"
+#include "UObject/ObjectMacros.h"
+#include "Commandlets/Commandlet.h"
+#include "IoStoreCommandlet.generated.h"
+
+/*
+ * Commandlet for creating container file(s) used by the IoDispatcher.
+ */
+UCLASS()
+class UIoStoreCommandlet  
+	: public UCommandlet
+{
+	GENERATED_UCLASS_BODY()
+
+	//~ Begin UCommandlet Interface
+	virtual int32 Main(const FString& Params) override;
+	//~ End UCommandlet Interface
+
+	static int32 StaticMain(const FString& Params);
+};
diff --git a/Engine/Source/Editor/UnrealEd/Private/Commandlets/IoStoreCommandlet.cpp b/Engine/Source/Editor/UnrealEd/Private/Commandlets/IoStoreCommandlet.cpp
new file mode 100644
--- /dev/null
+++ b/Engine/Source/Editor/UnrealEd/Private/Commandlets/IoStoreCommandlet.cpp
@@ -0,0 +1,14 @@
+// Copyright 1998-2019 Epic Games, Inc. All Rights Reserved.
+
+#include "Commandlets/IoStoreCommandlet.h"
+#include "IoStoreUtilities.h"
+
+UIoStoreCommandlet::UIoStoreCommandlet(const FObjectInitializer& ObjectInitializer)
+	: Super(ObjectInitializer)
+{
+}
+
+int32 UIoStoreCommandlet::Main(const FString& Params)
+{
+	return CreateIoStoreContainerFiles(*Params);
+}
diff --git a/Engine/Source/Editor/UnrealEd/UnrealEd.Build.cs b/Engine/Source/Editor/UnrealEd/UnrealEd.Build.cs
--- a/Engine/Source/Editor/UnrealEd/UnrealEd.Build.cs
+++ b/Engine/Source/Editor/UnrealEd/UnrealEd.Build.cs
@@ -184,6 +184,7 @@ public class UnrealEd : ModuleRules
 				"ScriptDisassembler",
 				"ToolMenus",
 				"FreeImage",
+				"IoStoreUtilities",
             }
 		);
 
diff --git a/Engine/Source/Programs/AutomationTool/AutomationUtils/ProjectParams.cs b/Engine/Source/Programs/AutomationTool/AutomationUtils/ProjectParams.cs
--- a/Engine/Source/Programs/AutomationTool/AutomationUtils/ProjectParams.cs
+++ b/Engine/Source/Programs/AutomationTool/AutomationUtils/ProjectParams.cs
@@ -280,6 +280,7 @@ namespace AutomationTool
             this.RunTimeoutSeconds = InParams.RunTimeoutSeconds;
 			this.Clean = InParams.Clean;
 			this.Pak = InParams.Pak;
+			this.IoStore = InParams.IoStore;
 			this.SignPak = InParams.SignPak;
 			this.SignedPak = InParams.SignedPak;
 			this.PakAlignForMemoryMapping = InParams.PakAlignForMemoryMapping;
@@ -454,6 +455,7 @@ namespace AutomationTool
 			bool? NoXGE = null,
 			bool? Package = null,
 			bool? Pak = null,
+			bool? IoStore = null,
 			bool? Prereqs = null,
 			string AppLocalDirectory = null,
 			bool? NoBootstrapExe = null,
@@ -626,6 +628,7 @@ namespace AutomationTool
 			}
 			this.PakAlignForMemoryMapping = GetParamValueIfNotSpecified(Command, PakAlignForMemoryMapping, this.PakAlignForMemoryMapping, "PakAlignForMemoryMapping");
 			this.Pak = GetParamValueIfNotSpecified(Command, Pak, this.Pak, "pak");
+			this.IoStore = GetParamValueIfNotSpecified(Command, IoStore, this.IoStore, "iostore");
 			this.SkipPak = GetParamValueIfNotSpecified(Command, SkipPak, this.SkipPak, "skippak");
 			if (this.SkipPak)
 			{
@@ -1179,6 +1182,12 @@ namespace AutomationTool
 		[Help("pak", "generate a pak file")]
 		public bool Pak { private set; get; }
 
+		/// <summary>
+		/// Shared: True if container file(s) should be generated with ZenPak.
+		/// </summary>
+		[Help("iostore", "generate I/O store container file(s)")]
+		public bool IoStore { private set; get; }
+
 		/// <summary>
 		/// 
 		/// </summary>
@@ -2766,6 +2775,7 @@ namespace AutomationTool
 				CommandUtils.LogLog("MapsToCook={0}", MapsToCook.ToString());
 				CommandUtils.LogLog("MapIniSectionsToCook={0}", MapIniSectionsToCook.ToString());
 				CommandUtils.LogLog("Pak={0}", Pak);
+				CommandUtils.LogLog("IoStore={0}", IoStore);
 				CommandUtils.LogLog("Package={0}", Package);
 				CommandUtils.LogLog("ForcePackageData={0}", ForcePackageData);
 				CommandUtils.LogLog("NullRHI={0}", NullRHI);
diff --git a/Engine/Source/Programs/AutomationTool/Scripts/CopyBuildToStagingDirectory.Automation.cs b/Engine/Source/Programs/AutomationTool/Scripts/CopyBuildToStagingDirectory.Automation.cs
--- a/Engine/Source/Programs/AutomationTool/Scripts/CopyBuildToStagingDirectory.Automation.cs
+++ b/Engine/Source/Programs/AutomationTool/Scripts/CopyBuildToStagingDirectory.Automation.cs
@@ -1744,6 +1744,21 @@ public partial class Project : CommandUtils
 				}
 			}
 
+			// Filter I/O store container files
+			if (Src.HasExtension(".namemap") ||
+				Src.HasExtension(".ucas") ||
+				Src.HasExtension(".utoc") ||
+				Src.HasExtension(".ubulkmanifest") ||
+				Src.HasExtension(".uglimport") ||
+				Src.HasExtension(".unamemap") ||
+				Src.HasExtension(".uscriptarc") ||
+				Src.HasExtension(".uslimport") ||
+				Src.HasExtension(".ustoretoc"))
+			{
+				LogInformation("Excluding {0}", Src);
+				continue;
+			}
+
 			// Do a filtered copy of all ini files to allow stripping of values that we don't want to distribute
 			if (Src.HasExtension(".ini"))
 			{
@@ -2780,6 +2795,16 @@ public partial class Project : CommandUtils
 		return TmpPackagingPath;
 	}
 
+	private static bool ShouldCreateIoStoreContainerFiles(ProjectParams Params, DeploymentContext SC)
+	{
+		if (Params.CookOnTheFly || Params.NoClient)
+		{
+			return false;
+		}
+
+		return Params.IoStore;
+	}
+
 	private static bool ShouldCreatePak(ProjectParams Params, DeploymentContext SC)
 	{
 		if (Params.CookOnTheFly)
@@ -2897,6 +2922,21 @@ public partial class Project : CommandUtils
 
 	public static void ApplyStagingManifest(ProjectParams Params, DeploymentContext SC)
 	{
+		if (ShouldCreateIoStoreContainerFiles(Params, SC))
+		{
+			foreach (var ClientPlatform in Params.ClientTargetPlatforms)
+			{
+				TargetPlatformDescriptor DataPlatformDesc = Params.GetCookedDataPlatformForClientTarget(ClientPlatform);
+				String TargetPlatformName = Platform.Platforms[DataPlatformDesc].GetCookPlatform(false, Params.Client);
+
+				DirectoryReference OutputDirectory = DirectoryReference.Combine(SC.RuntimeRootDir, SC.RelativeProjectRootForStage.Name);
+
+				String CommandletParams = String.Format("-TargetPlatform={0} -OutputDirectory={1}", TargetPlatformName, OutputDirectory.FullName);
+
+				RunCommandlet(SC.RawProjectPath, Params.UE4Exe, "IoStore", CommandletParams);
+			}
+		}
+
 		if (ShouldCreatePak(Params, SC))
 		{
 			if (SC.CrashReporterUFSFiles.Count > 0)
diff --git a/Engine/Source/Runtime/CoreUObject/Public/UObject/SavePackage.h b/Engine/Source/Runtime/CoreUObject/Public/UObject/SavePackage.h
--- a/Engine/Source/Runtime/CoreUObject/Public/UObject/SavePackage.h
+++ b/Engine/Source/Runtime/CoreUObject/Public/UObject/SavePackage.h
@@ -83,8 +83,9 @@ class FPackageStoreBulkDataManifest : public IBulkDataManifest
 {
 public:
 	COREUOBJECT_API FPackageStoreBulkDataManifest(const FString& InRootPath);
+	COREUOBJECT_API virtual ~FPackageStoreBulkDataManifest() = default;
 
-	bool Load();
+	bool COREUOBJECT_API Load();
 	virtual void Save() override;
 	
 	const FString& GetFilename() const { return Filename; }
@@ -106,7 +107,7 @@ public:
 		TArray<BulkDataDesc> Data;
 	};
 
-	const PackageDesc* Find(const FString& PackageName) const;
+	COREUOBJECT_API const PackageDesc* Find(const FString& PackageName) const;
 
 private:
 	virtual void AddFileAccess(const FString& PackageFilename, uint16 InIndex, uint64 InOffset, uint64 InSize) override;
