AsyncLoading2: Move generation of global namemap from cooker to iostore commandlet (running post cook).
Add PackageNameMap to PackageSummary to remap local name indices to global name indices when serializing export blobs.

diff --git a/Engine/Source/Developer/IoStoreUtilities/Private/IoStoreUtilities.cpp b/Engine/Source/Developer/IoStoreUtilities/Private/IoStoreUtilities.cpp
--- a/Engine/Source/Developer/IoStoreUtilities/Private/IoStoreUtilities.cpp
+++ b/Engine/Source/Developer/IoStoreUtilities/Private/IoStoreUtilities.cpp
@@ -31,8 +31,11 @@ IMPLEMENT_MODULE(FDefaultModuleImpl, IoStoreUtilities);
 DEFINE_LOG_CATEGORY_STATIC(LogIoStore, Log, All);
 
 #define OUTPUT_CHUNKID_DIRECTORY 0
+#define OUTPUT_NAMEMAP_CSV 0
+#define OUTPUT_IMPORTMAP_CSV 0
 #define SKIP_WRITE_CONTAINER 0
 #define SKIP_BULKDATA 0
+#define USE_NAMEMAPBUILDER 1
 
 struct FContainerTarget 
 {
@@ -41,6 +44,7 @@ struct FContainerTarget
 	FString OutputDirectory;
 };
 
+#if !USE_NAMEMAPBUILDER
 class FZenGlobalNameMap
 {
 public:
@@ -120,6 +124,109 @@ private:
 	TMap<FNameEntryId, int32> DisplayEntryToIndex;
 	TMap<FNameEntryId, int32> ComparisonEntryToIndex;
 };
+#else
+class FNameMapBuilder
+{
+public:
+	void MarkNameAsReferenced(const FName& Name)
+	{
+		const FNameEntryId Id = Name.GetComparisonIndex();
+		int32& Index = NameIndices.FindOrAdd(Id);
+		if (Index == 0)
+		{
+			Index = NameIndices.Num();
+			NameMap.Add(Id);
+		}
+		// debug counts
+		{
+			const int32 Number = Name.GetNumber();
+			TTuple<int32,int32,int32>& Counts = DebugNameCounts.FindOrAdd(Id);
+
+			if (Number == 0)
+			{
+				++Counts.Get<0>();
+			}
+			else
+			{
+				++Counts.Get<1>();
+				if (Number > Counts.Get<2>())
+				{
+					Counts.Get<2>() = Number;
+				}
+			}
+		}
+	}
+
+	int32 MapName(const FName& Name)
+	{
+		const FNameEntryId Id = Name.GetComparisonIndex();
+		const int32* Index = NameIndices.Find(Id);
+		check(Index);
+		return Index ? *Index - 1 : INDEX_NONE;
+	}
+
+	void SerializeName(FArchive& A, const FName& N)
+	{
+		int32 NameIndex = MapName(N);
+		int32 NameNumber = N.GetNumber();
+		A << NameIndex << NameNumber;
+	}
+
+	int32 Num()
+	{
+		return NameIndices.Num();
+	}
+
+	int32 Save(const FString& Filename)
+	{
+		int32 TotalSize = 0;
+		{
+			TUniquePtr<FArchive> BinArchive(IFileManager::Get().CreateFileWriter(*Filename));
+			int32 NameCount = NameMap.Num();
+			*BinArchive << NameCount;
+			for (int32 I = 0; I < NameCount; ++I)
+			{
+				FName::GetEntry(NameMap[I])->Write(*BinArchive);
+			}
+			TotalSize = BinArchive->TotalSize();
+		}
+#if OUTPUT_NAMEMAP_CSV
+		{
+			FString CsvFilePath = Filename;
+			CsvFilePath.Append(TEXT(".csv"));
+			TUniquePtr<FArchive> CsvArchive(IFileManager::Get().CreateFileWriter(*CsvFilePath));
+			if (CsvArchive)
+			{
+				TCHAR Name[FName::StringBufferSize];
+				ANSICHAR Line[MAX_SPRINTF + FName::StringBufferSize];
+				ANSICHAR Header[] = "Length\tMaxNumber\tNumberCount\tBaseCount\tTotalCount\tFName\n";
+				CsvArchive->Serialize(Header, sizeof(Header) - 1);
+				for (auto& Counts : DebugNameCounts)
+				{
+					const int32 NameLen = FName::CreateFromDisplayId(Counts.Key, 0).ToString(Name);
+					FCStringAnsi::Sprintf(Line, "%d\t%d\t%d\t%d\t%d\t",
+						NameLen, Counts.Value.Get<2>(), Counts.Value.Get<1>(), Counts.Value.Get<0>(), Counts.Value.Get<0>() + Counts.Value.Get<1>());
+					ANSICHAR* L = Line + FCStringAnsi::Strlen(Line);
+					const TCHAR* N = Name;
+					while (*N)
+					{
+						*L++ = CharCast<ANSICHAR,TCHAR>(*N++);
+					}
+					*L++ = '\n';
+					CsvArchive.Get()->Serialize(Line, L - Line);
+				}
+			}
+		}
+#endif
+		return TotalSize;
+	}
+
+private:
+	TMap<FNameEntryId, int32> NameIndices;
+	TArray<FNameEntryId> NameMap;
+	TMap<FNameEntryId, TTuple<int32,int32,int32>> DebugNameCounts; // <Number0Count,OtherNumberCount,MaxNumber>
+};
+#endif
 
 #if OUTPUT_CHUNKID_DIRECTORY
 class FChunkIdCsv
@@ -190,11 +297,13 @@ static FIoChunkId CreateZenChunkId(uint32 NameIndex, uint32 NameNumber, uint16 C
 struct FZenPackageSummary
 {
 	uint32 PackageFlags;
+	int32 NameMapOffset;
 	int32 ExportMapOffset;
 	int32 ExportBundlesOffset;
 	int32 GraphDataOffset;
 	int32 GraphDataSize;
 	int32 BulkDataStartOffset;
+	int32 Pad;
 };
 
 enum EPreloadDependencyType
@@ -360,6 +469,7 @@ struct FPackage
 	FString RelativeFileName;
 	int32 GlobalPackageId = 0;
 	uint32 PackageFlags = 0;
+	int32 NameCount = 0;
 	int32 ImportCount = 0;
 	int32 ImportOffset = 0;
 	int32 ScriptArcsCount = 0;
@@ -374,6 +484,7 @@ struct FPackage
 	int64 UAssetSize = 0;
 	int64 SummarySize = 0;
 	int64 UGraphSize = 0;
+	int64 NameMapSize = 0;
 	int64 ExportMapSize = 0;
 	int64 ExportBundlesSize = 0;
 
@@ -385,6 +496,9 @@ struct FPackage
 	TSet<FPackage*> AllReachablePackages;
 	TSet<FPackage*> ImportedPreloadPackages;
 
+	TArray<FNameEntryId> NameMap;
+	TArray<int32> NameIndices;
+
 	TArray<int32> Imports;
 	TArray<int32> Exports;
 	TArray<FArc> InternalArcs;
@@ -972,11 +1086,6 @@ int32 CreateTarget(const FContainerTarget& Target)
 	FString OutputDir = Target.OutputDirectory;
 	FString RelativePrefixForLegacyFilename = TEXT("../../../");
 
-	FZenGlobalNameMap GlobalNameMap;
-	FString GlobalNameMapPath = CookedDir / TEXT("megafile.unamemap");
-
-	GlobalNameMap.Load(GlobalNameMapPath);
-
 	FPackageStoreBulkDataManifest BulkDataManifest(CookedDir);
 	const bool bWithBulkDataManifest = BulkDataManifest.Load();
 	// if (!BulkDataManifest.Load())
@@ -996,12 +1105,16 @@ int32 CreateTarget(const FContainerTarget& Target)
 	IFileManager::Get().FindFilesRecursive(FileNames, *CookedDir, TEXT("*.umap"), true, false, false);
 	UE_LOG(LogIoStore, Display, TEXT("Found '%d' files"), FileNames.Num());
 
-	auto DeserializeName = [&](FArchive& A, FName& N)
+#if USE_NAMEMAPBUILDER
+	FNameMapBuilder NameMapBuilder;
+	auto SerializeName = [&](FArchive& A, const FName& N)
 	{
-		int32 NameIndex, NameNumber;
-		A << NameIndex << NameNumber;
-		N = GlobalNameMap.GetNameFromDisplayIndex(NameIndex, NameNumber);
+		NameMapBuilder.SerializeName(A, N);
 	};
+#else
+	FString GlobalNameMapPath = CookedDir / TEXT("megafile.unamemap");
+	FZenGlobalNameMap GlobalNameMap;
+	GlobalNameMap.Load(GlobalNameMapPath);
 
 	auto SerializeName = [&](FArchive& A, const FName& N)
 	{
@@ -1016,9 +1129,8 @@ int32 CreateTarget(const FContainerTarget& Target)
 			check(false);
 		}
 	};
+#endif
 
-	TArray<FName> Names;
-	TSet<FName> UniqueImports;
 	uint64 NameSize = 0;
 	TArray<FObjectImport> Imports;
 	TArray<FObjectExport> Exports;
@@ -1094,6 +1206,7 @@ int32 CreateTarget(const FContainerTarget& Target)
 		Package.GlobalPackageId = FileIndex;
 		Package.UAssetSize = Ar->TotalSize();
 		Package.SummarySize = Ar->Tell() - SummaryStartPos;
+		Package.NameCount = Summary.NameCount;
 		Package.ImportCount = Summary.ImportCount;
 		Package.ExportCount = Summary.ExportCount;
 		Package.PackageFlags = Summary.PackageFlags;
@@ -1102,6 +1215,46 @@ int32 CreateTarget(const FContainerTarget& Target)
 		Package.RelativeFileName = RelativePrefixForLegacyFilename;
 		Package.RelativeFileName.Append(*FileName + CookedDir.Len());
 
+		if (Summary.NameCount > 0)
+		{
+			Ar->Seek(Summary.NameOffset);
+			uint64 LastOffset = Summary.NameOffset;
+
+			Package.NameMap.Reserve(Summary.NameCount);
+			Package.NameIndices.Reserve(Summary.NameCount);
+			FNameEntrySerialized NameEntry(ENAME_LinkerConstructor);
+
+			for (int32 I = 0; I < Summary.NameCount; ++I)
+			{
+				*Ar << NameEntry;
+				FName Name(NameEntry);
+#if USE_NAMEMAPBUILDER
+				NameMapBuilder.MarkNameAsReferenced(Name);
+				Package.NameMap.Emplace(Name.GetDisplayIndex());
+				Package.NameIndices.Add(NameMapBuilder.MapName(Name));
+#endif
+			}
+
+			NameSize += Ar->Tell() - Summary.NameOffset;
+		}
+
+#if USE_NAMEMAPBUILDER
+		auto DeserializeName = [&](FArchive& A, FName& N)
+		{
+			int32 DisplayIndex, NameNumber;
+			A << DisplayIndex << NameNumber;
+			FNameEntryId DisplayEntry = Package.NameMap[DisplayIndex];
+			N = FName::CreateFromDisplayId(DisplayEntry, NameNumber);
+		};
+#else
+		auto DeserializeName = [&](FArchive& A, FName& N)
+		{
+			int32 NameIndex, NameNumber;
+			A << NameIndex << NameNumber;
+			N = GlobalNameMap.GetNameFromDisplayIndex(NameIndex, NameNumber);
+		};
+#endif
+
 		if (Summary.ImportCount > 0)
 		{
 			Ar->Seek(Summary.ImportOffset);
@@ -1238,6 +1391,7 @@ int32 CreateTarget(const FContainerTarget& Target)
 	uint64 ExportSize = Exports.Num() * 104;
 	uint64 PreloadDependenciesSize = PreloadDependencies.Num() * 4;
 
+#if OUTPUT_NAMEMAP_CSV
 	FString CsvFilePath = SlimportArchive->GetArchiveName();
 	CsvFilePath.Append(TEXT(".csv"));
 	TUniquePtr<FArchive> CsvArchive(IFileManager::Get().CreateFileWriter(*CsvFilePath));
@@ -1260,6 +1414,8 @@ int32 CreateTarget(const FContainerTarget& Target)
 			CsvArchive.Get()->Serialize(Line, L - Line);
 		}
 	}
+#endif
+
 	if (GlimportArchive)
 	{
 		int32 Pad = 0;
@@ -1507,17 +1663,22 @@ int32 CreateTarget(const FContainerTarget& Target)
 		}
 		Package.ExportBundlesSize = ZenExportBundlesArchive.Tell();
 
+		FName RelativeFileName(*Package.RelativeFileName);
+#if USE_NAMEMAPBUILDER
+		NameMapBuilder.MarkNameAsReferenced(Package.Name);
+		int32 PackageNameIndex = NameMapBuilder.MapName(Package.Name);
+		int32 PackageNameNumber = Package.Name.GetNumber();
+		NameMapBuilder.MarkNameAsReferenced(RelativeFileName);
+#else
 		int32 PackageNameIndex = GlobalNameMap.GetOrCreateComparisonIndex(Package.Name);
 		int32 PackageNameNumber = Package.Name.GetNumber();
-
-		FName RelativeFileName(*Package.RelativeFileName);
-		int32 FileNameIndex = GlobalNameMap.GetOrCreateComparisonIndex(RelativeFileName);
-		int32 FileNameNumber = Package.Name.GetNumber();
+		GlobalNameMap.GetOrCreateComparisonIndex(RelativeFileName);
+#endif
 
 		ensure(Package.GlobalPackageId == PackageIndex++);
 		{
-			*StoreTocArchive << PackageNameIndex << PackageNameNumber;
-			*StoreTocArchive << FileNameIndex << FileNameNumber;
+			SerializeName(*StoreTocArchive, Package.Name);
+			SerializeName(*StoreTocArchive, RelativeFileName);
 			*StoreTocArchive << Package.ImportCount;
 			*StoreTocArchive << Package.SlimportCount;
 			*StoreTocArchive << Package.SlimportOffset;
@@ -1528,9 +1689,12 @@ int32 CreateTarget(const FContainerTarget& Target)
 			*StoreTocArchive << Package.ScriptArcsCount;
 		}
 
+		Package.NameMapSize = Package.NameIndices.Num() * Package.NameIndices.GetTypeSize();
+
 		{
 			const uint64 ZenSummarySize = 
 						sizeof(FZenPackageSummary)
+						+ Package.NameMapSize
 						+ Package.ExportMapSize
 						+ Package.ExportBundlesSize
 						+ Package.UGraphSize;
@@ -1545,6 +1709,12 @@ int32 CreateTarget(const FContainerTarget& Target)
 			FBufferWriter ZenAr(ZenSummaryBuffer, ZenSummarySize); 
 			ZenAr.Seek(sizeof(FZenPackageSummary));
 
+			// NameMap data
+			{
+				ZenSummary->NameMapOffset = ZenAr.Tell();
+				ZenAr.Serialize(Package.NameIndices.GetData(), Package.NameMapSize);
+			}
+
 			// ExportMap data
 			{
 				check(ZenExportMapArchive.Tell() == Package.ExportMapSize);	
@@ -1680,7 +1850,14 @@ int32 CreateTarget(const FContainerTarget& Target)
 		}
 	}
 
-	GlobalNameMap.Save(OutputDir / TEXT("Container.namemap"));
+	UE_LOG(LogIoStore, Display, TEXT("Saving NameMap..."));
+	FString NameMapFilename = OutputDir / TEXT("Container.namemap");
+#if USE_NAMEMAPBUILDER
+	uint64 GlobalNameMapSize = NameMapBuilder.Save(NameMapFilename);
+#else
+	GlobalNameMap.Save(NameMapFilename);
+#endif
+
 	UE_LOG(LogIoStore, Display, TEXT("Saving StoreToc..."));
 	StoreTocArchive->Close();
 
@@ -1690,6 +1867,8 @@ int32 CreateTarget(const FContainerTarget& Target)
 	uint64 SummarySize = 0;
 	uint64 UGraphSize = 0;
 	uint64 ExportMapSize = 0;
+	uint64 NameMapSize = 0;
+	uint64 NameMapCount = 0;
 	uint64 ZenPackageSummarySize = FileNames.Num() * sizeof(FZenPackageSummary);
 	uint64 StoreTocSize = StoreTocArchive->Tell();
 	uint64 GlimportSize = GlimportArchive->Tell();
@@ -1699,6 +1878,7 @@ int32 CreateTarget(const FContainerTarget& Target)
 	uint64 CircularPackagesCount = 0;
 	uint64 TotalInternalArcCount = 0;
 	uint64 TotalExternalArcCount = 0;
+	uint64 NameCount = 0;
 
 	uint64 PackagesWithCircularDependenciesCount = 0;
 	uint64 PackagesWithoutImportDependenciesCount = 0;
@@ -1715,9 +1895,12 @@ int32 CreateTarget(const FContainerTarget& Target)
 		SummarySize += Package.SummarySize;
 		UGraphSize += Package.UGraphSize;
 		ExportMapSize += Package.ExportMapSize;
+		NameMapSize += Package.NameMapSize;
+		NameMapCount += Package.NameIndices.Num();
 		ScriptArcsCount += Package.ScriptArcsCount;
 		CircularPackagesCount += Package.bHasCircularImportDependencies;
 		TotalInternalArcCount += Package.InternalArcs.Num();
+		NameCount += Package.NameMap.Num();
 		PackagesWithCircularDependenciesCount += Package.bHasCircularPostLoadDependencies;
 		PackagesWithoutPreloadDependenciesCount += Package.ImportedPreloadPackages.Num() == 0;
 		PackagesWithoutImportDependenciesCount += Package.ImportedPackages.Num() == 0;
@@ -1735,32 +1918,35 @@ int32 CreateTarget(const FContainerTarget& Target)
 		}
 	}
 
-	UE_LOG(LogIoStore, Display, TEXT(""));
 	UE_LOG(LogIoStore, Display, TEXT("-------------------- IoStore Summary: %s --------------------"), *Target.TargetPlatform->PlatformName());
-
-	UE_LOG(LogIoStore, Display, TEXT("Packages: %5d total, %d circular dependencies, %d no preload dependencies, %d no import dependencies"),
+	UE_LOG(LogIoStore, Display, TEXT("Packages: %8d total, %d circular dependencies, %d no preload dependencies, %d no import dependencies"),
 		PackageMap.Num(), PackagesWithCircularDependenciesCount, PackagesWithoutPreloadDependenciesCount, PackagesWithoutImportDependenciesCount);
-	UE_LOG(LogIoStore, Display, TEXT("Bundles:  %5d total, %d entries"), BundleCount, BundleEntryCount);
-	UE_LOG(LogIoStore, Display, TEXT("Exports:  %5d total"), GlobalExports.Num());
-	
+	UE_LOG(LogIoStore, Display, TEXT("Bundles:  %8d total, %d entries, %d export objects"), BundleCount, BundleEntryCount, GlobalExports.Num());
+	UE_LOG(LogIoStore, Display, TEXT("Chunks:   %8d summary chunks, %d export chunks, %d bulk data chunks, %d partial bulk data chunks"), SummaryChunkCount, ExportChunkCount, BulkChunkCount, BulkPartialChunkCount);
+
 	UE_LOG(LogIoStore, Display, TEXT("Pak: %12.2f MB UExp, %d files"), (double)UExpSize / 1024.0 / 1024.0, FileNames.Num());
 	UE_LOG(LogIoStore, Display, TEXT("Pak: %12.2f MB UAsset/UMap, %d files"), (double)UAssetSize / 1024.0 / 1024.0, FileNames.Num());
 	UE_LOG(LogIoStore, Display, TEXT("Pak: %12.2f MB FPackageFileSummary"), (double)SummarySize / 1024.0 / 1024.0);
+	UE_LOG(LogIoStore, Display, TEXT("Pak: %12.2f MB NameMap, %d names"), (double)NameSize / 1024.0 / 1024.0, NameCount);
 	UE_LOG(LogIoStore, Display, TEXT("Pak: %12.2f MB ExportMap, %d exports"), (double)ExportSize / 1024.0 / 1024.0, Exports.Num());
 	UE_LOG(LogIoStore, Display, TEXT("Pak: %12.2f MB ImportMap, %d imports (%d UPackages)"), (double)ImportSize / 1024.0 / 1024.0, Imports.Num(), UPackageImports);
 	UE_LOG(LogIoStore, Display, TEXT("Pak: %12.2f MB PreloadDependencies, %d imports + %d exports = %d"),
 		(double)PreloadDependenciesSize / 1024.0 / 1024.0, ImportPreloadCount, ExportPreloadCount, PreloadDependencies.Num());
-	UE_LOG(LogIoStore, Display, TEXT(" "));
-	UE_LOG(LogIoStore, Display, TEXT("IoStore: %8.2f MB StoreToc"), (double)StoreTocSize / 1024.0 / 1024.0);
-	UE_LOG(LogIoStore, Display, TEXT("IoStore: %8.2f MB StoreGlimports"), (double)GlimportSize / 1024.0 / 1024.0);
-	UE_LOG(LogIoStore, Display, TEXT("IoStore: %8.2f MB StoreSlimports"), (double)SlimportSize / 1024.0 / 1024.0);
+
+	UE_LOG(LogIoStore, Display, TEXT("IoStore: %8.2f MB GlobalToc"), (double)StoreTocSize / 1024.0 / 1024.0);
+#if USE_NAMEMAPBUILDER
+	UE_LOG(LogIoStore, Display, TEXT("IoStore: %8.2f MB GlobalNameMap, %d unique names"), (double)GlobalNameMapSize / 1024.0 / 1024.0, NameMapBuilder.Num());
+#endif
+	UE_LOG(LogIoStore, Display, TEXT("IoStore: %8.2f MB GlobalImportMap, %d unique imports, %d unique import packages"),
+		(double)GlimportSize / 1024.0 / 1024.0, GlobalImportsByFullName.Num(), UniqueImportPackages);
 	UE_LOG(LogIoStore, Display, TEXT("IoStore: %8.2f MB PackageSummary"), (double)ZenPackageSummarySize / 1024.0 / 1024.0);
-	UE_LOG(LogIoStore, Display, TEXT("IoStore: %8.2f MB ExportMap"), (double)ExportMapSize / 1024.0 / 1024.0);
-	UE_LOG(LogIoStore, Display, TEXT("IoStore: %8.2f MB UGraph, %d internal arcs, %d external arcs"), (double)UGraphSize / 1024.0 / 1024.0, TotalInternalArcCount, TotalExternalArcCount);
+	UE_LOG(LogIoStore, Display, TEXT("IoStore: %8.2f MB PackageNameMap, %d indices"), (double)NameMapSize / 1024.0 / 1024.0, NameMapCount);
+	UE_LOG(LogIoStore, Display, TEXT("IoStore: %8.2f MB PackageImportMap"), (double)SlimportSize / 1024.0 / 1024.0);
+	UE_LOG(LogIoStore, Display, TEXT("IoStore: %8.2f MB PackageExportMap"), (double)ExportMapSize / 1024.0 / 1024.0);
+	UE_LOG(LogIoStore, Display, TEXT("IoStore: %8.2f MB SerializedArcs, %d internal arcs, %d external arcs, %d circular packages (%d chains)"),
+		(double)UGraphSize / 1024.0 / 1024.0, TotalInternalArcCount, TotalExternalArcCount, CircularPackagesCount, CircularChains.Num());
 	UE_LOG(LogIoStore, Display, TEXT("IoStore: %8.2f MB ScriptArcs, %d arcs"), (double)ScriptArcsSize / 1024.0 / 1024.0, ScriptArcsCount);
-	UE_LOG(LogIoStore, Display, TEXT("IoStore: %8d unique imports, %d unique packages"), GlobalImportsByFullName.Num(), UniqueImportPackages);
-	UE_LOG(LogIoStore, Display, TEXT("IoStore: %8d circular packages, %d unique chains resolved"), CircularPackagesCount, CircularChains.Num());
-	UE_LOG(LogIoStore, Display, TEXT("IoStore: %8d summary chunks, %d export chunks, %d bulk data chunks, %d partial bulk data chunks"), SummaryChunkCount, ExportChunkCount, BulkChunkCount, BulkPartialChunkCount);
+	UE_LOG(LogIoStore, Display, TEXT("-------------------- IoStore Summary: %s --------------------"), *Target.TargetPlatform->PlatformName());
 
 	return 0;
 }
@@ -1791,31 +1977,36 @@ int32 CreateIoStoreContainerFiles(const TCHAR* CmdLine)
 	}
 
 	ITargetPlatformManagerModule& TargetPlatformManager = *GetTargetPlatformManager();
+	ITargetPlatform* TargetPlatform = TargetPlatformManager.FindTargetPlatform(TargetPlatformName);
 
-	if (ITargetPlatform* TargetPlatform = TargetPlatformManager.FindTargetPlatform(TargetPlatformName))
+	if (!TargetPlatform)
 	{
-		FString TargetCookedDirectory = FPaths::ProjectSavedDir() / TEXT("Cooked") / TargetPlatform->PlatformName();
-		FString TargetOutputDirectory = OutputDirectory.Len() > 0
-			? OutputDirectory
-			: TargetCookedDirectory / FApp::GetProjectName();
+		UE_LOG(LogIoStore, Display, TEXT("Unknown target platform '%s'"), *TargetPlatformName);
+		return -1;
 
-		FString GlobalNameMapFilePath = TargetCookedDirectory / TEXT("megafile.unamemap");
+	}
 
-		if (IFileManager::Get().FileExists(*GlobalNameMapFilePath))
-		{
-			FContainerTarget Target { TargetPlatform, TargetCookedDirectory, TargetOutputDirectory };
+	FString TargetCookedDirectory = FPaths::ProjectSavedDir() / TEXT("Cooked") / TargetPlatform->PlatformName();
+	FString TargetOutputDirectory = OutputDirectory.Len() > 0
+		? OutputDirectory
+		: TargetCookedDirectory / FApp::GetProjectName();
 
-			UE_LOG(LogIoStore, Display, TEXT("Creating target: '%s' using output directory: '%s'"), *Target.TargetPlatform->PlatformName(), *Target.OutputDirectory);
+	FString GlobalNameMapFilePath = TargetCookedDirectory / TEXT("megafile.unamemap");
 
-			CreateTarget(Target);
-		}
-		else
-		{
-			UE_LOG(LogIoStore, Error,
-				TEXT("Expected to find global name map '%s' for target '%s'"),
-				*GlobalNameMapFilePath, *TargetPlatform->PlatformName());
-		}
+#if !USE_NAMEMAPBUILDER
+	if (!IFileManager::Get().FileExists(*GlobalNameMapFilePath))
+	{
+		UE_LOG(LogIoStore, Error,
+			TEXT("Expected to find global name map '%s' for target '%s'"),
+			*GlobalNameMapFilePath, *TargetPlatform->PlatformName());
+		return -1;
 	}
+#endif
+
+	FContainerTarget Target { TargetPlatform, TargetCookedDirectory, TargetOutputDirectory };
+
+	UE_LOG(LogIoStore, Display, TEXT("Creating target: '%s' using output directory: '%s'"), *Target.TargetPlatform->PlatformName(), *Target.OutputDirectory);
 
+	CreateTarget(Target);
 	return 0;
 }
diff --git a/Engine/Source/Runtime/CoreUObject/Private/Serialization/AsyncLoading2.cpp b/Engine/Source/Runtime/CoreUObject/Private/Serialization/AsyncLoading2.cpp
--- a/Engine/Source/Runtime/CoreUObject/Private/Serialization/AsyncLoading2.cpp
+++ b/Engine/Source/Runtime/CoreUObject/Private/Serialization/AsyncLoading2.cpp
@@ -177,11 +177,13 @@ struct FPackageStoreEntry
 struct FPackageSummary
 {
 	uint32 PackageFlags;
+	int32 NameMapOffset;
 	int32 ExportMapOffset;
 	int32 ExportBundlesOffset;
 	int32 GraphDataOffset;
 	int32 GraphDataSize;
 	int32 BulkDataStartOffset;
+	int32 Pad;
 };
 
 class FGlobalNameMap
@@ -609,7 +611,7 @@ public:
 
 	void BadNameIndexError(int32 NameIndex)
 	{
-		UE_LOG(LogLinker, Error, TEXT("Bad name index %i/%i"), NameIndex, NameMap->Num());
+		UE_LOG(LogStreaming, Error, TEXT("Bad name index %i/%i"), NameIndex, GlobalNameMap->Num());
 	}
 
 	inline virtual FArchive& operator<<(FName& Name) override
@@ -620,10 +622,16 @@ public:
 		int32 Number = 0;
 		Ar << Number;
 
-		if (NameMap->IsValidIndex(NameIndex))
+		// TODO: Remove support for old global namemap
+		if (PackageNameMap)
+		{
+			NameIndex = PackageNameMap[NameIndex];
+		}
+
+		if (GlobalNameMap->IsValidIndex(NameIndex))
 		{
 			// if the name wasn't loaded (because it wasn't valid in this context)
-			FNameEntryId MappedName = (*NameMap)[NameIndex];
+			FNameEntryId MappedName = (*GlobalNameMap)[NameIndex];
 
 			// simply create the name from the NameMap's name and the serialized instance number
 			Name = FName::CreateFromDisplayId(MappedName, Number);
@@ -644,7 +652,8 @@ private:
 
 	UObject* TemplateForGetArchetypeFromLoader = nullptr;
 	FPackageImportStore* ImportStore = nullptr;
-	const TArray<FNameEntryId>* NameMap = nullptr;
+	const int32* PackageNameMap = nullptr;
+	const TArray<FNameEntryId>* GlobalNameMap = nullptr;
 	const TArray<UObject*>* Exports = nullptr;
 	TArray<FExternalReadCallback>* ExternalReadDependencies;
 	FExternalIoRequests* ExternalIoRequests;
@@ -1285,6 +1294,7 @@ private:
 	FExternalIoRequests ExternalIoRequests;
 	int32 ExportCount = 0;
 	const FExportMapEntry* ExportMap = nullptr;
+	const int32* PackageNameMap = nullptr;
 	TArray<UObject*> Exports;
 	FPackageImportStore ImportStore;
 
@@ -2856,7 +2866,8 @@ void FAsyncPackage2::EventDrivenSerializeExport(int32 LocalExportIndex, const ui
 	Ar.ArAllowLazyLoading = true;
 
 	// FSimpleExportArchive special fields
-	Ar.NameMap = &AsyncLoadingThread.GlobalNameMap.GetNameEntries();
+	Ar.PackageNameMap = PackageNameMap;
+	Ar.GlobalNameMap = &AsyncLoadingThread.GlobalNameMap.GetNameEntries();
 	Ar.ImportStore = &ImportStore;
 	Ar.Exports = &Exports;
 	Ar.ExternalReadDependencies = &ExternalReadDependencies;
@@ -4208,8 +4219,12 @@ EAsyncPackageState::Type FAsyncPackage2::FinishLinker()
 	const uint8* PackageSummaryData = PackageSummaryIoBuffer.Data();
 	const FPackageSummary* PackageSummary = reinterpret_cast<const FPackageSummary*>(PackageSummaryData);
 	
-	ExportMap = reinterpret_cast<const FExportMapEntry*>(PackageSummaryData + PackageSummary->ExportMapOffset);
+	// TODO: Remove support for old global namemap
+	PackageNameMap = (PackageSummary->NameMapOffset == PackageSummary->ExportMapOffset) ? 
+		nullptr :
+		reinterpret_cast<const int32*>(PackageSummaryData + PackageSummary->NameMapOffset);
 
+	ExportMap = reinterpret_cast<const FExportMapEntry*>(PackageSummaryData + PackageSummary->ExportMapOffset);
 	ExportBundles = reinterpret_cast<const FExportBundleHeader*>(PackageSummaryData + PackageSummary->ExportBundlesOffset);
 	ExportBundleEntries = reinterpret_cast<const FExportBundleEntry*>(ExportBundles + ExportBundleCount);
 
