AsyncLoading2: Remove duplicated serialized arcs.
Disk size down from ~100 Mb to ~60 MB.
Runtime arcs down from ~150 000 to ~130 000 when loading ~60 000 packages.

diff --git a/Engine/Source/Developer/IoStoreUtilities/Private/IoStoreUtilities.cpp b/Engine/Source/Developer/IoStoreUtilities/Private/IoStoreUtilities.cpp
--- a/Engine/Source/Developer/IoStoreUtilities/Private/IoStoreUtilities.cpp
+++ b/Engine/Source/Developer/IoStoreUtilities/Private/IoStoreUtilities.cpp
@@ -354,7 +354,6 @@ struct FExportGraphNode;
 struct FExportBundle
 {
 	TArray<FExportGraphNode*> Nodes;
-	TSet<uint32> AddedScriptDependencies;
 };
 
 struct FPackage;
@@ -720,16 +719,16 @@ static void AddPostLoadArc(FPackage& FromPackage, FPackage& ToPackage)
 static void AddStartPostLoadArc(FPackage& FromPackage, FPackage& ToPackage)
 {
 	TArray<FArc>& ExternalArcs = ToPackage.ExternalArcs.FindOrAdd(FromPackage.Name);
-	ensure(!ExternalArcs.Contains(FArc({EEventLoadNode2::Package_ExportsSerialized, EEventLoadNode2::Package_StartPostLoad})));
-	ensure(!ExternalArcs.Contains(FArc({EEventLoadNode2::Package_StartPostLoad, EEventLoadNode2::Package_StartPostLoad})));
+	check(!ExternalArcs.Contains(FArc({EEventLoadNode2::Package_ExportsSerialized, EEventLoadNode2::Package_StartPostLoad})));
+	check(!ExternalArcs.Contains(FArc({EEventLoadNode2::Package_StartPostLoad, EEventLoadNode2::Package_StartPostLoad})));
 	ExternalArcs.Add({ EEventLoadNode2::Package_StartPostLoad, EEventLoadNode2::Package_StartPostLoad });
 }
 
 static void AddExportsDoneArc(FPackage& FromPackage, FPackage& ToPackage)
 {
 	TArray<FArc>& ExternalArcs = ToPackage.ExternalArcs.FindOrAdd(FromPackage.Name);
-	ensure(!ExternalArcs.Contains(FArc({EEventLoadNode2::Package_ExportsSerialized, EEventLoadNode2::Package_StartPostLoad})));
-	ensure(!ExternalArcs.Contains(FArc({EEventLoadNode2::Package_StartPostLoad, EEventLoadNode2::Package_StartPostLoad})));
+	check(!ExternalArcs.Contains(FArc({EEventLoadNode2::Package_ExportsSerialized, EEventLoadNode2::Package_StartPostLoad})));
+	check(!ExternalArcs.Contains(FArc({EEventLoadNode2::Package_StartPostLoad, EEventLoadNode2::Package_StartPostLoad})));
 	ExternalArcs.Add({ EEventLoadNode2::Package_ExportsSerialized, EEventLoadNode2::Package_StartPostLoad });
 }
 
@@ -737,21 +736,22 @@ static void AddInternalBundleArc(FPackage& Package, uint32 FromBundleIndex, uint
 {
 	uint32 FromNodeIndex = EEventLoadNode2::Package_NumPhases + FromBundleIndex * EEventLoadNode2::ExportBundle_NumPhases + EEventLoadNode2::ExportBundle_Process;
 	uint32 ToNodeIndex = EEventLoadNode2::Package_NumPhases + ToBundleIndex * EEventLoadNode2::ExportBundle_NumPhases + EEventLoadNode2::ExportBundle_StartIo;
+	check(!Package.InternalArcs.Contains(FArc({FromNodeIndex, ToNodeIndex})));
 	Package.InternalArcs.Add({ FromNodeIndex, ToNodeIndex });
 }
 
-static void AddExternalBundleArc(FPackage& FromPackage, uint32 FromBundleIndex, FPackage& ToPackage, uint32 ToBundleIndex)
+static void AddUniqueExternalBundleArc(FPackage& FromPackage, uint32 FromBundleIndex, FPackage& ToPackage, uint32 ToBundleIndex)
 {
 	uint32 FromNodeIndex = EEventLoadNode2::Package_NumPhases + FromBundleIndex * EEventLoadNode2::ExportBundle_NumPhases + EEventLoadNode2::ExportBundle_Process;
 	uint32 ToNodeIndex = EEventLoadNode2::Package_NumPhases + ToBundleIndex * EEventLoadNode2::ExportBundle_NumPhases + EEventLoadNode2::ExportBundle_StartIo;
 	TArray<FArc>& ExternalArcs = ToPackage.ExternalArcs.FindOrAdd(FromPackage.Name);
-	ExternalArcs.Add({ FromNodeIndex, ToNodeIndex });
+	ExternalArcs.AddUnique({ FromNodeIndex, ToNodeIndex });
 }
 
-static void AddScriptBundleArc(FPackage& Package, uint32 GlobalImportIndex, uint32 BundleIndex)
+static void AddUniqueScriptBundleArc(FPackage& Package, uint32 GlobalImportIndex, uint32 BundleIndex)
 {
 	uint32 NodeIndex = EEventLoadNode2::Package_NumPhases + BundleIndex * EEventLoadNode2::ExportBundle_NumPhases + EEventLoadNode2::ExportBundle_StartIo;
-	Package.ScriptArcs.Add({ GlobalImportIndex, NodeIndex });
+	Package.ScriptArcs.AddUnique({ GlobalImportIndex, NodeIndex });
 }
 
 static void AddPostLoadDependenciesRecursive(FPackage& Package, FPackage& ImportedPackage, TSet<FPackage*>& Visited)
@@ -933,16 +933,11 @@ static void BuildBundles(FExportGraph& ExportGraph, TMap<FName, FPackage*>& Pack
 		{
 			uint32* FindDependentBundleIndex = ExternalDependency->Package->ExportBundleMap.Find(ExternalDependency);
 			check(FindDependentBundleIndex);
-			AddExternalBundleArc(*ExternalDependency->Package, *FindDependentBundleIndex, *Package, BundleIndex);
+			AddUniqueExternalBundleArc(*ExternalDependency->Package, *FindDependentBundleIndex, *Package, BundleIndex);
 		}
 		for (uint32 ScriptDependencyGlobalImportIndex : Node->ScriptDependencies)
 		{
-			bool bAlreadyInSet;
-			Package->ExportBundles[BundleIndex].AddedScriptDependencies.Add(ScriptDependencyGlobalImportIndex, &bAlreadyInSet);
-			if (!bAlreadyInSet)
-			{
-				AddScriptBundleArc(*Package, ScriptDependencyGlobalImportIndex, BundleIndex);
-			}
+			AddUniqueScriptBundleArc(*Package, ScriptDependencyGlobalImportIndex, BundleIndex);
 		}
 
 		Package->ExportBundles[BundleIndex].Nodes.Add(Node);
