[AsyncLoading2] I/O chunk id┤s with global package id
- removed to usage of global package names when creating chunk id's
- removed code related to global names in cooker and linker
- added global package id to UPackage
- updated Bulk Data to use bulk data index and global package id when creating chunk id's

diff --git a/Engine/Source/Developer/IoStoreUtilities/Private/IoStoreUtilities.cpp b/Engine/Source/Developer/IoStoreUtilities/Private/IoStoreUtilities.cpp
--- a/Engine/Source/Developer/IoStoreUtilities/Private/IoStoreUtilities.cpp
+++ b/Engine/Source/Developer/IoStoreUtilities/Private/IoStoreUtilities.cpp
@@ -35,7 +35,6 @@ DEFINE_LOG_CATEGORY_STATIC(LogIoStore, Log, All);
 #define OUTPUT_IMPORTMAP_CSV 0
 #define SKIP_WRITE_CONTAINER 0
 #define SKIP_BULKDATA 0
-#define USE_NAMEMAPBUILDER 1
 
 struct FContainerTarget 
 {
@@ -44,87 +43,6 @@ struct FContainerTarget
 	FString OutputDirectory;
 };
 
-#if !USE_NAMEMAPBUILDER
-class FZenGlobalNameMap
-{
-public:
-	void Load(const FString& FilePath)
-	{
-		UE_LOG(LogIoStore, Display, TEXT("Loading global name map from '%s' for container files..."), *FilePath);
-
-		TUniquePtr<FArchive> Ar(IFileManager::Get().CreateFileReader(*FilePath));
-		check(Ar);
-
-		int32 NameCount;
-		*Ar << NameCount;
-
-		DisplayEntries.Reserve(NameCount);
-		ComparisonEntries.Reserve(NameCount);
-
-		FNameEntrySerialized SerializedNameEntry(ENAME_LinkerConstructor);
-
-		for (int32 I = 0; I < NameCount; ++I)
-		{
-			*Ar << SerializedNameEntry;
-			FName Name(SerializedNameEntry);
-
-			DisplayEntries.Emplace(Name.GetDisplayIndex());
-			DisplayEntryToIndex.Emplace(DisplayEntries[I], I);
-
-			ComparisonEntries.Emplace(Name.GetComparisonIndex());
-			ComparisonEntryToIndex.Emplace(ComparisonEntries[I], I);
-		}
-	}
-
-	void Save(const FString& FilePath)
-	{
-		UE_LOG(LogIoStore, Display, TEXT("Saving Container name map to '%s' with '%d' additional names"),
-			*FilePath,
-			ComparisonEntries.Num() - DisplayEntries.Num());
-
-		TUniquePtr<FArchive> Ar(IFileManager::Get().CreateFileWriter(*FilePath));
-		check(Ar);
-
-		int32 NameCount = ComparisonEntries.Num();
-		*Ar << NameCount;
-		
-		for (int32 I = 0; I < NameCount; ++I)
-		{
-			FName::GetEntry(ComparisonEntries[I])->Write(*Ar);
-		}
-	}
-
-	FName GetNameFromDisplayIndex(const uint32 DisplayIndex, const uint32 NameNumber) const
-	{
-		FNameEntryId DisplayEntry = DisplayEntries[DisplayIndex];
-		return FName::CreateFromDisplayId(DisplayEntry, NameNumber);
-	}
-
-	const int32* GetComparisonIndex(const FName& Name) const
-	{
-		return ComparisonEntryToIndex.Find(Name.GetComparisonIndex());
-	}
-
-	int32 GetOrCreateComparisonIndex(const FName& Name)
-	{
-		if (const int32* ExistingIndex = ComparisonEntryToIndex.Find(Name.GetComparisonIndex()))
-		{
-			return *ExistingIndex;
-		}
-		else
-		{
-			int32 NewIndex = ComparisonEntries.Add(Name.GetComparisonIndex());
-			ComparisonEntryToIndex.Add(ComparisonEntries[NewIndex], NewIndex);
-			return NewIndex;
-		}
-	}
-private:
-	TArray<FNameEntryId> DisplayEntries;
-	TArray<FNameEntryId> ComparisonEntries;
-	TMap<FNameEntryId, int32> DisplayEntryToIndex;
-	TMap<FNameEntryId, int32> ComparisonEntryToIndex;
-};
-#else
 class FNameMapBuilder
 {
 public:
@@ -226,7 +144,6 @@ private:
 	TArray<FNameEntryId> NameMap;
 	TMap<FNameEntryId, TTuple<int32,int32,int32>> DebugNameCounts; // <Number0Count,OtherNumberCount,MaxNumber>
 };
-#endif
 
 #if OUTPUT_CHUNKID_DIRECTORY
 class FChunkIdCsv
@@ -275,12 +192,11 @@ enum class EZenChunkType : uint8
 	BulkData
 };
 
-static FIoChunkId CreateZenChunkId(uint32 NameIndex, uint32 NameNumber, uint16 ChunkIndex, EZenChunkType ChunkType, const TCHAR* DebugString)
+static FIoChunkId CreateChunkId(int32 GlobalPackageId, uint16 ChunkIndex, EZenChunkType ChunkType, const TCHAR* DebugString)
 {
 	uint8 Data[12] = {0};
 
-	*reinterpret_cast<uint32*>(&Data[0]) = NameIndex;
-	*reinterpret_cast<int32*>(&Data[4]) = NameNumber;
+	*reinterpret_cast<int32*>(&Data[0]) = GlobalPackageId;
 	*reinterpret_cast<uint16*>(&Data[8]) = ChunkIndex;
 	*reinterpret_cast<uint8*>(&Data[10]) = static_cast<uint8>(ChunkType);
 
@@ -288,7 +204,7 @@ static FIoChunkId CreateZenChunkId(uint32 NameIndex, uint32 NameNumber, uint16 C
 	ChunkId.Set(Data, 12);
 	
 #if OUTPUT_CHUNKID_DIRECTORY
-	ChunkIdCsv.AddChunk(NameIndex, NameNumber, ChunkIndex, (uint8)ChunkType, GetTypeHash(ChunkId), DebugString);
+	ChunkIdCsv.AddChunk(GlobalPackageId, 0, ChunkIndex, (uint8)ChunkType, GetTypeHash(ChunkId), DebugString);
 #endif
 
 	return ChunkId;
@@ -1083,12 +999,10 @@ int32 CreateTarget(const FContainerTarget& Target)
 
 	FPackageStoreBulkDataManifest BulkDataManifest(CookedDir);
 	const bool bWithBulkDataManifest = BulkDataManifest.Load();
-	// if (!BulkDataManifest.Load())
-	// {
-	// 	UE_LOG(LogIoStore, Error, TEXT("Must cook with -packagestore on the command line (%s not found)"), *BulkDataManifest.GetFilename());
-	// 	RequestEngineExit(TEXT("ZenCreator failed"));
-	// 	return 1;
-	// }
+	if (bWithBulkDataManifest)
+	{
+		UE_LOG(LogIoStore, Display, TEXT("Loaded Bulk Data manifest '%s'"), *BulkDataManifest.GetFilename());
+	}
 
 #if OUTPUT_CHUNKID_DIRECTORY
 	ChunkIdCsv.CreateOutputFile(CookedDir);
@@ -1100,31 +1014,7 @@ int32 CreateTarget(const FContainerTarget& Target)
 	IFileManager::Get().FindFilesRecursive(FileNames, *CookedDir, TEXT("*.umap"), true, false, false);
 	UE_LOG(LogIoStore, Display, TEXT("Found '%d' files"), FileNames.Num());
 
-#if USE_NAMEMAPBUILDER
 	FNameMapBuilder NameMapBuilder;
-	auto SerializeName = [&](FArchive& A, const FName& N)
-	{
-		NameMapBuilder.SerializeName(A, N);
-	};
-#else
-	FString GlobalNameMapPath = CookedDir / TEXT("megafile.unamemap");
-	FZenGlobalNameMap GlobalNameMap;
-	GlobalNameMap.Load(GlobalNameMapPath);
-
-	auto SerializeName = [&](FArchive& A, const FName& N)
-	{
-		if (const int32* NameIndex = GlobalNameMap.GetComparisonIndex(N))
-		{
-			uint32 NameIndexNumber[] { uint32(*NameIndex), uint32(N.GetNumber()) };
-			A << NameIndexNumber[0] << NameIndexNumber[1];
-		}
-		else
-		{
-			UE_LOG(LogIoStore, Display, TEXT("FName '%s' has no valid name index in global name map"), *N.ToString());
-			check(false);
-		}
-	};
-#endif
 
 	uint64 NameSize = 0;
 	TArray<FObjectImport> Imports;
@@ -1223,17 +1113,14 @@ int32 CreateTarget(const FContainerTarget& Target)
 			{
 				*Ar << NameEntry;
 				FName Name(NameEntry);
-#if USE_NAMEMAPBUILDER
 				NameMapBuilder.MarkNameAsReferenced(Name);
 				Package.NameMap.Emplace(Name.GetDisplayIndex());
 				Package.NameIndices.Add(NameMapBuilder.MapName(Name));
-#endif
 			}
 
 			NameSize += Ar->Tell() - Summary.NameOffset;
 		}
 
-#if USE_NAMEMAPBUILDER
 		auto DeserializeName = [&](FArchive& A, FName& N)
 		{
 			int32 DisplayIndex, NameNumber;
@@ -1241,14 +1128,6 @@ int32 CreateTarget(const FContainerTarget& Target)
 			FNameEntryId DisplayEntry = Package.NameMap[DisplayIndex];
 			N = FName::CreateFromDisplayId(DisplayEntry, NameNumber);
 		};
-#else
-		auto DeserializeName = [&](FArchive& A, FName& N)
-		{
-			int32 NameIndex, NameNumber;
-			A << NameIndex << NameNumber;
-			N = GlobalNameMap.GetNameFromDisplayIndex(NameIndex, NameNumber);
-		};
-#endif
 
 		if (Summary.ImportCount > 0)
 		{
@@ -1417,7 +1296,7 @@ int32 CreateTarget(const FContainerTarget& Target)
 		for (const FImportData& ImportData : GlobalImports)
 		{
 			UniqueImportPackages += (ImportData.OuterIndex == 0);
-			SerializeName(*GlimportArchive, ImportData.ObjectName);
+			NameMapBuilder.SerializeName(*GlimportArchive, ImportData.ObjectName);
 			FPackageIndex Index;
 			Index = FPackageIndex::FromImport(ImportData.GlobalIndex);
 			*GlimportArchive << Index;
@@ -1605,7 +1484,7 @@ int32 CreateTarget(const FContainerTarget& Target)
 		for (auto& KV : Package.ExternalArcs)
 		{
 			const FName& ImportedPackageName = KV.Key;
-			SerializeName(ZenGraphArchive, ImportedPackageName);
+			NameMapBuilder.SerializeName(ZenGraphArchive, ImportedPackageName);
 
 			TArray<FArc>& Arcs = KV.Value;
 
@@ -1627,7 +1506,7 @@ int32 CreateTarget(const FContainerTarget& Target)
 			FExportData& ExportData = GlobalExports[Package.Exports[I]];
 
 			ZenExportMapArchive << ObjectExport.SerialSize;
-			SerializeName(ZenExportMapArchive, ObjectExport.ObjectName);
+			NameMapBuilder.SerializeName(ZenExportMapArchive, ObjectExport.ObjectName);
 			ZenExportMapArchive << ObjectExport.OuterIndex;
 			ZenExportMapArchive << ObjectExport.ClassIndex;
 			ZenExportMapArchive << ObjectExport.SuperIndex;
@@ -1659,21 +1538,15 @@ int32 CreateTarget(const FContainerTarget& Target)
 		Package.ExportBundlesSize = ZenExportBundlesArchive.Tell();
 
 		FName RelativeFileName(*Package.RelativeFileName);
-#if USE_NAMEMAPBUILDER
 		NameMapBuilder.MarkNameAsReferenced(Package.Name);
 		int32 PackageNameIndex = NameMapBuilder.MapName(Package.Name);
 		int32 PackageNameNumber = Package.Name.GetNumber();
 		NameMapBuilder.MarkNameAsReferenced(RelativeFileName);
-#else
-		int32 PackageNameIndex = GlobalNameMap.GetOrCreateComparisonIndex(Package.Name);
-		int32 PackageNameNumber = Package.Name.GetNumber();
-		GlobalNameMap.GetOrCreateComparisonIndex(RelativeFileName);
-#endif
 
 		ensure(Package.GlobalPackageId == PackageIndex++);
 		{
-			SerializeName(*StoreTocArchive, Package.Name);
-			SerializeName(*StoreTocArchive, RelativeFileName);
+			NameMapBuilder.SerializeName(*StoreTocArchive, Package.Name);
+			NameMapBuilder.SerializeName(*StoreTocArchive, RelativeFileName);
 			*StoreTocArchive << Package.ImportCount;
 			*StoreTocArchive << Package.SlimportCount;
 			*StoreTocArchive << Package.SlimportOffset;
@@ -1735,7 +1608,7 @@ int32 CreateTarget(const FContainerTarget& Target)
 			{
 				FIoBuffer IoBuffer(FIoBuffer::AssumeOwnership, ZenSummaryBuffer, ZenSummarySize);
 #if !SKIP_WRITE_CONTAINER
-				IoStoreWriter->Append(CreateZenChunkId(PackageNameIndex, PackageNameNumber, 0, EZenChunkType::PackageSummary, TEXT("PackageSummary")), IoBuffer);
+				IoStoreWriter->Append(CreateChunkId(Package.GlobalPackageId, 0, EZenChunkType::PackageSummary, TEXT("PackageSummary")), IoBuffer);
 #endif
 				++SummaryChunkCount;
 			}
@@ -1781,7 +1654,7 @@ int32 CreateTarget(const FContainerTarget& Target)
 					FObjectExport& ObjectExport = Exports[Package.ExportIndexOffset + I];
 					const int64 Offset = ObjectExport.SerialOffset - Package.UAssetSize;
 					FIoBuffer IoBuffer(FIoBuffer::Wrap, BundleBuffer, BundleBufferSize);
-					IoStoreWriter->Append(CreateZenChunkId(PackageNameIndex, PackageNameNumber, I, EZenChunkType::ExportBundleData, *Package.FileName), IoBuffer);
+					IoStoreWriter->Append(CreateChunkId(Package.GlobalPackageId, I, EZenChunkType::ExportBundleData, *Package.FileName), IoBuffer);
 #endif
 					++ExportChunkCount;
 					FMemory::Free(BundleBuffer);
@@ -1803,7 +1676,7 @@ int32 CreateTarget(const FContainerTarget& Target)
 				TUniquePtr<FArchive> BulkAr(IFileManager::Get().CreateFileReader(*UBulkFileName));
 				if (BulkAr != nullptr)
 				{
-					const FIoChunkId BulkDataChunkId = CreateZenChunkId(PackageNameIndex, PackageNameNumber, 0, EZenChunkType::BulkData, *Package.FileName);
+					const FIoChunkId BulkDataChunkId = CreateChunkId(Package.GlobalPackageId, 0, EZenChunkType::BulkData, *Package.FileName);
 #if !SKIP_WRITE_CONTAINER
 					uint8* BulkBuffer = static_cast<uint8*>(FMemory::Malloc(BulkAr->TotalSize()));
 					BulkAr->Serialize(BulkBuffer, BulkAr->TotalSize());
@@ -1821,13 +1694,15 @@ int32 CreateTarget(const FContainerTarget& Target)
 					++BulkChunkCount;
 
 					// Create additional mapping chunks as needed
+					FString BulkDataFileName = Package.FileName;
+					FPaths::MakePathRelativeTo(BulkDataFileName, *OutputDir);
 					const FPackageStoreBulkDataManifest::PackageDesc* PackageDesc = BulkDataManifest.Find(Package.FileName);
 					if (PackageDesc != nullptr)
 					{
 						for (const FPackageStoreBulkDataManifest::PackageDesc::BulkDataDesc& BulkDataDesc : PackageDesc->GetDataArray())
 						{
 #if !SKIP_WRITE_CONTAINER
-							const FIoChunkId AccessChunkId = CreateZenChunkId( PackageNameIndex, PackageNameNumber, BulkDataDesc.Index, EZenChunkType::BulkData, *Package.FileName);
+							const FIoChunkId AccessChunkId = CreateChunkId(Package.GlobalPackageId, BulkDataDesc.Index, EZenChunkType::BulkData, *Package.FileName);
 							const FIoStatus PartialResult = IoStoreWriter->MapPartialRange(BulkDataChunkId, BulkDataDesc.Offset, BulkDataDesc.Size, AccessChunkId);
 							if (!PartialResult.IsOk())
 							{
@@ -1847,11 +1722,7 @@ int32 CreateTarget(const FContainerTarget& Target)
 
 	UE_LOG(LogIoStore, Display, TEXT("Saving NameMap..."));
 	FString NameMapFilename = OutputDir / TEXT("Container.namemap");
-#if USE_NAMEMAPBUILDER
 	uint64 GlobalNameMapSize = NameMapBuilder.Save(NameMapFilename);
-#else
-	GlobalNameMap.Save(NameMapFilename);
-#endif
 
 	UE_LOG(LogIoStore, Display, TEXT("Saving StoreToc..."));
 	StoreTocArchive->Close();
@@ -1929,9 +1800,7 @@ int32 CreateTarget(const FContainerTarget& Target)
 		(double)PreloadDependenciesSize / 1024.0 / 1024.0, ImportPreloadCount, ExportPreloadCount, PreloadDependencies.Num());
 
 	UE_LOG(LogIoStore, Display, TEXT("IoStore: %8.2f MB GlobalToc"), (double)StoreTocSize / 1024.0 / 1024.0);
-#if USE_NAMEMAPBUILDER
 	UE_LOG(LogIoStore, Display, TEXT("IoStore: %8.2f MB GlobalNameMap, %d unique names"), (double)GlobalNameMapSize / 1024.0 / 1024.0, NameMapBuilder.Num());
-#endif
 	UE_LOG(LogIoStore, Display, TEXT("IoStore: %8.2f MB GlobalImportMap, %d unique imports, %d unique import packages"),
 		(double)GlimportSize / 1024.0 / 1024.0, GlobalImportsByFullName.Num(), UniqueImportPackages);
 	UE_LOG(LogIoStore, Display, TEXT("IoStore: %8.2f MB PackageSummary"), (double)ZenPackageSummarySize / 1024.0 / 1024.0);
@@ -1941,7 +1810,6 @@ int32 CreateTarget(const FContainerTarget& Target)
 	UE_LOG(LogIoStore, Display, TEXT("IoStore: %8.2f MB SerializedArcs, %d internal arcs, %d external arcs, %d circular packages (%d chains)"),
 		(double)UGraphSize / 1024.0 / 1024.0, TotalInternalArcCount, TotalExternalArcCount, CircularPackagesCount, CircularChains.Num());
 	UE_LOG(LogIoStore, Display, TEXT("IoStore: %8.2f MB ScriptArcs, %d arcs"), (double)ScriptArcsSize / 1024.0 / 1024.0, ScriptArcsCount);
-	UE_LOG(LogIoStore, Display, TEXT("-------------------- IoStore Summary: %s --------------------"), *Target.TargetPlatform->PlatformName());
 
 	return 0;
 }
@@ -1978,7 +1846,6 @@ int32 CreateIoStoreContainerFiles(const TCHAR* CmdLine)
 	{
 		UE_LOG(LogIoStore, Display, TEXT("Unknown target platform '%s'"), *TargetPlatformName);
 		return -1;
-
 	}
 
 	FString TargetCookedDirectory = FPaths::ProjectSavedDir() / TEXT("Cooked") / TargetPlatform->PlatformName();
@@ -1986,22 +1853,9 @@ int32 CreateIoStoreContainerFiles(const TCHAR* CmdLine)
 		? OutputDirectory
 		: TargetCookedDirectory / FApp::GetProjectName();
 
-	FString GlobalNameMapFilePath = TargetCookedDirectory / TEXT("megafile.unamemap");
-
-#if !USE_NAMEMAPBUILDER
-	if (!IFileManager::Get().FileExists(*GlobalNameMapFilePath))
-	{
-		UE_LOG(LogIoStore, Error,
-			TEXT("Expected to find global name map '%s' for target '%s'"),
-			*GlobalNameMapFilePath, *TargetPlatform->PlatformName());
-		return -1;
-	}
-#endif
-
 	FContainerTarget Target { TargetPlatform, TargetCookedDirectory, TargetOutputDirectory };
 
 	UE_LOG(LogIoStore, Display, TEXT("Creating target: '%s' using output directory: '%s'"), *Target.TargetPlatform->PlatformName(), *Target.OutputDirectory);
 
-	CreateTarget(Target);
-	return 0;
+	return CreateTarget(Target);
 }
diff --git a/Engine/Source/Editor/UnrealEd/Private/CookOnTheFlyServer.cpp b/Engine/Source/Editor/UnrealEd/Private/CookOnTheFlyServer.cpp
--- a/Engine/Source/Editor/UnrealEd/Private/CookOnTheFlyServer.cpp
+++ b/Engine/Source/Editor/UnrealEd/Private/CookOnTheFlyServer.cpp
@@ -6722,14 +6722,11 @@ void UCookOnTheFlyServer::InitializePackageStore(const TArray<FName>& TargetPlat
 	for (const FName PlatformName : TargetPlatformNames)
 	{
 		const FString ResolvedRootPath = RootPathSandbox.Replace(TEXT("[Platform]"), *PlatformName.ToString());
-		const FString NameMapCookedSandboxFilename = ResolvedRootPath / TEXT("megafile.unamemap");
 
 		// just leak all memory for now
-		FPackageStoreNameMapSaver* NameMapSaver			= new FPackageStoreNameMapSaver(*NameMapCookedSandboxFilename);
 		FPackageStoreBulkDataManifest* BulkDataManifest	= new FPackageStoreBulkDataManifest(ResolvedRootPath);
-		FPackageHeaderSaver* PackageHeaderSaver			= new FPackageHeaderSaver(*NameMapSaver);
 		FLooseFileWriter* LooseFileWriter				= new FLooseFileWriter();
-		FSavePackageContext* SavePackageContext			= new FSavePackageContext(*PackageHeaderSaver, *LooseFileWriter, *BulkDataManifest);
+		FSavePackageContext* SavePackageContext			= new FSavePackageContext(*LooseFileWriter, *BulkDataManifest);
 
 		SavePackageContexts.Add(SavePackageContext);
 	}
@@ -6748,9 +6745,6 @@ void UCookOnTheFlyServer::FinalizePackageStore()
 	const TArray<ITargetPlatform*>& TargetPlatforms = GetCookingTargetPlatforms();
 	for (int32 PlatformIndex = 0; PlatformIndex < TargetPlatforms.Num(); ++PlatformIndex)
 	{
-		INameMapSaver& NameMapSaver = SavePackageContexts[PlatformIndex]->HeaderSaver.NameMapSaver;
-		NameMapSaver.End();
-		
 		SavePackageContexts[PlatformIndex]->BulkDataManifest.Save();
 	}
 	UE_LOG(LogCook, Display, TEXT("Done saving PackageStoreHeaders(s)"));
diff --git a/Engine/Source/Runtime/CoreUObject/Private/Serialization/AsyncLoading2.cpp b/Engine/Source/Runtime/CoreUObject/Private/Serialization/AsyncLoading2.cpp
--- a/Engine/Source/Runtime/CoreUObject/Private/Serialization/AsyncLoading2.cpp
+++ b/Engine/Source/Runtime/CoreUObject/Private/Serialization/AsyncLoading2.cpp
@@ -205,21 +205,6 @@ public:
 		{
 			*Ar << SerializedNameEntry;
 			NameEntries.Emplace(FName(SerializedNameEntry).GetDisplayIndex());
-			EntryToIndex.Add(NameEntries[I], I);
-		}
-	}
-
-	void Save(const FString& FilePath)
-	{
-		TUniquePtr<FArchive> Ar(IFileManager::Get().CreateFileWriter(*FilePath));
-		check(Ar);
-
-		int32 NameCount = NameEntries.Num();
-		*Ar << NameCount;
-		
-		for (int32 I = 0; I < NameCount; ++I)
-		{
-			FName::GetEntry(NameEntries[I])->Write(*Ar);
 		}
 	}
 
@@ -236,25 +221,6 @@ public:
 		return FName::CreateFromDisplayId(NameEntry, SerializedName.GetNumber());
 	}
 
-	const int32* GetIndex(const FName& Name) const
-	{
-		return EntryToIndex.Find(Name.GetDisplayIndex());
-	}
-
-	int32 GetOrCreateIndex(const FName& Name)
-	{
-		if (const int32* ExistingIndex = EntryToIndex.Find(Name.GetDisplayIndex()))
-		{
-			return *ExistingIndex;
-		}
-		else
-		{
-			int32 NewIndex = NameEntries.Add(Name.GetDisplayIndex());
-			EntryToIndex.Add(NameEntries[NewIndex], NewIndex);
-			return NewIndex;
-		}
-	}
-
 	const TArray<FNameEntryId>& GetNameEntries() const
 	{
 		return NameEntries;
@@ -262,7 +228,6 @@ public:
 
 private:
 	TArray<FNameEntryId> NameEntries;
-	TMap<FNameEntryId, int32> EntryToIndex;
 };
 
 struct FGlobalImportStore
@@ -670,12 +635,11 @@ enum class EChunkType : uint8
 	BulkData,
 };
 
-FIoChunkId CreateChunkId(uint32 NameIndex, uint32 NameNumber, uint16 ChunkIndex, EChunkType ChunkType)
+FIoChunkId CreateChunkId(FGlobalPackageId GlobalPackageId, uint16 ChunkIndex, EChunkType ChunkType)
 {
 	uint8 Data[12] = {0};
 
-	*reinterpret_cast<uint32*>(&Data[0]) = NameIndex;
-	*reinterpret_cast<int32*>(&Data[4]) = NameNumber;
+	*reinterpret_cast<uint32*>(&Data[0]) = GlobalPackageId.Id;
 	*reinterpret_cast<uint16*>(&Data[8]) = ChunkIndex;
 	*reinterpret_cast<uint8*>(&Data[10]) = static_cast<uint8>(ChunkType);
 
@@ -695,29 +659,6 @@ FIoChunkId CreateChunkId(const FIoChunkId& ChunkId, uint16 ChunkIndex, EChunkTyp
 	return Out;
 }
 
-FName GetChunkName(const FIoChunkId& ChunkId, const FGlobalNameMap& GlobalNameMap)
-{
-	const uint8* Data = reinterpret_cast<const uint8*>(&ChunkId);
-	const uint32 NameIndex = *reinterpret_cast<const uint32*>(&Data[0]);
-	const int32 NameNumber = *reinterpret_cast<const int32*>(&Data[4]);
-
-	return GlobalNameMap.GetName(NameIndex, NameNumber);
-}
-
-EChunkType GetChunkType(const FIoChunkId& ChunkId)
-{
-	const uint8* Data = reinterpret_cast<const uint8*>(&ChunkId);
-	const uint8 ChunkType = *reinterpret_cast<const uint8*>(&Data[10]);
-	return static_cast<EChunkType>(ChunkType);
-}
-
-uint16 GetChunkIndex(const FIoChunkId& ChunkId)
-{
-	const uint8* Data = reinterpret_cast<const uint8*>(&ChunkId);
-	return *reinterpret_cast<const uint16*>(&Data[8]);
-
-}
-
 enum class EAsyncPackageLoadingState2 : uint8
 {
 	NewPackage,
@@ -1258,8 +1199,6 @@ private:
 	FAsyncLoadingThread2Impl& AsyncLoadingThread;
 	IEDLBootNotificationManager& EDLBootNotificationManager;
 	FAsyncLoadEventGraphAllocator& GraphAllocator;
-	/** Package chunk id																				*/
-	FIoChunkId PackageChunkId;
 	/** Global package id used to index into the package store											*/
 	FGlobalPackageId GlobalPackageId;
 	/** Packages that have been imported by this async package */
@@ -2471,7 +2410,7 @@ void FAsyncPackage2::StartLoading()
 	LoadStartTime = FPlatformTime::Seconds();;
 
 	FIoReadOptions ReadOptions;
-	AsyncLoadingThread.IoDispatcher.ReadWithCallback(CreateChunkId(PackageChunkId, 0, EChunkType::PackageSummary),
+	AsyncLoadingThread.IoDispatcher.ReadWithCallback(CreateChunkId(GlobalPackageId, 0, EChunkType::PackageSummary),
 		ReadOptions,
 		[this](TIoStatusOr<FIoBuffer> Result)
 		{
@@ -2548,7 +2487,7 @@ EAsyncPackageState::Type FAsyncPackage2::Event_StartExportBundleIo(FAsyncPackage
 	TRACE_CPUPROFILER_EVENT_SCOPE(Event_StartExportBundleIo);
 
 	FIoReadOptions ReadOptions;
-	Package->AsyncLoadingThread.IoDispatcher.ReadWithCallback(CreateChunkId(Package->PackageChunkId, ExportBundleIndex, EChunkType::ExportBundleData),
+	Package->AsyncLoadingThread.IoDispatcher.ReadWithCallback(CreateChunkId(Package->GlobalPackageId, ExportBundleIndex, EChunkType::ExportBundleData),
 		ReadOptions,
 		[Package, ExportBundleIndex](TIoStatusOr<FIoBuffer> Result)
 		{
@@ -3941,10 +3880,6 @@ FAsyncPackage2::FAsyncPackage2(const FAsyncPackageDesc& InDesc, int32 InSerialNu
 	TRACE_LOADTIME_NEW_ASYNC_PACKAGE(this, InDesc.Name);
 	AddRequestID(InDesc.RequestID);
 
-	const int32* NameIndex = AsyncLoadingThread.GlobalNameMap.GetIndex(Desc.NameToLoad);
-	check(NameIndex != nullptr);
-	PackageChunkId = CreateChunkId(*NameIndex, Desc.NameToLoad.GetNumber(), 0, EChunkType::None);
-
 	ExportBundleCount = AsyncLoadingThread.PackageStore.GetPackageExportBundleCount(GlobalPackageId);
 	ExportCount = AsyncLoadingThread.PackageStore.GetPackageExportCount(GlobalPackageId);
 	Exports.AddZeroed(ExportCount);
@@ -4135,6 +4070,7 @@ EAsyncPackageState::Type FAsyncPackage2::CreateLinker()
 			Package = NewObject<UPackage>(/*Outer*/nullptr, Desc.Name, RF_Public);
 			Package->SetPackageFlags(Desc.PackageFlags);
 			Package->FileName = Desc.NameToLoad;
+			Package->SetGlobalPackageId(GlobalPackageId.Id);
 			bCreatedLinkerRoot = true;
 		}
 	}
diff --git a/Engine/Source/Runtime/CoreUObject/Private/Serialization/BulkData.cpp b/Engine/Source/Runtime/CoreUObject/Private/Serialization/BulkData.cpp
--- a/Engine/Source/Runtime/CoreUObject/Private/Serialization/BulkData.cpp
+++ b/Engine/Source/Runtime/CoreUObject/Private/Serialization/BulkData.cpp
@@ -988,9 +988,9 @@ void FUntypedBulkData::Serialize( FArchive& Ar, UObject* Owner, int32 Idx, bool
 			
 			if ((BulkDataFlags & BULKDATA_CookedForIoDispatcher) != 0)
 			{
-				// Load the ChunkID that will be used by the IoDispatcher, but this version of BulkData will not use it!
-				FIoChunkId ChunkId;
- 				Ar << ChunkId;
+				// Load the bulk data index that will be used to create I/O chunk id's
+				uint16 BulkDataIndex = ~uint16(0);
+ 				Ar << BulkDataIndex;
 			}
 
 			// determine whether the payload is stored inline or at the end of the file
@@ -1100,8 +1100,8 @@ void FUntypedBulkData::Serialize( FArchive& Ar, UObject* Owner, int32 Idx, bool
 
 						if ((BulkDataFlags & BULKDATA_CookedForIoDispatcher) != 0)
 						{
-							FIoChunkId DummyChunkID;
-							Ar << DummyChunkID;
+							uint16 BulkDataIndex = ~uint16(0);
+							Ar << BulkDataIndex;
 						}
 					}
 					else
@@ -1126,8 +1126,8 @@ void FUntypedBulkData::Serialize( FArchive& Ar, UObject* Owner, int32 Idx, bool
 
 						if ((BulkDataFlags & BULKDATA_CookedForIoDispatcher) != 0)
 						{
-							FIoChunkId DummyChunkID;
-							Ar << DummyChunkID;
+							uint16 BulkDataIndex = ~uint16(0);
+							Ar << BulkDataIndex;
 						}
 
 						Filename = FPaths::ChangeExtension(Filename, TEXT(".ubulk"));
@@ -1277,8 +1277,8 @@ void FUntypedBulkData::Serialize( FArchive& Ar, UObject* Owner, int32 Idx, bool
 
 				SavedBulkDataChunkIdPos = Ar.Tell();
 
-				FIoChunkId InvalidChunkId = FIoChunkId::InvalidChunkId;
-				Ar << InvalidChunkId;
+				uint16 BulkDataIndex = ~uint16(0);
+				Ar << BulkDataIndex;
 
 				BulkDataFlags |= BULKDATA_CookedForIoDispatcher;
 			}
@@ -1342,8 +1342,8 @@ void FUntypedBulkData::Serialize( FArchive& Ar, UObject* Owner, int32 Idx, bool
 						{
 							SavedDupeBulkDataChunkIdPos = Ar.Tell();
 							
-							FIoChunkId InvalidChunkId = FIoChunkId::InvalidChunkId;
-							Ar << InvalidChunkId;
+							uint16 BulkDataIndex = ~uint16(0);
+							Ar << BulkDataIndex;
 
 							check((SavedDupeBulkDataFlags& BULKDATA_CookedForIoDispatcher) != 0); // Validate that the BULKDATA_CookedForIoDispatcher flag was already set
 						}
diff --git a/Engine/Source/Runtime/CoreUObject/Private/Serialization/BulkData2.cpp b/Engine/Source/Runtime/CoreUObject/Private/Serialization/BulkData2.cpp
--- a/Engine/Source/Runtime/CoreUObject/Private/Serialization/BulkData2.cpp
+++ b/Engine/Source/Runtime/CoreUObject/Private/Serialization/BulkData2.cpp
@@ -25,6 +25,33 @@ namespace
 	static constexpr const TCHAR* DefaultExt = TEXT(".ubulk");			// Stored in a separate file
 	static constexpr const TCHAR* MemoryMappedExt = TEXT(".m.ubulk");	// Stored in a separate file aligned for memory mapping
 	static constexpr const TCHAR* OptionalExt = TEXT(".uptnl");			// Stored in a separate file that is optional
+
+	const uint16 InvalidBulkDataIndex = ~uint16(0);
+
+	// TODO: Move to common header
+	enum class EZenChunkType : uint8
+	{
+		None,
+		PackageSummary,
+		ExportData,
+		ExportBundleData,
+		BulkData
+	};
+
+	// TODO: Move to common header
+	FIoChunkId CreateChunkId(int32 GlobalPackageId, uint16 ChunkIndex, EZenChunkType ChunkType)
+	{
+		uint8 Data[12] = {0};
+
+		*reinterpret_cast<int32*>(&Data[0]) = GlobalPackageId;
+		*reinterpret_cast<uint16*>(&Data[8]) = ChunkIndex;
+		*reinterpret_cast<uint8*>(&Data[10]) = static_cast<uint8>(ChunkType);
+
+		FIoChunkId ChunkId;
+		ChunkId.Set(Data, 12);
+		
+		return ChunkId;
+	}
 }
 
 // TODO: The code in the FileTokenSystem namespace is a temporary system so that FBulkDataBase can hold
@@ -566,7 +593,18 @@ void FBulkDataBase::Serialize(FArchive& Ar, UObject* Owner, int32 /*Index*/, boo
 
 		if ((BulkDataFlags & BULKDATA_CookedForIoDispatcher) != 0)
 		{
-			Ar << ChunkID;
+			uint16 BulkDataIndex;
+			Ar << BulkDataIndex;
+
+			if (BulkDataIndex != InvalidBulkDataIndex)
+			{
+				UPackage* Package = Owner->GetOutermost();
+				ChunkID = CreateChunkId(Package->GetGlobalPackageId(), BulkDataIndex, EZenChunkType::BulkData);
+			}
+			else
+			{
+				ChunkID = FIoChunkId::InvalidChunkId;
+			}
 		}
 
 		if (!bUseZenLoader || ChunkID == FIoChunkId::InvalidChunkId)
@@ -1057,8 +1095,8 @@ void FBulkDataBase::SerializeDuplicateData(FArchive& Ar, FLinkerLoad* Linker, ui
 
 	if ((OutBulkDataFlags & BULKDATA_CookedForIoDispatcher) != 0)
 	{
-		FIoChunkId DummyChunkID;
-		Ar << DummyChunkID;
+		uint16 DummyBulkDataIndex = InvalidBulkDataIndex;
+		Ar << DummyBulkDataIndex;
 	}
 
 	// Fix up the file offset if we have a valid linker
diff --git a/Engine/Source/Runtime/CoreUObject/Private/UObject/LinkerSave.cpp b/Engine/Source/Runtime/CoreUObject/Private/UObject/LinkerSave.cpp
--- a/Engine/Source/Runtime/CoreUObject/Private/UObject/LinkerSave.cpp
+++ b/Engine/Source/Runtime/CoreUObject/Private/UObject/LinkerSave.cpp
@@ -6,7 +6,6 @@
 #include "Serialization/LargeMemoryWriter.h"
 #include "Serialization/UnversionedPropertySerialization.h"
 #include "UObject/Package.h"
-#include "UObject/SavePackage.h"
 #include "UObject/Class.h"
 #include "Templates/Casts.h"
 #include "UObject/LazyObjectPtr.h"
@@ -23,10 +22,9 @@
 /** A mapping of package name to generated script SHA keys */
 TMap<FString, TArray<uint8> > FLinkerSave::PackagesToScriptSHAMap;
 
-FLinkerSave::FLinkerSave(FPackageHeaderSaver& InPackageHeaderSaver, UPackage* InParent, const TCHAR* InFilename, bool bForceByteSwapping, bool bInSaveUnversioned)
+FLinkerSave::FLinkerSave(UPackage* InParent, const TCHAR* InFilename, bool bForceByteSwapping, bool bInSaveUnversioned)
 :	FLinker(ELinkerType::Save, InParent, InFilename)
 ,	Saver(nullptr)
-,	HeaderSaver(InPackageHeaderSaver)
 {
 	if (FPlatformProperties::HasEditorOnlyData())
 	{
@@ -79,10 +77,9 @@ FLinkerSave::FLinkerSave(FPackageHeaderSaver& InPackageHeaderSaver, UPackage* In
 }
 
 
-FLinkerSave::FLinkerSave(FPackageHeaderSaver& InPackageHeaderSaver, UPackage* InParent, FArchive *InSaver, bool bForceByteSwapping, bool bInSaveUnversioned)
+FLinkerSave::FLinkerSave(UPackage* InParent, FArchive *InSaver, bool bForceByteSwapping, bool bInSaveUnversioned)
 : FLinker(ELinkerType::Save, InParent, TEXT("$$Memory$$"))
 , Saver(nullptr)
-, HeaderSaver(InPackageHeaderSaver)
 {
 	if (FPlatformProperties::HasEditorOnlyData())
 	{
@@ -135,10 +132,9 @@ FLinkerSave::FLinkerSave(FPackageHeaderSaver& InPackageHeaderSaver, UPackage* In
 	}
 }
 
-FLinkerSave::FLinkerSave(FPackageHeaderSaver& InPackageHeaderSaver, UPackage* InParent, bool bForceByteSwapping, bool bInSaveUnversioned )
+FLinkerSave::FLinkerSave(UPackage* InParent, bool bForceByteSwapping, bool bInSaveUnversioned )
 :	FLinker(ELinkerType::Save, InParent, TEXT("$$Memory$$"))
 ,	Saver(nullptr)
-,	HeaderSaver(InPackageHeaderSaver)
 {
 	if (FPlatformProperties::HasEditorOnlyData())
 	{
@@ -206,6 +202,18 @@ FLinkerSave::~FLinkerSave()
 	CloseAndDestroySaver();
 }
 
+int32 FLinkerSave::MapName(FNameEntryId Id) const
+{
+	const int32* IndexPtr = NameIndices.Find(Id);
+
+	if (IndexPtr)
+	{
+		return *IndexPtr;
+	}
+
+	return INDEX_NONE;
+}
+
 FPackageIndex FLinkerSave::MapObject( const UObject* Object ) const
 {
 	if (Object)
@@ -280,7 +288,7 @@ FString FLinkerSave::GetArchiveName() const
 
 FArchive& FLinkerSave::operator<<( FName& InName )
 {
-	int32 Save = HeaderSaver.NameMapSaver.MapName(InName);
+	int32 Save = MapName(InName.GetDisplayIndex());
 
 	check(GetSerializeContext());
 	ensureMsgf(Save != INDEX_NONE, TEXT("Name \"%s\" is not mapped when saving %s (object: %s, property: %s)"), 
diff --git a/Engine/Source/Runtime/CoreUObject/Private/UObject/SavePackage.cpp b/Engine/Source/Runtime/CoreUObject/Private/UObject/SavePackage.cpp
--- a/Engine/Source/Runtime/CoreUObject/Private/UObject/SavePackage.cpp
+++ b/Engine/Source/Runtime/CoreUObject/Private/UObject/SavePackage.cpp
@@ -437,135 +437,37 @@ static void AsyncWriteFileWithSplitExports(FLargeMemoryPtr Data, const int64 Dat
 	(new FAutoDeleteAsyncTask<FAsyncWriteWorkerWithSplitExports>(Filename, MoveTemp(Data), DataSize, HeaderSize))->StartBackgroundTask();
 }
 
-void FSinglePackageNameMapSaver::MarkNameAsReferenced(FNameEntryId Name)
+class FPackageNameMapSaver
 {
-	ReferencedNames.Add(Name);
-}
-
-int32 FSinglePackageNameMapSaver::MapName(FName Name) const
-{
-	const FNameEntryId Id = Name.GetDisplayIndex();
-	const int32* Index = NameIndices.Find(Id);
-	return Index ? *Index : INDEX_NONE;
-}
-
-bool FSinglePackageNameMapSaver::NameExistsInCurrentPackage(FNameEntryId ComparisonId) const
-{
-	for (FNameEntryId DisplayId : ReferencedNames)
+public:
+	void MarkNameAsReferenced(FName Name)
 	{
-		if (FName::GetComparisonIdFromDisplayId(DisplayId) == ComparisonId)
-		{
-			return true;
-		}
+		ReferencedNames.Add(Name.GetDisplayIndex());
 	}
-	return false;
-}
 
-void FSinglePackageNameMapSaver::MarkNameAsReferenced(FName Name)
-{
-	ReferencedNames.Add(Name.GetDisplayIndex());
-}
-
-
-FPackageStoreNameMapSaver::FPackageStoreNameMapSaver(const TCHAR* InFilename)
-: Filename(InFilename) {}
-
-void FPackageStoreNameMapSaver::BeginPackage()
-{
-	PackageReferencedNames.Reset();
-}
-
-void FPackageStoreNameMapSaver::MarkNameAsReferenced(FName Name)
-{
-	const FNameEntryId Id = Name.GetDisplayIndex();
-	int32& Index = NameIndices.FindOrAdd(Id);
-	if (Index == 0)
+	void MarkNameAsReferenced(FNameEntryId Name)
 	{
-		Index = NameIndices.Num();
-		NameMap.Add(Id);
+		ReferencedNames.Add(Name);
 	}
-	PackageReferencedNames.Add(Id);
-}
 
-int32 FPackageStoreNameMapSaver::MapName(FName Name) const
-{
-	const FNameEntryId Id = Name.GetDisplayIndex();
-	// debug counts
+	bool NameExists(FNameEntryId ComparisonId) const
 	{
-		const int32 Number = Name.GetNumber();
-		TTuple<int32,int32,int32>& Counts = DebugNameCounts.FindOrAdd(Id);
-
-		if (Number == 0)
-		{
-			++Counts.Get<0>();
-		}
-		else
+		for (FNameEntryId DisplayId : ReferencedNames)
 		{
-			++Counts.Get<1>();
-			if (Number > Counts.Get<2>())
+			if (FName::GetComparisonIdFromDisplayId(DisplayId) == ComparisonId)
 			{
-				Counts.Get<2>() = Number;
+				return true;
 			}
 		}
 
+		return false;
 	}
 
-	const int32* Index = NameIndices.Find(Id);
-	return Index ? *Index - 1 : INDEX_NONE;
-}
-
-bool FPackageStoreNameMapSaver::NameExistsInCurrentPackage(FNameEntryId ComparisonId) const
-{
-	for (FNameEntryId DisplayId : PackageReferencedNames)
-	{
-		if (FName::GetComparisonIdFromDisplayId(DisplayId) == ComparisonId)
-		{
-			return true;
-		}
-	}
-	return false;
-}
+	void UpdateLinker(FLinkerSave& Linker, FLinkerLoad* Conform, FArchive* BinarySaver);
 
-void FPackageStoreNameMapSaver::End()
-{
-	// real megafile.unamemap
-	{
-		TUniquePtr<FArchive> BinArchive(IFileManager::Get().CreateFileWriter(*Filename));
-		int32 NameCount = NameMap.Num();
-		*BinArchive << NameCount;
-		for (int32 I = 0; I < NameCount; ++I)
-		{
-			FName::GetEntry(NameMap[I])->Write(*BinArchive);
-		}
-	}
-	// debug megafile.unamemap.csv
-	{
-		FString CsvFilePath = Filename;
-		CsvFilePath.Append(TEXT(".csv"));
-		TUniquePtr<FArchive> CsvArchive(IFileManager::Get().CreateFileWriter(*CsvFilePath));
-		if (CsvArchive)
-		{
-			TCHAR Name[FName::StringBufferSize];
-			ANSICHAR Line[MAX_SPRINTF + FName::StringBufferSize];
-			ANSICHAR Header[] = "Length\tMaxNumber\tNumberCount\tBaseCount\tTotalCount\tFName\n";
-			CsvArchive->Serialize(Header, sizeof(Header) - 1);
-			for (auto& Counts : DebugNameCounts)
-			{
-				const int32 NameLen = FName::CreateFromDisplayId(Counts.Key, 0).ToString(Name);
-				FCStringAnsi::Sprintf(Line, "%d\t%d\t%d\t%d\t%d\t",
-					NameLen, Counts.Value.Get<2>(), Counts.Value.Get<1>(), Counts.Value.Get<0>(), Counts.Value.Get<0>() + Counts.Value.Get<1>());
-				ANSICHAR* L = Line + FCStringAnsi::Strlen(Line);
-				const TCHAR* N = Name;
-				while (*N)
-				{
-					*L++ = CharCast<ANSICHAR,TCHAR>(*N++);
-				}
-				*L++ = '\n';
-				CsvArchive.Get()->Serialize(Line, L - Line);
-			}
-		}
-	}
-}
+private:
+	TSet<FNameEntryId> ReferencedNames;
+};
 
 void FPackageStoreBulkDataManifest::PackageDesc::AddData(uint16 InIndex, uint64 InOffset, uint64 InSize)
 {
@@ -580,22 +482,6 @@ void FPackageStoreBulkDataManifest::PackageDesc::AddData(uint16 InIndex, uint64
 	Entry.Size = InSize;
 }
 
-FIoChunkId IBulkDataManifest::CreateChunkId(uint32 NameIndex, uint32 NameNumber, uint16 ChunkIndex)
-{
-	uint8 Data[12] = { 0 };
-
-	*reinterpret_cast<uint32*>(&Data[0]) = NameIndex;
-	*reinterpret_cast<int32*>(&Data[4]) = NameNumber;
-	*reinterpret_cast<uint16*>(&Data[8]) = ChunkIndex;
-	*reinterpret_cast<uint8*>(&Data[10]) = static_cast<uint8>(3); 	// EZenChunkType::BulkData in ZenCreator
-																	// TODO: We need to expose this somewhere!
-
-	FIoChunkId ChunkId;
-	ChunkId.Set(Data, 12);
-
-	return ChunkId;
-}
-
 FPackageStoreBulkDataManifest::FPackageStoreBulkDataManifest(const FString& InRootPath)
 {
 	RootPath = InRootPath;
@@ -625,6 +511,10 @@ FArchive& operator<<(FArchive& Ar, FPackageStoreBulkDataManifest::PackageDesc& E
 	return Ar;
 }
 
+FPackageStoreBulkDataManifest::~FPackageStoreBulkDataManifest()
+{
+}
+
 bool FPackageStoreBulkDataManifest::Load()
 {
 	Data.Empty();
@@ -647,7 +537,6 @@ void FPackageStoreBulkDataManifest::Save()
 	*BinArchive << Data;
 }
 
-
 const FPackageStoreBulkDataManifest::PackageDesc* FPackageStoreBulkDataManifest::Find(const FString& PackageFilename) const
 {
 	const  FString NormalizedFilename = FixFilename(PackageFilename);
@@ -678,17 +567,13 @@ FString FPackageStoreBulkDataManifest::FixFilename(const FString& InFilename) co
 {
 	FString OutFilename = InFilename;
 	FPaths::NormalizeFilename(OutFilename);
-
-	if (!FPaths::IsRelative(OutFilename))
-	{
-		FPaths::MakePathRelativeTo(OutFilename, *RootPath);
-	}
+	FPaths::MakePathRelativeTo(OutFilename, *RootPath);
 
 	return OutFilename;
 }
 
 #if WITH_EDITOR
-static void AddReplacementsNames(INameMapSaver& NameMapSaver, UObject* Obj, const ITargetPlatform* TargetPlatform)
+static void AddReplacementsNames(FPackageNameMapSaver& NameMapSaver, UObject* Obj, const ITargetPlatform* TargetPlatform)
 {
 	if (TargetPlatform)
 	{
@@ -1137,7 +1022,7 @@ class FArchiveSaveTagImports : public FArchiveUObject
 {
 public:
 	FLinkerSave* Linker;
-	INameMapSaver& NameMapSaver;
+	FPackageNameMapSaver& NameMapSaver;
 	TArray<UObject*> Dependencies;
 	TArray<UObject*> NativeDependencies;
 	TArray<UObject*> OtherImports;
@@ -1162,7 +1047,7 @@ public:
 		}
 	};
 
-	FArchiveSaveTagImports(FLinkerSave* InLinker, INameMapSaver& InNameMapSaver)
+	FArchiveSaveTagImports(FLinkerSave* InLinker, FPackageNameMapSaver& InNameMapSaver)
 		: Linker(InLinker)
 		, NameMapSaver(InNameMapSaver)
 		, bIgnoreDependencies(false)
@@ -1478,7 +1363,7 @@ public:
 	 * @param	Linker				linker containing the names that need to be sorted
 	 * @param	LinkerToConformTo	optional linker to conform against.
 	 */
-	void SortNames( FLinkerSave* Linker, FLinkerLoad* LinkerToConformTo, FSinglePackageNameMapSaver& NameMapSaver)
+	void SortNames( FLinkerSave* Linker, FLinkerLoad* LinkerToConformTo, FPackageNameMapSaver& NameMapSaver)
 	{
 		int32 SortStartPosition = 0;
 
@@ -1510,7 +1395,7 @@ public:
 	}
 };
 
-void FSinglePackageNameMapSaver::EndPackage(FLinkerSave& Linker, FLinkerLoad* Conform, FArchive* BinarySaver)
+void FPackageNameMapSaver::UpdateLinker(FLinkerSave& Linker, FLinkerLoad* Conform, FArchive* BinarySaver)
 {
 	// Add names
 	Linker.NameMap.Reserve(Linker.NameMap.Num() + ReferencedNames.Num());
@@ -1530,7 +1415,7 @@ void FSinglePackageNameMapSaver::EndPackage(FLinkerSave& Linker, FLinkerLoad* Co
 		for (int32 i = 0; i < Linker.NameMap.Num(); i++)
 		{
 			FName::GetEntry(Linker.NameMap[i])->Write(Linker);
-			NameIndices.Add(Linker.NameMap[i], i);
+			Linker.NameIndices.Add(Linker.NameMap[i], i);
 		}
 	}
 }
@@ -3774,12 +3659,7 @@ FSavePackageResultStruct UPackage::Save(UPackage* InOuter, UObject* Base, EObjec
 		int32 PackageSize = INDEX_NONE;
 		{
 			// TODO: Require a SavePackageContext and move to EditorEngine
-			FSinglePackageNameMapSaver SinglePackageNameMapSaver;
-			FPackageHeaderSaver LocalSaver(SinglePackageNameMapSaver);
-			FPackageHeaderSaver& PackageHeaderSaver = SavePackageContext ? SavePackageContext->HeaderSaver : LocalSaver;
-
-			INameMapSaver& NameMapSaver = PackageHeaderSaver.NameMapSaver;
-			NameMapSaver.BeginPackage();
+			FPackageNameMapSaver NameMapSaver;
 
 			uint32 ComparisonFlags = PPF_DeepCompareInstances;
 
@@ -3933,13 +3813,13 @@ FSavePackageResultStruct UPackage::Save(UPackage* InOuter, UObject* Base, EObjec
 					// The entire package will be serialized to memory and then compared against package on disk.
 					// Each difference will be log with its Serialize call stack trace
 					FArchive* Saver = new FArchiveStackTrace(FindAssetInPackage(InOuter), *InOuter->FileName.ToString(), true, InOutDiffMap);
-					Linker = TUniquePtr<FLinkerSave>(new FLinkerSave(PackageHeaderSaver, InOuter, Saver, bForceByteSwapping, bSaveUnversioned));
+					Linker = TUniquePtr<FLinkerSave>(new FLinkerSave(InOuter, Saver, bForceByteSwapping, bSaveUnversioned));
 				}
 				else if (TargetPlatform != nullptr && (SaveFlags & SAVE_DiffOnly))
 				{
 					// The entire package will be serialized to memory and then compared against package on disk
 					FArchive* Saver = new FArchiveStackTrace(FindAssetInPackage(InOuter), *InOuter->FileName.ToString(), false);
-					Linker = TUniquePtr<FLinkerSave>(new FLinkerSave(PackageHeaderSaver, InOuter, Saver, bForceByteSwapping, bSaveUnversioned));
+					Linker = TUniquePtr<FLinkerSave>(new FLinkerSave(InOuter, Saver, bForceByteSwapping, bSaveUnversioned));
 				}
 				else if ((!!TargetPlatform) && FParse::Value(FCommandLine::Get(), TEXT("DiffCookedPackages="), DiffCookedPackagesPath))
 				{
@@ -3952,19 +3832,19 @@ FSavePackageResultStruct UPackage::Save(UPackage* InOuter, UObject* Base, EObjec
 					
 					FArchive* TestArchive = IFileManager::Get().CreateFileReader(*TestArchiveFilename); 
 					FArchive* Saver = new FDiffSerializeArchive(*InOuter->FileName.ToString(), TestArchive);
-					Linker = TUniquePtr<FLinkerSave>(new FLinkerSave(PackageHeaderSaver, InOuter, Saver, bForceByteSwapping));
+					Linker = TUniquePtr<FLinkerSave>(new FLinkerSave(InOuter, Saver, bForceByteSwapping));
 				}
 				else 
 #endif
 				if (bSaveAsync)
 				{
 					// Allocate the linker with a memory writer, forcing byte swapping if wanted.
-					Linker = TUniquePtr<FLinkerSave>(new FLinkerSave(PackageHeaderSaver, InOuter, bForceByteSwapping, bSaveUnversioned));
+					Linker = TUniquePtr<FLinkerSave>(new FLinkerSave(InOuter, bForceByteSwapping, bSaveUnversioned));
 				}
 				else
 				{
 					// Allocate the linker, forcing byte swapping if wanted.
-					Linker = TUniquePtr<FLinkerSave>(new FLinkerSave(PackageHeaderSaver, InOuter, *TempFilename, bForceByteSwapping, bSaveUnversioned));
+					Linker = TUniquePtr<FLinkerSave>(new FLinkerSave(InOuter, *TempFilename, bForceByteSwapping, bSaveUnversioned));
 				}
 
 #if WITH_TEXT_ARCHIVE_SUPPORT
@@ -4537,7 +4417,7 @@ FSavePackageResultStruct UPackage::Save(UPackage* InOuter, UObject* Base, EObjec
 				{
 					// check the name list for uniqueobjectnamefor cooking
 					static FNameEntryId NAME_UniqueObjectNameForCooking = FName("UniqueObjectNameForCooking").GetComparisonIndex();
-					if (NameMapSaver.NameExistsInCurrentPackage(NAME_UniqueObjectNameForCooking))
+					if (NameMapSaver.NameExists(NAME_UniqueObjectNameForCooking))
 					{
 						UE_LOG(LogSavePackage, Warning, TEXT("Saving object into cooked package %s which was created at cook time"), Filename);
 					}
@@ -4551,7 +4431,7 @@ FSavePackageResultStruct UPackage::Save(UPackage* InOuter, UObject* Base, EObjec
 					FArchive::FScopeSetDebugSerializationFlags S(*Linker, DSF_IgnoreDiff, true);
 					FArchiveStackTraceIgnoreScope IgnoreSummaryDiffsScope(DiffSettings.bIgnoreHeaderDiffs);
 #endif
-					NameMapSaver.EndPackage(*Linker.Get(), Conform, bTextFormat ? nullptr : Linker->Saver);
+					NameMapSaver.UpdateLinker(*Linker.Get(), Conform, bTextFormat ? nullptr : Linker->Saver);
 				}
 				if ( EndSavingIfCancelled( Linker.Get(), TempFilename ) )
 				{ 
@@ -6411,22 +6291,19 @@ void SaveBulkData(FLinkerSave* Linker, const UPackage* InOuter, const TCHAR* Fil
 
 			if ((BulkDataStorageInfo.BulkDataFlags & BULKDATA_CookedForIoDispatcher) != 0 && SavePackageContext != nullptr)
 			{
-				const int32 PackageNameIndex = SavePackageContext->HeaderSaver.NameMapSaver.MapName(InOuter->GetFName());
-				const int32 PackageNameNumber = InOuter->GetFName().GetNumber();
 				check(BulkDataIndex < TNumericLimits<uint16>::Max());
-				const uint16 ChunkIndex = BulkDataIndex++;
-
-				FIoChunkId ChunkId = IBulkDataManifest::CreateChunkId(PackageNameIndex, PackageNameNumber, ChunkIndex);
 
 				SavePackageContext->BulkDataManifest.AddFileAccess
 				(
-					Filename, ChunkIndex, BulkStartOffset, SizeOnDisk
+					Filename, BulkDataIndex, BulkStartOffset, SizeOnDisk
 				);
 
 				check(BulkDataStorageInfo.BulkDataChunkIdPos != INDEX_NONE);
 
 				Linker->Seek(BulkDataStorageInfo.BulkDataChunkIdPos);
-				*Linker << ChunkId;
+				*Linker << BulkDataIndex;
+				
+				BulkDataIndex++;
 			}
 
 			Linker->Seek(LinkerEndOffset);
diff --git a/Engine/Source/Runtime/CoreUObject/Public/UObject/LinkerSave.h b/Engine/Source/Runtime/CoreUObject/Public/UObject/LinkerSave.h
--- a/Engine/Source/Runtime/CoreUObject/Public/UObject/LinkerSave.h
+++ b/Engine/Source/Runtime/CoreUObject/Public/UObject/LinkerSave.h
@@ -10,7 +10,6 @@
 #include "Templates/RefCounting.h"
 
 struct FUntypedBulkData;
-class FPackageHeaderSaver;
 
 /*----------------------------------------------------------------------------
 	FLinkerSave.
@@ -43,7 +42,8 @@ public:
 	/** List of Searchable Names, by object containing them. This gets turned into package indices later */
 	TMap<const UObject *, TArray<FName> > SearchableNamesObjectMap;
 
-	FPackageHeaderSaver& HeaderSaver;
+	/** Index array - location of the name in the NameMap array for each FName is stored in the NameIndices array using the FName's Index */
+	TMap<FNameEntryId, int32> NameIndices;
 
 	/** Save context associated with this linker */
 	TRefCountPtr<FUObjectSerializeContext> SaveContext;
@@ -70,11 +70,14 @@ public:
 	COREUOBJECT_API static TMap<FString, TArray<uint8> > PackagesToScriptSHAMap;
 
 	/** Constructor for file writer */
-	FLinkerSave(FPackageHeaderSaver& InHeaderSaver, UPackage* InParent, const TCHAR* InFilename, bool bForceByteSwapping, bool bInSaveUnversioned = false );
+	FLinkerSave(UPackage* InParent, const TCHAR* InFilename, bool bForceByteSwapping, bool bInSaveUnversioned = false );
 	/** Constructor for memory writer */
-	FLinkerSave(FPackageHeaderSaver& InHeaderSaver, UPackage* InParent, bool bForceByteSwapping, bool bInSaveUnversioned = false );
+	FLinkerSave(UPackage* InParent, bool bForceByteSwapping, bool bInSaveUnversioned = false );
 	/** Constructor for custom savers. The linker assumes ownership of the custom saver. */
-	FLinkerSave(FPackageHeaderSaver& InHeaderSaver, UPackage* InParent, FArchive *InSaver, bool bForceByteSwapping, bool bInSaveUnversioned = false);
+	FLinkerSave(UPackage* InParent, FArchive *InSaver, bool bForceByteSwapping, bool bInSaveUnversioned = false);
+
+	/** Returns the appropriate name index for the source name, or 0 if not found in NameIndices */
+	int32 MapName( FNameEntryId Name) const;
 
 	/** Returns the appropriate package index for the source object, or default value if not found in ObjectIndicesMap */
 	FPackageIndex MapObject(const UObject* Object) const;
diff --git a/Engine/Source/Runtime/CoreUObject/Public/UObject/Package.h b/Engine/Source/Runtime/CoreUObject/Public/UObject/Package.h
--- a/Engine/Source/Runtime/CoreUObject/Public/UObject/Package.h
+++ b/Engine/Source/Runtime/CoreUObject/Public/UObject/Package.h
@@ -186,7 +186,9 @@ public:
 private:
 	/** Package Flags */
 	uint32	PackageFlagsPrivate;
-
+	
+	/** Globally unique id used to address I/O chunks within the package */
+	uint32 GlobalPackageId = ~uint32(0); 
 public:
 
 	/** Editor only: PIE instance ID this package belongs to, INDEX_NONE otherwise */
@@ -512,6 +514,18 @@ public:
 		ChunkIDs = InChunkIDs;
 	}
 
+	/** returns the globally unique package id */
+	FORCEINLINE uint32 GetGlobalPackageId() const
+	{
+		return GlobalPackageId;
+	}
+
+	/** sets the globally unique package id */
+	FORCEINLINE void SetGlobalPackageId(uint32 InGlobalPackageId)
+	{
+		GlobalPackageId = InGlobalPackageId;
+	}
+
 	////////////////////////////////////////////////////////
 	// MetaData 
 
diff --git a/Engine/Source/Runtime/CoreUObject/Public/UObject/SavePackage.h b/Engine/Source/Runtime/CoreUObject/Public/UObject/SavePackage.h
--- a/Engine/Source/Runtime/CoreUObject/Public/UObject/SavePackage.h
+++ b/Engine/Source/Runtime/CoreUObject/Public/UObject/SavePackage.h
@@ -12,62 +12,6 @@ class FArchive;
 class FLinkerLoad;
 class FLinkerSave;
 
-class INameMapSaver
-{
-public:
-	virtual ~INameMapSaver() = default;
-
-	virtual void Begin() = 0;
-	virtual void End() = 0;
-	virtual void BeginPackage() = 0;
-	virtual void EndPackage(FLinkerSave& Linker, FLinkerLoad* Conform, FArchive* BinarySaver) = 0;
-
-	virtual void MarkNameAsReferenced(FName Name) = 0;
-	virtual int32 MapName(FName Name) const = 0;
-	virtual bool NameExistsInCurrentPackage(FNameEntryId ComparisonId) const = 0;
-};
-
-class FSinglePackageNameMapSaver : public INameMapSaver
-{
-public:
-	virtual void Begin() override {};
-	virtual void End() override {};
-	virtual void BeginPackage() override {};
-	virtual void EndPackage(FLinkerSave& Linker, FLinkerLoad* Conform, FArchive* BinarySaver) override;
-
-	virtual void MarkNameAsReferenced(FName Name) override;
-	virtual int32 MapName(FName Name) const override;
-	virtual bool NameExistsInCurrentPackage(FNameEntryId ComparisonId) const override;
-
-	void MarkNameAsReferenced(FNameEntryId Name);
-
-private:
-	TSet<FNameEntryId> ReferencedNames;
-	TMap<FNameEntryId, int32> NameIndices;
-};
-
-class FPackageStoreNameMapSaver : public INameMapSaver
-{
-public:
-	COREUOBJECT_API FPackageStoreNameMapSaver(const TCHAR* InFilename);
-
-	COREUOBJECT_API virtual void Begin() override {};
-	COREUOBJECT_API virtual void End() override;
-	virtual void BeginPackage();
-	virtual void EndPackage(FLinkerSave& Linker, FLinkerLoad* Conform, FArchive* BinarySaver) override {}
-
-	virtual void MarkNameAsReferenced(FName Name) override;
-	virtual int32 MapName(FName Name) const override;
-	virtual bool NameExistsInCurrentPackage(FNameEntryId ComparisonId) const override;
-
-private:
-	TMap<FNameEntryId, int32> NameIndices;
-	TArray<FNameEntryId> NameMap;
-	mutable TMap<FNameEntryId, TTuple<int32,int32,int32>> DebugNameCounts; // <Number0Count,OtherNumberCount,MaxNumber>
-	TSet<FNameEntryId> PackageReferencedNames;
-	const FString Filename;
-};
-
 class IBulkDataManifest 
 {
 public:
@@ -75,17 +19,15 @@ public:
 
 	virtual void Save() = 0;
 	virtual void AddFileAccess(const FString& PackageFilename, uint16 InIndex, uint64 InOffset, uint64 InSize) = 0;
-
-	static FIoChunkId CreateChunkId(uint32 NameIndex, uint32 NameNumber, uint16 ChunkIndex);
 };
 
 class FPackageStoreBulkDataManifest : public IBulkDataManifest
 {
 public:
 	COREUOBJECT_API FPackageStoreBulkDataManifest(const FString& InRootPath);
-	COREUOBJECT_API virtual ~FPackageStoreBulkDataManifest() = default;
+	COREUOBJECT_API virtual ~FPackageStoreBulkDataManifest();
 
-	bool COREUOBJECT_API Load();
+	COREUOBJECT_API bool Load();
 	virtual void Save() override;
 	
 	const FString& GetFilename() const { return Filename; }
@@ -121,16 +63,6 @@ private:
 	TMap<FString, PackageDesc> Data;
 };
 
-class FPackageHeaderSaver
-{
-public:
-	FPackageHeaderSaver(INameMapSaver& InNameMapSaver)
-		: NameMapSaver(InNameMapSaver) 
-	{}
-
-	INameMapSaver& NameMapSaver;
-};
-
 class FPackageStoreWriter
 {
 public:
@@ -194,14 +126,12 @@ private:
 class FSavePackageContext
 {
 public:
-	FSavePackageContext(FPackageHeaderSaver& InHeaderSaver, FPackageStoreWriter& InPackageStoreWriter, IBulkDataManifest& InBulkDataManifest)
-	: HeaderSaver(InHeaderSaver)
-	, PackageStoreWriter(InPackageStoreWriter) 
+	FSavePackageContext(FPackageStoreWriter& InPackageStoreWriter, IBulkDataManifest& InBulkDataManifest)
+	: PackageStoreWriter(InPackageStoreWriter) 
 	, BulkDataManifest(InBulkDataManifest)
 	{
 	}
 
-	FPackageHeaderSaver& HeaderSaver;
 	FPackageStoreWriter& PackageStoreWriter;
 	IBulkDataManifest& BulkDataManifest;
 };
