[AsyncLoading2] Moved all create chunk id function duplicates to IoDispatcher

diff --git a/Engine/Source/Developer/IoStoreUtilities/Private/IoStoreUtilities.cpp b/Engine/Source/Developer/IoStoreUtilities/Private/IoStoreUtilities.cpp
--- a/Engine/Source/Developer/IoStoreUtilities/Private/IoStoreUtilities.cpp
+++ b/Engine/Source/Developer/IoStoreUtilities/Private/IoStoreUtilities.cpp
@@ -183,30 +183,12 @@ FChunkIdCsv ChunkIdCsv;
 
 #endif
 
-enum class EZenChunkType : uint8
+static FIoChunkId CreateChunkId(int32 GlobalPackageId, uint16 ChunkIndex, EIoChunkType ChunkType, const TCHAR* DebugString)
 {
-	None,
-	PackageSummary,
-	ExportData,
-	ExportBundleData,
-	BulkData
-};
-
-static FIoChunkId CreateChunkId(int32 GlobalPackageId, uint16 ChunkIndex, EZenChunkType ChunkType, const TCHAR* DebugString)
-{
-	uint8 Data[12] = {0};
-
-	*reinterpret_cast<int32*>(&Data[0]) = GlobalPackageId;
-	*reinterpret_cast<uint16*>(&Data[8]) = ChunkIndex;
-	*reinterpret_cast<uint8*>(&Data[10]) = static_cast<uint8>(ChunkType);
-
-	FIoChunkId ChunkId;
-	ChunkId.Set(Data, 12);
-	
+	FIoChunkId ChunkId = CreateIoChunkId(GlobalPackageId, ChunkIndex, ChunkType);
 #if OUTPUT_CHUNKID_DIRECTORY
 	ChunkIdCsv.AddChunk(GlobalPackageId, 0, ChunkIndex, (uint8)ChunkType, GetTypeHash(ChunkId), DebugString);
 #endif
-
 	return ChunkId;
 }
 
@@ -1608,7 +1590,7 @@ int32 CreateTarget(const FContainerTarget& Target)
 			{
 				FIoBuffer IoBuffer(FIoBuffer::AssumeOwnership, ZenSummaryBuffer, ZenSummarySize);
 #if !SKIP_WRITE_CONTAINER
-				IoStoreWriter->Append(CreateChunkId(Package.GlobalPackageId, 0, EZenChunkType::PackageSummary, TEXT("PackageSummary")), IoBuffer);
+				IoStoreWriter->Append(CreateChunkId(Package.GlobalPackageId, 0, EIoChunkType::PackageSummary, TEXT("PackageSummary")), IoBuffer);
 #endif
 				++SummaryChunkCount;
 			}
@@ -1654,7 +1636,7 @@ int32 CreateTarget(const FContainerTarget& Target)
 					FObjectExport& ObjectExport = Exports[Package.ExportIndexOffset + I];
 					const int64 Offset = ObjectExport.SerialOffset - Package.UAssetSize;
 					FIoBuffer IoBuffer(FIoBuffer::Wrap, BundleBuffer, BundleBufferSize);
-					IoStoreWriter->Append(CreateChunkId(Package.GlobalPackageId, I, EZenChunkType::ExportBundleData, *Package.FileName), IoBuffer);
+					IoStoreWriter->Append(CreateChunkId(Package.GlobalPackageId, I, EIoChunkType::ExportBundleData, *Package.FileName), IoBuffer);
 #endif
 					++ExportChunkCount;
 					FMemory::Free(BundleBuffer);
@@ -1676,7 +1658,7 @@ int32 CreateTarget(const FContainerTarget& Target)
 				TUniquePtr<FArchive> BulkAr(IFileManager::Get().CreateFileReader(*UBulkFileName));
 				if (BulkAr != nullptr)
 				{
-					const FIoChunkId BulkDataChunkId = CreateChunkId(Package.GlobalPackageId, 0, EZenChunkType::BulkData, *Package.FileName);
+					const FIoChunkId BulkDataChunkId = CreateChunkId(Package.GlobalPackageId, 0, EIoChunkType::BulkData, *Package.FileName);
 #if !SKIP_WRITE_CONTAINER
 					uint8* BulkBuffer = static_cast<uint8*>(FMemory::Malloc(BulkAr->TotalSize()));
 					BulkAr->Serialize(BulkBuffer, BulkAr->TotalSize());
@@ -1702,7 +1684,7 @@ int32 CreateTarget(const FContainerTarget& Target)
 						for (const FPackageStoreBulkDataManifest::PackageDesc::BulkDataDesc& BulkDataDesc : PackageDesc->GetDataArray())
 						{
 #if !SKIP_WRITE_CONTAINER
-							const FIoChunkId AccessChunkId = CreateChunkId(Package.GlobalPackageId, BulkDataDesc.Index, EZenChunkType::BulkData, *Package.FileName);
+							const FIoChunkId AccessChunkId = CreateChunkId(Package.GlobalPackageId, BulkDataDesc.Index, EIoChunkType::BulkData, *Package.FileName);
 							const FIoStatus PartialResult = IoStoreWriter->MapPartialRange(BulkDataChunkId, BulkDataDesc.Offset, BulkDataDesc.Size, AccessChunkId);
 							if (!PartialResult.IsOk())
 							{
diff --git a/Engine/Source/Runtime/Core/Public/IO/IoDispatcher.h b/Engine/Source/Runtime/Core/Public/IO/IoDispatcher.h
--- a/Engine/Source/Runtime/Core/Public/IO/IoDispatcher.h
+++ b/Engine/Source/Runtime/Core/Public/IO/IoDispatcher.h
@@ -513,6 +513,9 @@ private:
 	friend class FIoBufferManager;
 };
 
+/**
+ * Identifier to a chunk of data.
+ */
 class FIoChunkId
 {
 public:
@@ -544,8 +547,6 @@ public:
 		return !(*this == Rhs);
 	}
 
-	CORE_API void GenerateFromData(const void* InData, SIZE_T InDataSize);
-
 	void Set(const void* InIdPtr, SIZE_T InSize)
 	{
 		check(InSize == sizeof Id);
@@ -570,6 +571,35 @@ private:
 	uint8	Id[12];
 };
 
+/**
+ * Addressable chunk types.
+ */
+enum class EIoChunkType : uint8
+{
+	Invalid,
+	PackageSummary,
+	ExportData,
+	ExportBundleData,
+	BulkData,
+};
+
+/**
+ * Creates a chunk identifier,
+ */
+static FIoChunkId CreateIoChunkId(uint32 GlobalPackageId, uint16 ChunkIndex, EIoChunkType IoChunkType)
+{
+	uint8 Data[12] = {0};
+
+	*reinterpret_cast<uint32*>(&Data[0]) = GlobalPackageId;
+	*reinterpret_cast<uint16*>(&Data[4]) = ChunkIndex;
+	*reinterpret_cast<uint8*>(&Data[11]) = static_cast<uint8>(IoChunkType);
+
+	FIoChunkId ChunkId;
+	ChunkId.Set(Data, 12);
+
+	return ChunkId;
+}
+
 //////////////////////////////////////////////////////////////////////////
 
 class FIoReadOptions
diff --git a/Engine/Source/Runtime/CoreUObject/Private/Serialization/AsyncLoading2.cpp b/Engine/Source/Runtime/CoreUObject/Private/Serialization/AsyncLoading2.cpp
--- a/Engine/Source/Runtime/CoreUObject/Private/Serialization/AsyncLoading2.cpp
+++ b/Engine/Source/Runtime/CoreUObject/Private/Serialization/AsyncLoading2.cpp
@@ -626,39 +626,6 @@ private:
 	FExternalIoRequests* ExternalIoRequests;
 };
 
-enum class EChunkType : uint8
-{
-	None,
-	PackageSummary,
-	ExportData,
-	ExportBundleData,
-	BulkData,
-};
-
-FIoChunkId CreateChunkId(FGlobalPackageId GlobalPackageId, uint16 ChunkIndex, EChunkType ChunkType)
-{
-	uint8 Data[12] = {0};
-
-	*reinterpret_cast<uint32*>(&Data[0]) = GlobalPackageId.Id;
-	*reinterpret_cast<uint16*>(&Data[8]) = ChunkIndex;
-	*reinterpret_cast<uint8*>(&Data[10]) = static_cast<uint8>(ChunkType);
-
-	FIoChunkId ChunkId;
-	ChunkId.Set(Data, 12);
-	return ChunkId;
-}
-
-FIoChunkId CreateChunkId(const FIoChunkId& ChunkId, uint16 ChunkIndex, EChunkType ChunkType)
-{
-	FIoChunkId Out = ChunkId;
-
-	uint8* Data = reinterpret_cast<uint8*>(&Out);
-	*reinterpret_cast<uint16*>(&Data[8]) = ChunkIndex;
-	*reinterpret_cast<uint8*>(&Data[10]) = static_cast<uint8>(ChunkType);
-
-	return Out;
-}
-
 enum class EAsyncPackageLoadingState2 : uint8
 {
 	NewPackage,
@@ -2410,7 +2377,7 @@ void FAsyncPackage2::StartLoading()
 	LoadStartTime = FPlatformTime::Seconds();;
 
 	FIoReadOptions ReadOptions;
-	AsyncLoadingThread.IoDispatcher.ReadWithCallback(CreateChunkId(GlobalPackageId, 0, EChunkType::PackageSummary),
+	AsyncLoadingThread.IoDispatcher.ReadWithCallback(CreateIoChunkId(GlobalPackageId.Id, 0, EIoChunkType::PackageSummary),
 		ReadOptions,
 		[this](TIoStatusOr<FIoBuffer> Result)
 		{
@@ -2487,7 +2454,7 @@ EAsyncPackageState::Type FAsyncPackage2::Event_StartExportBundleIo(FAsyncPackage
 	TRACE_CPUPROFILER_EVENT_SCOPE(Event_StartExportBundleIo);
 
 	FIoReadOptions ReadOptions;
-	Package->AsyncLoadingThread.IoDispatcher.ReadWithCallback(CreateChunkId(Package->GlobalPackageId, ExportBundleIndex, EChunkType::ExportBundleData),
+	Package->AsyncLoadingThread.IoDispatcher.ReadWithCallback(CreateIoChunkId(Package->GlobalPackageId.Id, ExportBundleIndex, EIoChunkType::ExportBundleData),
 		ReadOptions,
 		[Package, ExportBundleIndex](TIoStatusOr<FIoBuffer> Result)
 		{
diff --git a/Engine/Source/Runtime/CoreUObject/Private/Serialization/BulkData2.cpp b/Engine/Source/Runtime/CoreUObject/Private/Serialization/BulkData2.cpp
--- a/Engine/Source/Runtime/CoreUObject/Private/Serialization/BulkData2.cpp
+++ b/Engine/Source/Runtime/CoreUObject/Private/Serialization/BulkData2.cpp
@@ -27,31 +27,6 @@ namespace
 	static constexpr const TCHAR* OptionalExt = TEXT(".uptnl");			// Stored in a separate file that is optional
 
 	const uint16 InvalidBulkDataIndex = ~uint16(0);
-
-	// TODO: Move to common header
-	enum class EZenChunkType : uint8
-	{
-		None,
-		PackageSummary,
-		ExportData,
-		ExportBundleData,
-		BulkData
-	};
-
-	// TODO: Move to common header
-	FIoChunkId CreateChunkId(int32 GlobalPackageId, uint16 ChunkIndex, EZenChunkType ChunkType)
-	{
-		uint8 Data[12] = {0};
-
-		*reinterpret_cast<int32*>(&Data[0]) = GlobalPackageId;
-		*reinterpret_cast<uint16*>(&Data[8]) = ChunkIndex;
-		*reinterpret_cast<uint8*>(&Data[10]) = static_cast<uint8>(ChunkType);
-
-		FIoChunkId ChunkId;
-		ChunkId.Set(Data, 12);
-		
-		return ChunkId;
-	}
 }
 
 // TODO: The code in the FileTokenSystem namespace is a temporary system so that FBulkDataBase can hold
@@ -599,7 +574,7 @@ void FBulkDataBase::Serialize(FArchive& Ar, UObject* Owner, int32 /*Index*/, boo
 			if (BulkDataIndex != InvalidBulkDataIndex)
 			{
 				UPackage* Package = Owner->GetOutermost();
-				ChunkID = CreateChunkId(Package->GetGlobalPackageId(), BulkDataIndex, EZenChunkType::BulkData);
+				ChunkID = CreateIoChunkId(Package->GetGlobalPackageId(), BulkDataIndex, EIoChunkType::BulkData);
 			}
 			else
 			{
