Move the BulkData manifest file to be written to <project_root>\saved\cooked\<platform>\<project_name>\Metadata\ which is where the other data files produced during cook but not included in an actual release are being placed, as suggested by PJ.

#rnx

diff --git a/Engine/Source/Developer/IoStoreUtilities/Private/IoStoreUtilities.cpp b/Engine/Source/Developer/IoStoreUtilities/Private/IoStoreUtilities.cpp
--- a/Engine/Source/Developer/IoStoreUtilities/Private/IoStoreUtilities.cpp
+++ b/Engine/Source/Developer/IoStoreUtilities/Private/IoStoreUtilities.cpp
@@ -40,6 +40,7 @@ struct FContainerTarget
 {
 	ITargetPlatform* TargetPlatform;
 	FString CookedDirectory;
+	FString CookedProjectDirectory;
 	FString OutputDirectory;
 };
 
@@ -984,11 +985,11 @@ int32 CreateTarget(const FContainerTarget& Target)
 {
 	TGuardValue<int32> GuardAllowUnversionedContentInEditor(GAllowUnversionedContentInEditor, 1);
 
-	FString CookedDir = Target.CookedDirectory;
-	FString OutputDir = Target.OutputDirectory;
-	FString RelativePrefixForLegacyFilename = TEXT("../../../");
+	const FString CookedDir = Target.CookedDirectory;
+	const FString OutputDir = Target.OutputDirectory;
+	const FString RelativePrefixForLegacyFilename = TEXT("../../../");
 	
-	FPackageStoreBulkDataManifest BulkDataManifest(Target.TargetPlatform->PlatformName());
+	FPackageStoreBulkDataManifest BulkDataManifest(Target.CookedProjectDirectory);
 	const bool bWithBulkDataManifest = BulkDataManifest.Load();
 	if (bWithBulkDataManifest)
 	{
@@ -1844,11 +1845,13 @@ int32 CreateIoStoreContainerFiles(const TCHAR* CmdLine)
 	}
 
 	FString TargetCookedDirectory = FPaths::ProjectSavedDir() / TEXT("Cooked") / TargetPlatform->PlatformName();
+	FString TargetCookedProjectDirectory = TargetCookedDirectory / FApp::GetProjectName();
+
 	FString TargetOutputDirectory = OutputDirectory.Len() > 0
 		? OutputDirectory
-		: TargetCookedDirectory / FApp::GetProjectName();
+		: TargetCookedProjectDirectory;
 
-	FContainerTarget Target { TargetPlatform, TargetCookedDirectory, TargetOutputDirectory };
+	FContainerTarget Target { TargetPlatform, TargetCookedDirectory, TargetCookedProjectDirectory, TargetOutputDirectory };
 
 	UE_LOG(LogIoStore, Display, TEXT("Creating target: '%s' using output directory: '%s'"), *Target.TargetPlatform->PlatformName(), *Target.OutputDirectory);
 
diff --git a/Engine/Source/Editor/UnrealEd/Private/CookOnTheFlyServer.cpp b/Engine/Source/Editor/UnrealEd/Private/CookOnTheFlyServer.cpp
--- a/Engine/Source/Editor/UnrealEd/Private/CookOnTheFlyServer.cpp
+++ b/Engine/Source/Editor/UnrealEd/Private/CookOnTheFlyServer.cpp
@@ -6714,15 +6714,20 @@ void UCookOnTheFlyServer::InitializePackageStore(const TArray<FName>& TargetPlat
 	const FString RootPath = FPaths::RootDir();
 	const FString RootPathSandbox = ConvertToFullSandboxPath(*RootPath, true);
 
+	const FString ProjectPath = FPaths::ProjectDir();
+	const FString ProjectPathSandbox = ConvertToFullSandboxPath(*ProjectPath, true);
+
 	SavePackageContexts.Reserve(TargetPlatformNames.Num());
 
 	for (const FName PlatformName : TargetPlatformNames)
 	{
 		const FString PlatformString = PlatformName.ToString();
+
 		const FString ResolvedRootPath = RootPathSandbox.Replace(TEXT("[Platform]"), *PlatformString);
+		const FString ResolvedProjectPath = ProjectPathSandbox.Replace(TEXT("[Platform]"), *PlatformString);
 
 		// just leak all memory for now
-		FPackageStoreBulkDataManifest* BulkDataManifest	= new FPackageStoreBulkDataManifest(PlatformString);
+		FPackageStoreBulkDataManifest* BulkDataManifest	= new FPackageStoreBulkDataManifest(ResolvedProjectPath);
 		FLooseFileWriter* LooseFileWriter				= IsUsingPackageStore() ? new FLooseFileWriter() : nullptr;
 
 		FSavePackageContext* SavePackageContext			= new FSavePackageContext(LooseFileWriter, BulkDataManifest);
diff --git a/Engine/Source/Runtime/CoreUObject/Private/UObject/SavePackage.cpp b/Engine/Source/Runtime/CoreUObject/Private/UObject/SavePackage.cpp
--- a/Engine/Source/Runtime/CoreUObject/Private/UObject/SavePackage.cpp
+++ b/Engine/Source/Runtime/CoreUObject/Private/UObject/SavePackage.cpp
@@ -482,9 +482,9 @@ void FPackageStoreBulkDataManifest::PackageDesc::AddData(uint64 InChunkId, uint6
 	Entry.Size = InSize;
 }
 
-FPackageStoreBulkDataManifest::FPackageStoreBulkDataManifest(const FString& PlatformName)
+FPackageStoreBulkDataManifest::FPackageStoreBulkDataManifest(const FString& ProjectPath)
 {
-	Filename = FPaths::ProjectSavedDir() / TEXT("BulkDataManifests") / PlatformName + TEXT(".ubulkmanifest");
+	Filename = ProjectPath / TEXT("Metadata/BulkDataInfo.ubulkmanifest");
 }
 
 FArchive& operator<<(FArchive& Ar, FPackageStoreBulkDataManifest::PackageDesc::BulkDataDesc& Entry)
diff --git a/Engine/Source/Runtime/CoreUObject/Public/UObject/SavePackage.h b/Engine/Source/Runtime/CoreUObject/Public/UObject/SavePackage.h
--- a/Engine/Source/Runtime/CoreUObject/Public/UObject/SavePackage.h
+++ b/Engine/Source/Runtime/CoreUObject/Public/UObject/SavePackage.h
@@ -15,7 +15,7 @@ class FLinkerSave;
 class FPackageStoreBulkDataManifest
 {
 public:
-	COREUOBJECT_API FPackageStoreBulkDataManifest(const FString& PlatformName);
+	COREUOBJECT_API FPackageStoreBulkDataManifest(const FString& ProjectPath);
 	COREUOBJECT_API ~FPackageStoreBulkDataManifest();
 
 	COREUOBJECT_API bool Load();
