Added support for loading optional and duplicate non optional bulkdata from the IoDispatcher.

- Save the EIoChunkType for each bulkdata entry to the bulkdata manifest.
- Add ::DoesExist to both old and new bulkdata, this returns if the file containing the data actually exists or not, unlike ::CanLoadFromDisk which only returns if we can potentially load (it does not check if the data actually exists or not)
- Added EIoChunkType::OptionalBulkData to reference optional bulkdata
- Moved the code for adding bulkdata in IoStoreUtilities.cpp to it's own method that we can call for either .ubulk files or .uptnl files.
- Added a log error if we cannot find a valid  FPackageStoreBulkDataManifest::PackageDesc for a given file, as this indicates we have mismatched data, but do not fatally close the commandlet.
- The mipstreaming systems can now poll BulkData to see if the file exists (ie optional data) rather than storing a FString locally which allows us to abstract the knowledge of where the file ie (loose file vs IoDispatcher)
 - This will also reduce the memory costs as we will not be storing the filename in the FStreamingRenderAsset and the bulkdata.
 - Implemented for UTexture and UStaticMesh
 - Behaviour is controlled by STORE_OPTIONAL_DATA_FILENAME to allow the change to be easily reverted until we commit to it
- Duplicate non optional data can now work properly with FIoChunkId

diff --git a/Engine/Source/Developer/IoStoreUtilities/Private/IoStoreUtilities.cpp b/Engine/Source/Developer/IoStoreUtilities/Private/IoStoreUtilities.cpp
--- a/Engine/Source/Developer/IoStoreUtilities/Private/IoStoreUtilities.cpp
+++ b/Engine/Source/Developer/IoStoreUtilities/Private/IoStoreUtilities.cpp
@@ -851,6 +851,60 @@ static void BuildBundles(FExportGraph& ExportGraph, TMap<FName, FPackage*>& Pack
 	}
 }
 
+static bool WriteBulkData(	const FString& Filename, EIoChunkType Type, const FPackage& Package, const FPackageStoreBulkDataManifest& BulkDataManifest,
+							FIoStoreWriter* IoStoreWriter, uint64& OutBulkChunkCount, uint64& OutBulkPartialChunkCount)
+{
+	TUniquePtr<FArchive> BulkAr(IFileManager::Get().CreateFileReader(*Filename));
+	if (BulkAr != nullptr)
+	{
+		const FIoChunkId BulkDataChunkId = CreateChunkIdForBulkData(Package.GlobalPackageId, INDEX_NONE, Type, *Package.FileName);
+#if !SKIP_WRITE_CONTAINER
+		uint8* BulkBuffer = static_cast<uint8*>(FMemory::Malloc(BulkAr->TotalSize()));
+		BulkAr->Serialize(BulkBuffer, BulkAr->TotalSize());
+		FIoBuffer IoBuffer(FIoBuffer::AssumeOwnership, BulkBuffer, BulkAr->TotalSize());
+		const FIoStatus AppendResult = IoStoreWriter->Append(BulkDataChunkId, IoBuffer);
+		if (!AppendResult.IsOk())
+		{
+			UE_LOG(LogIoStore, Error, TEXT("Failed to append bulkdata for '%s' due to: %s"), *Package.FileName, *AppendResult.ToString());
+			return false;
+		}
+
+		BulkAr->Close();
+#endif
+		++OutBulkChunkCount;
+
+		// Create additional mapping chunks as needed
+		const FPackageStoreBulkDataManifest::PackageDesc* PackageDesc = BulkDataManifest.Find(Package.FileName);
+		if (PackageDesc != nullptr)
+		{
+			for (const FPackageStoreBulkDataManifest::PackageDesc::BulkDataDesc& BulkDataDesc : PackageDesc->GetDataArray())
+			{
+				if (BulkDataDesc.Type == Type)
+				{
+#if !SKIP_WRITE_CONTAINER
+					const FIoChunkId AccessChunkId = CreateChunkIdForBulkData(Package.GlobalPackageId, BulkDataDesc.ChunkId, Type, *Package.FileName);
+					const FIoStatus PartialResult = IoStoreWriter->MapPartialRange(BulkDataChunkId, BulkDataDesc.Offset, BulkDataDesc.Size, AccessChunkId);
+					if (!PartialResult.IsOk())
+					{
+						UE_LOG(LogIoStore, Error, TEXT("Failed to map partial range for '%s' due to: %s"), *Package.FileName, *PartialResult.ToString());
+						return false;
+					}
+#endif
+				}
+				
+				++OutBulkPartialChunkCount;
+			}
+		}
+		else
+		{
+			UE_LOG(LogIoStore, Error, TEXT("Unable to find an entry in the bulkdata manifest for '%s' the file might be out of date!"), *Package.FileName);
+			return false;
+		}
+	}
+
+	return true;
+}
+
 struct FImportData
 {
 	int32 GlobalIndex = -1;
@@ -1660,54 +1714,20 @@ int32 CreateTarget(const FContainerTarget& Target)
 			// Bulk chunks
 			if (bWithBulkDataManifest)
 			{
-				FString UBulkFileName = FPaths::ChangeExtension(Package.FileName, TEXT(".ubulk"));
-				FPaths::NormalizeFilename(UBulkFileName);
+				FString BulkFileName = FPaths::ChangeExtension(Package.FileName, TEXT(".ubulk"));
+				FPaths::NormalizeFilename(BulkFileName);
 
-				TUniquePtr<FArchive> BulkAr(IFileManager::Get().CreateFileReader(*UBulkFileName));
-				if (BulkAr != nullptr)
+				if (!WriteBulkData(BulkFileName, EIoChunkType::BulkData, Package, BulkDataManifest, IoStoreWriter.Get(), BulkChunkCount, BulkPartialChunkCount))
 				{
-					const FIoChunkId BulkDataChunkId = CreateChunkIdForBulkData(Package.GlobalPackageId, INDEX_NONE, EIoChunkType::BulkData, *Package.FileName);
-#if !SKIP_WRITE_CONTAINER
-					uint8* BulkBuffer = static_cast<uint8*>(FMemory::Malloc(BulkAr->TotalSize()));
-					BulkAr->Serialize(BulkBuffer, BulkAr->TotalSize());
-					FIoBuffer IoBuffer(FIoBuffer::AssumeOwnership, BulkBuffer, BulkAr->TotalSize());
-					const FIoStatus AppendResult = IoStoreWriter->Append(BulkDataChunkId, IoBuffer);
-					if (!AppendResult.IsOk())
-					{
-						UE_LOG(LogIoStore, Error, TEXT("Failed to append bulkdata for '%s' due to: %s"), *Package.FileName, *AppendResult.ToString());
-						RequestEngineExit(TEXT("ZenCreator failed"));
-						return 1;
-					}
+					return 1;
+				}
 
-					BulkAr->Close();
-#endif
-					++BulkChunkCount;
+				FString OptionalBulkFileName = FPaths::ChangeExtension(Package.FileName, TEXT(".uptnl"));
+				FPaths::NormalizeFilename(OptionalBulkFileName);
 
-					// Create additional mapping chunks as needed
-					FString BulkDataFileName = Package.FileName;
-					FPaths::MakePathRelativeTo(BulkDataFileName, *OutputDir);
-					const FPackageStoreBulkDataManifest::PackageDesc* PackageDesc = BulkDataManifest.Find(Package.FileName);
-					if (PackageDesc != nullptr)
-					{
-						for (const FPackageStoreBulkDataManifest::PackageDesc::BulkDataDesc& BulkDataDesc : PackageDesc->GetDataArray())
-						{
-#if !SKIP_WRITE_CONTAINER
-							const FIoChunkId AccessChunkId = CreateChunkIdForBulkData(Package.GlobalPackageId, BulkDataDesc.ChunkId, EIoChunkType::BulkData, *Package.FileName);
-							const FIoStatus PartialResult = IoStoreWriter->MapPartialRange(BulkDataChunkId, BulkDataDesc.Offset, BulkDataDesc.Size, AccessChunkId);
-							if (!PartialResult.IsOk())
-							{
-								UE_LOG(LogIoStore, Error, TEXT("Failed to map partial range for '%s' due to: %s"), *Package.FileName, *PartialResult.ToString());
-								RequestEngineExit(TEXT("ZenCreator failed"));
-								return 1;
-							}
-#endif
-							++BulkPartialChunkCount;
-						}
-					}
-					else
-					{
-						UE_LOG(LogIoStore, Error, TEXT("Unable to find an entry in the bulkdata manifest for '%s' the file might be out of date!"), *Package.FileName);
-					}
+				if (!WriteBulkData(OptionalBulkFileName, EIoChunkType::OptionalBulkData, Package, BulkDataManifest, IoStoreWriter.Get(), BulkChunkCount, BulkPartialChunkCount))
+				{
+					return 1;
 				}
 			}
 #endif
diff --git a/Engine/Source/Runtime/Core/Public/IO/IoDispatcher.h b/Engine/Source/Runtime/Core/Public/IO/IoDispatcher.h
--- a/Engine/Source/Runtime/Core/Public/IO/IoDispatcher.h
+++ b/Engine/Source/Runtime/Core/Public/IO/IoDispatcher.h
@@ -581,6 +581,7 @@ enum class EIoChunkType : uint8
 	ExportData,
 	ExportBundleData,
 	BulkData,
+	OptionalBulkData,
 	LoaderGlobalMeta,
 	LoaderInitialLoadMeta,
 	LoaderGlobalNameMap
diff --git a/Engine/Source/Runtime/CoreUObject/Private/Serialization/BulkData.cpp b/Engine/Source/Runtime/CoreUObject/Private/Serialization/BulkData.cpp
--- a/Engine/Source/Runtime/CoreUObject/Private/Serialization/BulkData.cpp
+++ b/Engine/Source/Runtime/CoreUObject/Private/Serialization/BulkData.cpp
@@ -413,6 +413,11 @@ bool FUntypedBulkData::CanLoadFromDisk() const
 #endif // WITH_EDITOR
 }
 
+bool FUntypedBulkData::DoesExist() const
+{
+	return IFileManager::Get().FileExists(*Filename);
+}
+
 /**
  * Returns flags usable to decompress the bulk data
  * 
diff --git a/Engine/Source/Runtime/CoreUObject/Private/Serialization/BulkData2.cpp b/Engine/Source/Runtime/CoreUObject/Private/Serialization/BulkData2.cpp
--- a/Engine/Source/Runtime/CoreUObject/Private/Serialization/BulkData2.cpp
+++ b/Engine/Source/Runtime/CoreUObject/Private/Serialization/BulkData2.cpp
@@ -18,6 +18,9 @@ DEFINE_LOG_CATEGORY_STATIC(LogBulkDataRuntime, Log, All);
 // If set to 0 then the loose file fallback will be used even if the -ZenLoader command line flag is present.
 #define ENABLE_IO_DISPATCHER 1
 
+// If set to 0 then we will pretend that optional data does not exist, useful for testing.
+#define ALLOW_OPTIONAL_DATA 1
+
 namespace
 {
 	// TODO: Maybe expose this and start using everywhere?
@@ -566,8 +569,6 @@ void FBulkDataBase::Serialize(FArchive& Ar, UObject* Owner, int32 /*Index*/, boo
 			Ar << DummyValue;
 		}
 
-		check((BulkDataFlags & BULKDATA_UsesIoDispatcher) == 0);
-
 		// Assuming that Owner/Package/Linker are all valid, the old BulkData system would
 		// generally fail if any of these were nullptr but had plenty of inconsistent checks
 		// scattered throughout.
@@ -577,7 +578,8 @@ void FBulkDataBase::Serialize(FArchive& Ar, UObject* Owner, int32 /*Index*/, boo
 
 		if (!IsInlined() && bUseZenLoader)
 		{		
-			ChunkID = CreateBulkdataChunkId(Package->GetGlobalPackageId(), BulkDataOffsetInFile, EIoChunkType::BulkData);
+			EIoChunkType Type = IsOptional() ? EIoChunkType::OptionalBulkData : EIoChunkType::BulkData;
+			ChunkID = CreateBulkdataChunkId(Package->GetGlobalPackageId(), BulkDataOffsetInFile, Type);
 
 			BulkDataFlags |= BULKDATA_UsesIoDispatcher;
 		}
@@ -605,7 +607,7 @@ void FBulkDataBase::Serialize(FArchive& Ar, UObject* Owner, int32 /*Index*/, boo
 			Filename = Linker->Filename;
 			PackageName = Package->FileName;
 		}
-	
+
 		FArchive* CacheableArchive = Ar.GetCacheableArchive();
 		if (Ar.IsAllowingLazyLoading() && CacheableArchive != nullptr)
 		{
@@ -619,15 +621,35 @@ void FBulkDataBase::Serialize(FArchive& Ar, UObject* Owner, int32 /*Index*/, boo
 			{
 				if (IsDuplicateNonOptional())
 				{
-					if (IFileManager::Get().FileExists(*Filename))
+					auto DoesOptionalDataExist = [this](const FString& PackageFilename)->bool
+					{
+#if ALLOW_OPTIONAL_DATA
+						if (!IsUsingIODispatcher())
+						{
+							FString OptionalDataFilename = ConvertFilenameFromFlags(PackageFilename);
+							return IFileManager::Get().FileExists(*OptionalDataFilename);
+						}
+						else
+						{
+							return IoDispatcher->DoesChunkExist(ChunkID);
+						}
+#else
+						return false;
+#endif
+					};
+
+					if (DoesOptionalDataExist(Filename))
 					{
 						SerializeDuplicateData(Ar, Linker, BulkDataFlags, BulkDataSizeOnDisk, BulkDataOffsetInFile);
 
 						// Update the fallback data again if needed
-						if (!bUseZenLoader || ChunkID == FIoChunkId::InvalidChunkId)
+						if (!IsInlined() && bUseZenLoader)
+						{
+							ChunkID = CreateBulkdataChunkId(Package->GetGlobalPackageId(), BulkDataOffsetInFile, EIoChunkType::OptionalBulkData);
+							BulkDataFlags |= BULKDATA_UsesIoDispatcher;
+						}
+						else
 						{
-							// Remove the flag if we are not going to use it
-							BulkDataFlags &= ~BULKDATA_UsesIoDispatcher;
 							Fallback.Token = InvalidToken;
 							Fallback.BulkDataSize = BulkDataSize;
 						}
@@ -647,6 +669,8 @@ void FBulkDataBase::Serialize(FArchive& Ar, UObject* Owner, int32 /*Index*/, boo
 			{
 				Fallback.Token = FileTokenSystem::RegisterFileToken( PackageName, Filename, BulkDataOffsetInFile);
 			}
+
+		//	DebugFilename = ConvertFilenameFromFlags(Filename);
 		}
 		else
 		{
@@ -797,6 +821,25 @@ int64 FBulkDataBase::GetBulkDataSize() const
 	}	
 }
 
+bool FBulkDataBase::DoesExist() const
+{
+#if ALLOW_OPTIONAL_DATA
+	if (!IsUsingIODispatcher())
+	{
+		FString Filename = FileTokenSystem::GetFilename(Fallback.Token);
+		Filename = ConvertFilenameFromFlags(Filename);
+
+		return IFileManager::Get().FileExists(*Filename);
+	}
+	else
+	{
+		return IoDispatcher->DoesChunkExist(ChunkID);
+	}
+#else
+	return false;
+#endif
+}
+
 bool FBulkDataBase::IsStoredCompressedOnDisk() const
 {
 	return (GetBulkDataFlags() & BULKDATA_SerializeCompressed) != 0;
@@ -1112,13 +1155,16 @@ void FBulkDataBase::FreeData()
 
 FString FBulkDataBase::ConvertFilenameFromFlags(const FString& Filename) const
 {
-	if (IsInlined())
+	if (IsOptional())
 	{
-		return FPaths::ChangeExtension(Filename, InlinedExt);
+		// Optional data should be tested for first as we in theory can have data that would
+		// be marked as inline, also marked as optional and in this case we should treat it as
+		// optional data first.
+		return FPaths::ChangeExtension(Filename, OptionalExt);
 	}
-	else if (IsOptional())
+	else if (IsInlined())
 	{
-		return FPaths::ChangeExtension(Filename, OptionalExt);
+		return FPaths::ChangeExtension(Filename, InlinedExt);
 	}
 	else if (IsMemoryMapped())
 	{
diff --git a/Engine/Source/Runtime/CoreUObject/Private/UObject/SavePackage.cpp b/Engine/Source/Runtime/CoreUObject/Private/UObject/SavePackage.cpp
--- a/Engine/Source/Runtime/CoreUObject/Private/UObject/SavePackage.cpp
+++ b/Engine/Source/Runtime/CoreUObject/Private/UObject/SavePackage.cpp
@@ -469,19 +469,18 @@ private:
 	TSet<FNameEntryId> ReferencedNames;
 };
 
-void FPackageStoreBulkDataManifest::PackageDesc::AddData(uint64 InChunkId, uint64 InOffset, uint64 InSize)
+void FPackageStoreBulkDataManifest::PackageDesc::AddData(EIoChunkType InType, uint64 InChunkId, uint64 InOffset, uint64 InSize)
 {
-#if 0 // DIsable until optional data is working again
-	// The ChunkId is supposed to be unique
-	auto func = [=](const BulkDataDesc& Entry) { return Entry.ChunkId == InChunkId; };
+	// The ChunkId is supposed to be unique for each type
+	auto func = [=](const BulkDataDesc& Entry) { return Entry.ChunkId == InChunkId && Entry.Type == InType; };
 	check(Data.FindByPredicate(func) == nullptr);
-#endif 
 
 	BulkDataDesc& Entry = Data.Emplace_GetRef();
 
 	Entry.ChunkId = InChunkId;
 	Entry.Offset = InOffset;
 	Entry.Size = InSize;
+	Entry.Type = InType;
 }
 
 FPackageStoreBulkDataManifest::FPackageStoreBulkDataManifest(const FString& ProjectPath)
@@ -494,6 +493,7 @@ FArchive& operator<<(FArchive& Ar, FPackageStoreBulkDataManifest::PackageDesc::B
 	Ar << Entry.ChunkId;
 	Ar << Entry.Offset;
 	Ar << Entry.Size; 
+	Ar << Entry.Type;
 
 	return Ar;
 }
@@ -537,12 +537,12 @@ const FPackageStoreBulkDataManifest::PackageDesc* FPackageStoreBulkDataManifest:
 	return Data.Find(NormalizedFilename);
 }
 
-void FPackageStoreBulkDataManifest::AddFileAccess(const FString& PackageFilename, uint64 InChunkId, uint64 InOffset, uint64 InSize)
+void FPackageStoreBulkDataManifest::AddFileAccess(const FString& PackageFilename, EIoChunkType InType, uint64 InChunkId, uint64 InOffset, uint64 InSize)
 {
 	const FString NormalizedFilename = FixFilename(PackageFilename);
 	
 	PackageDesc& Entry = GetOrCreateFileAccess(NormalizedFilename);
-	Entry.AddData(InChunkId, InOffset, InSize);
+	Entry.AddData(InType, InChunkId, InOffset, InSize);
 }
 
 FPackageStoreBulkDataManifest::PackageDesc& FPackageStoreBulkDataManifest::GetOrCreateFileAccess(const FString& PackageFilename)
@@ -6285,7 +6285,10 @@ void SaveBulkData(FLinkerSave* Linker, const UPackage* InOuter, const TCHAR* Fil
 
 			if (SavePackageContext != nullptr && SavePackageContext->BulkDataManifest != nullptr)
 			{
-				SavePackageContext->BulkDataManifest->AddFileAccess(Filename, StoredBulkStartOffset, BulkStartOffset, SizeOnDisk);
+				const bool bIsOptional = (BulkDataStorageInfo.BulkDataFlags & BULKDATA_OptionalPayload) != 0;
+				const EIoChunkType Type = bIsOptional ? EIoChunkType::OptionalBulkData : EIoChunkType::BulkData;
+
+				SavePackageContext->BulkDataManifest->AddFileAccess(Filename, Type, StoredBulkStartOffset, BulkStartOffset, SizeOnDisk);
 			}
 			
 			Linker->Seek(LinkerEndOffset);
diff --git a/Engine/Source/Runtime/CoreUObject/Public/Serialization/BulkData.h b/Engine/Source/Runtime/CoreUObject/Public/Serialization/BulkData.h
--- a/Engine/Source/Runtime/CoreUObject/Public/Serialization/BulkData.h
+++ b/Engine/Source/Runtime/CoreUObject/Public/Serialization/BulkData.h
@@ -11,7 +11,7 @@
 #include "BulkData2.h"
 
 #if WITH_EDITOR == 0 && WITH_EDITORONLY_DATA == 0 //Runtime
-	#define USE_NEW_BULKDATA 0 // Set to 1 to enable 
+	#define USE_NEW_BULKDATA 1 // Set to 1 to enable 
 #else
 	#define USE_NEW_BULKDATA 0
 #endif
@@ -444,6 +444,11 @@ public:
 	 * Returns true if the data can be loaded from disk.
 	 */
 	bool CanLoadFromDisk() const;
+	
+	/**
+	 * Returns true if the data references a file that currently exists and can be referenced by the file system.
+	 */
+	bool DoesExist() const;
 
 	/**
 	 * Returns flags usable to decompress the bulk data
diff --git a/Engine/Source/Runtime/CoreUObject/Public/Serialization/BulkData2.h b/Engine/Source/Runtime/CoreUObject/Public/Serialization/BulkData2.h
--- a/Engine/Source/Runtime/CoreUObject/Public/Serialization/BulkData2.h
+++ b/Engine/Source/Runtime/CoreUObject/Public/Serialization/BulkData2.h
@@ -92,6 +92,11 @@ public:
 	uint32 GetBulkDataFlags() const { return BulkDataFlags; }
 
 	bool CanLoadFromDisk() const { return IsUsingIODispatcher() || Fallback.Token != InvalidToken; }
+	
+	/**
+	 * Returns true if the data references a file that currently exists and can be referenced by the file system.
+	 */
+	bool DoesExist() const;
 
 	bool IsStoredCompressedOnDisk() const;
 	FName GetDecompressionFormat() const;
diff --git a/Engine/Source/Runtime/CoreUObject/Public/UObject/SavePackage.h b/Engine/Source/Runtime/CoreUObject/Public/UObject/SavePackage.h
--- a/Engine/Source/Runtime/CoreUObject/Public/UObject/SavePackage.h
+++ b/Engine/Source/Runtime/CoreUObject/Public/UObject/SavePackage.h
@@ -21,7 +21,7 @@ public:
 	COREUOBJECT_API bool Load();
 	COREUOBJECT_API void Save() ;
 
-	void AddFileAccess(const FString& PackageFilename, uint64 InChunkId, uint64 InOffset, uint64 InSize);
+	void AddFileAccess(const FString& PackageFilename, EIoChunkType InType, uint64 InChunkId, uint64 InOffset, uint64 InSize);
 	
 	const FString& GetFilename() const { return Filename; }
 
@@ -34,9 +34,10 @@ public:
 							// applied, to make it easier to compute at runtime.
 			uint64 Offset;
 			uint64 Size;
+			EIoChunkType Type;
 		};
 
-		void AddData(uint64 InChunkId, uint64 InOffset, uint64 InSize);
+		void AddData(EIoChunkType InType, uint64 InChunkId, uint64 InOffset, uint64 InSize);
 		const TArray<BulkDataDesc>& GetDataArray() const { return Data; }
 	private:
 		friend FArchive& operator<<(FArchive& Ar, PackageDesc& Entry);
diff --git a/Engine/Source/Runtime/Engine/Classes/Engine/SkeletalMesh.h b/Engine/Source/Runtime/Engine/Classes/Engine/SkeletalMesh.h
--- a/Engine/Source/Runtime/Engine/Classes/Engine/SkeletalMesh.h
+++ b/Engine/Source/Runtime/Engine/Classes/Engine/SkeletalMesh.h
@@ -999,6 +999,7 @@ public:
 	virtual int32 CalcNumOptionalMips() const final override;
 	virtual int32 CalcCumulativeLODSize(int32 NumLODs) const final override;
 	virtual bool GetMipDataFilename(const int32 MipIndex, FString& OutBulkDataFilename) const final override;
+	virtual bool DoesMipDataExist(const int32 MipIndex) const final override;
 	virtual bool IsReadyForStreaming() const final override;
 	virtual int32 GetNumResidentMips() const final override;
 	virtual int32 GetNumRequestedMips() const final override;
diff --git a/Engine/Source/Runtime/Engine/Classes/Engine/StaticMesh.h b/Engine/Source/Runtime/Engine/Classes/Engine/StaticMesh.h
--- a/Engine/Source/Runtime/Engine/Classes/Engine/StaticMesh.h
+++ b/Engine/Source/Runtime/Engine/Classes/Engine/StaticMesh.h
@@ -979,6 +979,7 @@ public:
 	virtual int32 CalcNumOptionalMips() const final override;
 	virtual int32 CalcCumulativeLODSize(int32 NumLODs) const final override;
 	virtual bool GetMipDataFilename(const int32 MipIndex, FString& BulkDataFilename) const final override;
+	virtual bool DoesMipDataExist(const int32 MipIndex) const final override;
 	virtual bool IsReadyForStreaming() const final override;
 	virtual int32 GetNumResidentMips() const final override;
 	virtual int32 GetNumRequestedMips() const final override;
diff --git a/Engine/Source/Runtime/Engine/Classes/Engine/StreamableRenderAsset.h b/Engine/Source/Runtime/Engine/Classes/Engine/StreamableRenderAsset.h
--- a/Engine/Source/Runtime/Engine/Classes/Engine/StreamableRenderAsset.h
+++ b/Engine/Source/Runtime/Engine/Classes/Engine/StreamableRenderAsset.h
@@ -55,6 +55,12 @@ public:
 		return false;
 	}
 
+	virtual bool DoesMipDataExist(const int32 MipIndex) const
+	{
+		STREAMABLERENDERASSET_NODEFAULT(DoesMipDataExist);
+		return false;
+	}
+
 	virtual bool IsReadyForStreaming() const
 	{
 		STREAMABLERENDERASSET_NODEFAULT(IsReadyForStreaming);
diff --git a/Engine/Source/Runtime/Engine/Classes/Engine/Texture2D.h b/Engine/Source/Runtime/Engine/Classes/Engine/Texture2D.h
--- a/Engine/Source/Runtime/Engine/Classes/Engine/Texture2D.h
+++ b/Engine/Source/Runtime/Engine/Classes/Engine/Texture2D.h
@@ -130,7 +130,8 @@ public:
 	virtual int32 CalcCumulativeLODSize(int32 NumLODs) const final override { return CalcTextureMemorySize(NumLODs); }
 
 	virtual bool GetMipDataFilename(const int32 MipIndex, FString& OutBulkDataFilename) const final override;
-
+	virtual bool DoesMipDataExist(const int32 MipIndex) const final override;
+	
 	/**
 	* Returns whether the texture is ready for streaming aka whether it has had InitRHI called on it.
 	*
diff --git a/Engine/Source/Runtime/Engine/Private/SkeletalMesh.cpp b/Engine/Source/Runtime/Engine/Private/SkeletalMesh.cpp
--- a/Engine/Source/Runtime/Engine/Private/SkeletalMesh.cpp
+++ b/Engine/Source/Runtime/Engine/Private/SkeletalMesh.cpp
@@ -984,6 +984,18 @@ bool USkeletalMesh::GetMipDataFilename(const int32 MipIndex, FString& OutBulkDat
 #endif
 }
 
+bool USkeletalMesh::DoesMipDataExist(const int32 MipIndex) const
+{
+	check(MipIndex < CalcNumOptionalMips());
+
+#if !USE_BULKDATA_STREAMING_TOKEN	
+	return SkeletalMeshRenderData->LODRenderData[MipIndex].StreamingBulkData.DoesExist();
+#else
+	checkf(false, TEXT("Should not be possible to reach this path, if USE_NEW_BULKDATA is enabled then USE_BULKDATA_STREAMING_TOKEN should be disabled!"));
+	return false;
+#endif
+}
+
 bool USkeletalMesh::IsReadyForStreaming() const
 {
 	return SkeletalMeshRenderData && SkeletalMeshRenderData->bReadyForStreaming;
diff --git a/Engine/Source/Runtime/Engine/Private/StaticMesh.cpp b/Engine/Source/Runtime/Engine/Private/StaticMesh.cpp
--- a/Engine/Source/Runtime/Engine/Private/StaticMesh.cpp
+++ b/Engine/Source/Runtime/Engine/Private/StaticMesh.cpp
@@ -5373,6 +5373,18 @@ bool UStaticMesh::GetMipDataFilename(const int32 MipIndex, FString& OutBulkDataF
 	return true;
 }
 
+bool UStaticMesh::DoesMipDataExist(const int32 MipIndex) const
+{
+	check(MipIndex < MinLOD.Default);
+
+#if !USE_BULKDATA_STREAMING_TOKEN	
+	return RenderData->LODResources[MipIndex].StreamingBulkData.DoesExist();
+#else
+	checkf(false, TEXT("Should not be possible to reach this path, if USE_NEW_BULKDATA is enabled then USE_BULKDATA_STREAMING_TOKEN should be disabled!"));
+	return false;
+#endif
+}
+
 bool UStaticMesh::IsReadyForStreaming() const
 {
 	return RenderData && RenderData->bReadyForStreaming;
diff --git a/Engine/Source/Runtime/Engine/Private/Streaming/StreamingTexture.cpp b/Engine/Source/Runtime/Engine/Private/Streaming/StreamingTexture.cpp
--- a/Engine/Source/Runtime/Engine/Private/Streaming/StreamingTexture.cpp
+++ b/Engine/Source/Runtime/Engine/Private/Streaming/StreamingTexture.cpp
@@ -43,8 +43,12 @@ FStreamingRenderAsset::FStreamingRenderAsset(
 
 void FStreamingRenderAsset::UpdateStaticData(const FRenderAssetStreamingSettings& Settings)
 {
-
+#if STORE_OPTIONAL_DATA_FILENAME
 	OptionalBulkDataFilename = TEXT("");
+#else
+	OptionalMipIndex = INDEX_NONE;
+#endif
+
 	if (RenderAsset)
 	{
 		LODGroup = RenderAsset->GetLODGroupForStreaming();
@@ -101,11 +105,16 @@ void FStreamingRenderAsset::UpdateStaticData(const FRenderAssetStreamingSettings
 		}
 
 		const int32 OptionalMipCount = MipCount - NumNonOptionalMips;
-		const int32 OptionalMipIndex = OptionalMipCount - 1; // just here so it's clear why this -1 is here
-		if (!RenderAsset->GetMipDataFilename(OptionalMipIndex, OptionalBulkDataFilename))
+		const int32 FirstOptionalMipIndex = OptionalMipCount - 1; // just here so it's clear why this -1 is here
+
+#if STORE_OPTIONAL_DATA_FILENAME
+		if (!RenderAsset->GetMipDataFilename(FirstOptionalMipIndex, OptionalBulkDataFilename))
 		{
 			OptionalBulkDataFilename.Empty();
-		}
+		}		
+#else
+		OptionalMipIndex = FirstOptionalMipIndex;
+#endif
 	}
 	else
 	{
@@ -130,13 +139,26 @@ void FStreamingRenderAsset::UpdateStaticData(const FRenderAssetStreamingSettings
 
 void FStreamingRenderAsset::UpdateOptionalMipsState_Async()
 {
+#if STORE_OPTIONAL_DATA_FILENAME
 	// Here we do a lazy update where we check if the highres mip file exists only if it could be useful to do so.
 	// This requires texture to be at max resolution before the optional mips .
 	if (OptionalMipsState == EOptionalMipsState::OMS_NotCached && !OptionalBulkDataFilename.IsEmpty())
 	{
 		OptionalMipsState = IFileManager::Get().FileExists(*OptionalBulkDataFilename) ? EOptionalMipsState::OMS_HasOptionalMips : EOptionalMipsState::OMS_NoOptionalMips;
 	}
-
+#else
+	if (OptionalMipsState == EOptionalMipsState::OMS_NotCached && OptionalMipIndex != INDEX_NONE)
+	{
+		if (RenderAsset->DoesMipDataExist(OptionalMipIndex))
+		{
+			OptionalMipsState = EOptionalMipsState::OMS_HasOptionalMips;
+		}
+		else
+		{
+			OptionalMipsState = EOptionalMipsState::OMS_NoOptionalMips;
+		}
+	}	
+#endif
 }
 
 void FStreamingRenderAsset::UpdateDynamicData(const int32* NumStreamedMips, int32 NumLODGroups, const FRenderAssetStreamingSettings& Settings, bool bWaitForMipFading)
diff --git a/Engine/Source/Runtime/Engine/Private/Streaming/StreamingTexture.h b/Engine/Source/Runtime/Engine/Private/Streaming/StreamingTexture.h
--- a/Engine/Source/Runtime/Engine/Private/Streaming/StreamingTexture.h
+++ b/Engine/Source/Runtime/Engine/Private/Streaming/StreamingTexture.h
@@ -12,6 +12,10 @@ StreamingTexture.h: Definitions of classes used for texture streaming.
 struct FRenderAssetStreamingManager;
 struct FRenderAssetStreamingSettings;
 
+// When enabled we store the filename of optional data as a FString in FStreamingRenderAsset and 
+// use that to check for the existence of the optional data. When disabled we poll the BulkData instead.
+#define STORE_OPTIONAL_DATA_FILENAME !USE_NEW_BULKDATA
+
 /*-----------------------------------------------------------------------------
 	FStreamingRenderAsset, the streaming system's version of UTexture2D.
 -----------------------------------------------------------------------------*/
@@ -208,10 +212,14 @@ struct FStreamingRenderAsset
 	/** (1) Texture/mesh to manage. Note that this becomes null when the texture/mesh is removed. */
 	UStreamableRenderAsset*		RenderAsset;
 	/** (2) */
+#if STORE_OPTIONAL_DATA_FILENAME
 	FString			OptionalBulkDataFilename;
+#else
+	int32			OptionalMipIndex;
+#endif
 	
 	/** (1) Cached texture/mesh LOD group. */
-	int32	LODGroup;
+	int32			LODGroup;
 	/** (1) Cached number of mipmaps that are not allowed to stream. */
 	int32			NumNonStreamingMips;
 	/** (1) Cached number of mip-maps in the asset's mip array (including the base mip) */
diff --git a/Engine/Source/Runtime/Engine/Private/Texture2D.cpp b/Engine/Source/Runtime/Engine/Private/Texture2D.cpp
--- a/Engine/Source/Runtime/Engine/Private/Texture2D.cpp
+++ b/Engine/Source/Runtime/Engine/Private/Texture2D.cpp
@@ -531,6 +531,23 @@ bool UTexture2D::GetMipDataFilename(const int32 MipIndex, FString& OutBulkDataFi
 	return false;
 }
 
+bool UTexture2D::DoesMipDataExist(const int32 MipIndex) const
+{
+#if TEXTURE2DMIPMAP_USE_COMPACT_BULKDATA == 0
+	if (PlatformData)
+	{
+		if (MipIndex < PlatformData->Mips.Num() && MipIndex >= 0)
+		{
+			return PlatformData->Mips[MipIndex].BulkData.DoesExist();
+		}
+	}
+#else
+	checkf(false, TEXT("Should not be possible to reach this path, if USE_NEW_BULKDATA is enabled then TEXTURE2DMIPMAP_USE_COMPACT_BULKDATA should be disabled!"));
+#endif
+
+	return false;
+}
+
 int32 UTexture2D::GetNumResidentMips() const
 {
 	if (IsCurrentlyVirtualTextured())
