IoStore:
Tactical submit to prepare for pakchunking

diff --git a/Engine/Source/Developer/IoStoreUtilities/Private/IoStoreUtilities.cpp b/Engine/Source/Developer/IoStoreUtilities/Private/IoStoreUtilities.cpp
--- a/Engine/Source/Developer/IoStoreUtilities/Private/IoStoreUtilities.cpp
+++ b/Engine/Source/Developer/IoStoreUtilities/Private/IoStoreUtilities.cpp
@@ -24,8 +24,8 @@
 #include "UObject/ObjectResource.h"
 #include "UObject/Package.h"
 #include "UObject/SavePackage.h"
-
 #include "Algo/Find.h"
+
 //PRAGMA_DISABLE_OPTIMIZATION
 
 IMPLEMENT_MODULE(FDefaultModuleImpl, IoStoreUtilities);
@@ -78,7 +78,7 @@ public:
 		}
 	}
 
-	int32 MapName(const FName& Name)
+	int32 MapName(const FName& Name) const
 	{
 		const FNameEntryId Id = Name.GetComparisonIndex();
 		const int32* Index = NameIndices.Find(Id);
@@ -86,7 +86,7 @@ public:
 		return Index ? *Index - 1 : INDEX_NONE;
 	}
 
-	void SerializeName(FArchive& A, const FName& N)
+	void SerializeName(FArchive& A, const FName& N) const
 	{
 		int32 NameIndex = MapName(N);
 		int32 NameNumber = N.GetNumber();
@@ -344,7 +344,7 @@ public:
 		ToNode->ExternalDependencies.Add(FromNode);
 	}
 
-	TArray<FExportGraphNode*> ComputeLoadOrder(TMap<FName, FPackage*>& PackageMap) const;
+	TArray<FExportGraphNode*> ComputeLoadOrder(const TArray<FPackage*>& Packages) const;
 
 private:
 	void AddEdge(FExportGraphNode* FromNode, FExportGraphNode* ToNode)
@@ -368,13 +368,12 @@ struct FPackage
 	int32 NameCount = 0;
 	int32 ImportCount = 0;
 	int32 ImportOffset = 0;
-	int32 ImportedPackagesCount = 0;
-	int32 ImportedPackagesOffset = -1;
 	int32 ExportCount = 0;
-	int32 FirstGlobalImport;
-	int32 GlobalImportCount;
+	int32 FirstGlobalImport = -1;
+	int32 GlobalImportCount = -1;
 	int32 ExportIndexOffset = -1;
 	int32 PreloadIndexOffset = -1;
+	int32 FirstExportBundleMetaEntry = -1;
 	int64 BulkDataStartOffset = -1;
 	int64 UExpSize = 0;
 	int64 UAssetSize = 0;
@@ -523,18 +522,15 @@ TArray<FPackage*> FPackageGraph::TopologicalSort() const
 	return Result;
 }
 
-TArray<FExportGraphNode*> FExportGraph::ComputeLoadOrder(TMap<FName, FPackage*>& PackageMap) const
+TArray<FExportGraphNode*> FExportGraph::ComputeLoadOrder(const TArray<FPackage*>& Packages) const
 {
 	FPackageGraph PackageGraph;
-	for (auto& KV : PackageMap)
+	for (FPackage* Package : Packages)
 	{
-		FPackage* Package = KV.Value;
 		Package->Node = PackageGraph.AddNode(Package);
 	}
-	for (auto& KV : PackageMap)
+	for (FPackage* Package : Packages)
 	{
-		FPackage* Package = KV.Value;
-		
 		for (FPackage* ImportedPackage : Package->ImportedPackages)
 		{
 			PackageGraph.AddImportDependency(ImportedPackage->Node, Package->Node);
@@ -799,10 +795,9 @@ static void AddPostLoadDependencies(
 	*/
 }
 
-static void BuildBundles(FExportGraph& ExportGraph, TMap<FName, FPackage*>& PackageMap)
+static void BuildBundles(FExportGraph& ExportGraph, const TArray<FPackage*>& Packages)
 {
-
-	TArray<FExportGraphNode*> ExportLoadOrder = ExportGraph.ComputeLoadOrder(PackageMap);
+	TArray<FExportGraphNode*> ExportLoadOrder = ExportGraph.ComputeLoadOrder(Packages);
 	FPackage* LastPackage = nullptr;
 	uint32 BundleLoadOrder = 0;
 	for (FExportGraphNode* Node : ExportLoadOrder)
@@ -847,7 +842,7 @@ static void BuildBundles(FExportGraph& ExportGraph, TMap<FName, FPackage*>& Pack
 }
 
 static bool WriteBulkData(	const FString& Filename, EIoChunkType Type, const FPackage& Package, const FPackageStoreBulkDataManifest& BulkDataManifest,
-							FIoStoreWriter* IoStoreWriter, uint64& OutBulkChunkCount, uint64& OutBulkPartialChunkCount)
+							FIoStoreWriter* IoStoreWriter)
 {
 	TUniquePtr<FArchive> BulkAr(IFileManager::Get().CreateFileReader(*Filename));
 	if (BulkAr != nullptr)
@@ -866,8 +861,6 @@ static bool WriteBulkData(	const FString& Filename, EIoChunkType Type, const FPa
 
 		BulkAr->Close();
 #endif
-		++OutBulkChunkCount;
-
 		// Create additional mapping chunks as needed
 		const FPackageStoreBulkDataManifest::PackageDesc* PackageDesc = BulkDataManifest.Find(Package.FileName);
 		if (PackageDesc != nullptr)
@@ -885,8 +878,6 @@ static bool WriteBulkData(	const FString& Filename, EIoChunkType Type, const FPa
 					}
 #endif
 				}
-				
-				++OutBulkPartialChunkCount;
 			}
 		}
 		else
@@ -947,6 +938,12 @@ struct FExportData
 	FExportGraphNode* SerializeNode = nullptr;
 };
 
+struct FExportBundleMeta
+{
+	uint32 LoadOrder = -1;
+	uint32 PayloadSize = -1;
+};
+
 static void FindImport(TArray<FImportData>& GlobalImports, TMap<FString, int32>& GlobalImportsByFullName, TArray<FString>& TempFullNames, FObjectImport* ImportMap, int32 LocalImportIndex)
 {
 	FObjectImport* Import = &ImportMap[LocalImportIndex];
@@ -1051,6 +1048,271 @@ static void FindExport(TArray<FExportData>& GlobalExports, TMap<FString, int32>&
 	}
 };
 
+FPackage* AddPackage(const TCHAR* FileName, const TCHAR* CookedDir, TArray<FPackage*>& Packages, TMap<FName, FPackage*>& PackageMap)
+{
+	FString RelativeFileName = FileName;
+	RelativeFileName.RemoveFromStart(CookedDir);
+	RelativeFileName.RemoveFromStart(TEXT("/"));
+	RelativeFileName = TEXT("../../../") / RelativeFileName;
+
+	FString PackageName;
+	FString ErrorMessage;
+	if (!FPackageName::TryConvertFilenameToLongPackageName(RelativeFileName, PackageName, &ErrorMessage))
+	{
+		UE_LOG(LogIoStore, Warning, TEXT("Failed to convert file name from file name '%s'"), *ErrorMessage);
+		return nullptr;
+	}
+
+	FName PackageFName = *PackageName;
+
+	FPackage* Package = new FPackage();
+	Package->Name = PackageFName;
+	Package->FileName = FileName;
+	Package->RelativeFileName = MoveTemp(RelativeFileName);
+	Package->GlobalPackageId = Packages.Num();
+	Packages.Add(Package);
+	PackageMap.Add(PackageFName, Package);
+
+	return Package;
+}
+
+void SerializePackageData(
+	FIoStoreWriter* IoStoreWriter,
+	const TArray<FPackage*>& Packages,
+	const FNameMapBuilder& NameMapBuilder,
+	const TArray<FObjectExport>& ObjectExports,
+	const TArray<FExportData>& GlobalExports,
+	const TMap<FString, int32>& GlobalImportsByFullName,
+	FExportBundleMeta* ExportBundleMetaEntries,
+	const FPackageStoreBulkDataManifest& BulkDataManifest,
+	bool bWithBulkDataManifest)
+{
+	for (FPackage* Package : Packages)
+	{
+		UE_CLOG(Package->GlobalPackageId % 1000 == 0, LogIoStore, Display, TEXT("Serializing %d/%d: '%s'"), Package->GlobalPackageId, Packages.Num(), *Package->Name.ToString());
+
+		// Temporary Archive for serializing ImportMap
+		FBufferWriter ImportMapArchive(nullptr, 0, EBufferWriterFlags::AllowResize | EBufferWriterFlags::TakeOwnership);
+		for (int32 GlobalImportIndex : Package->Imports)
+		{
+			ImportMapArchive << GlobalImportIndex;
+		}
+		Package->ImportMapSize = ImportMapArchive.Tell();
+
+		// Temporary Archive for serializing EDL graph data
+		FBufferWriter ZenGraphArchive(nullptr, 0, EBufferWriterFlags::AllowResize | EBufferWriterFlags::TakeOwnership);
+
+		int32 InternalArcCount = Package->InternalArcs.Num();
+		ZenGraphArchive << InternalArcCount;
+		for (FArc& InternalArc : Package->InternalArcs)
+		{
+			ZenGraphArchive << InternalArc.FromNodeIndex;
+			ZenGraphArchive << InternalArc.ToNodeIndex;
+		}
+
+		int32 ReferencedPackagesCount = Package->ExternalArcs.Num();
+		ZenGraphArchive << ReferencedPackagesCount;
+		for (auto& KV : Package->ExternalArcs)
+		{
+			FPackage* ImportedPackage = KV.Key;
+			TArray<FArc>& Arcs = KV.Value;
+			int32 ExternalArcCount = Arcs.Num();
+
+			ZenGraphArchive << ImportedPackage->GlobalPackageId;
+			ZenGraphArchive << ExternalArcCount;
+			for (FArc& ExternalArc : Arcs)
+			{
+				ZenGraphArchive << ExternalArc.FromNodeIndex;
+				ZenGraphArchive << ExternalArc.ToNodeIndex;
+			}
+		}
+		Package->UGraphSize = ZenGraphArchive.Tell();
+
+		// Temporary Archive for serializing export map data
+		FBufferWriter ZenExportMapArchive(nullptr, 0, EBufferWriterFlags::AllowResize | EBufferWriterFlags::TakeOwnership);
+		for (int32 I = 0; I < Package->ExportCount; ++I)
+		{
+			const FObjectExport& ObjectExport = ObjectExports[Package->ExportIndexOffset + I];
+			const FExportData& ExportData = GlobalExports[Package->Exports[I]];
+
+			int64 SerialSize = ObjectExport.SerialSize;
+			ZenExportMapArchive << SerialSize;
+			NameMapBuilder.SerializeName(ZenExportMapArchive, ObjectExport.ObjectName);
+			FPackageIndex OuterIndex = ExportData.OuterIndex;
+			ZenExportMapArchive << OuterIndex;
+			FPackageIndex ClassIndex = ExportData.ClassIndex;
+			ZenExportMapArchive << ClassIndex;
+			FPackageIndex SuperIndex = ExportData.SuperIndex;
+			ZenExportMapArchive << SuperIndex;
+			FPackageIndex TemplateIndex = ExportData.TemplateIndex;
+			ZenExportMapArchive << TemplateIndex;
+			int32 GlobalImportIndex = ExportData.GlobalImportIndex;
+			ZenExportMapArchive << GlobalImportIndex;
+			uint32 ObjectFlags = ObjectExport.ObjectFlags;
+			ZenExportMapArchive << ObjectFlags;
+		}
+		Package->ExportMapSize = ZenExportMapArchive.Tell();
+
+		// Temporary archive for serializing export bundle data
+		FBufferWriter ZenExportBundlesArchive(nullptr, 0, EBufferWriterFlags::AllowResize | EBufferWriterFlags::TakeOwnership);
+		int32 ExportBundleEntryIndex = 0;
+		for (FExportBundle& ExportBundle : Package->ExportBundles)
+		{
+			ZenExportBundlesArchive << ExportBundleEntryIndex;
+			int32 EntryCount = ExportBundle.Nodes.Num();
+			ZenExportBundlesArchive << EntryCount;
+			ExportBundleEntryIndex += ExportBundle.Nodes.Num();
+		}
+		for (FExportBundle& ExportBundle : Package->ExportBundles)
+		{
+			for (FExportGraphNode* ExportNode : ExportBundle.Nodes)
+			{
+				uint32 CommandType = uint32(ExportNode->BundleEntry.CommandType);
+				ZenExportBundlesArchive << CommandType;
+				ZenExportBundlesArchive << ExportNode->BundleEntry.LocalExportIndex;
+			}
+		}
+		Package->ExportBundlesSize = ZenExportBundlesArchive.Tell();
+
+		Package->NameMapSize = Package->NameIndices.Num() * Package->NameIndices.GetTypeSize();
+
+		{
+			const uint64 ZenSummarySize =
+				sizeof(FZenPackageSummary)
+				+ Package->NameMapSize
+				+ Package->ImportMapSize
+				+ Package->ExportMapSize
+				+ Package->ExportBundlesSize
+				+ Package->UGraphSize;
+
+			uint8* ZenSummaryBuffer = static_cast<uint8*>(FMemory::Malloc(ZenSummarySize));
+			FZenPackageSummary* ZenSummary = reinterpret_cast<FZenPackageSummary*>(ZenSummaryBuffer);
+
+			ZenSummary->PackageFlags = Package->PackageFlags;
+			ZenSummary->GraphDataSize = Package->UGraphSize;
+			ZenSummary->BulkDataStartOffset = Package->BulkDataStartOffset;
+			const int32* FindGlobalImportIndexForPackage = GlobalImportsByFullName.Find(Package->Name.ToString());
+			ZenSummary->GlobalImportIndex = FindGlobalImportIndexForPackage ? *FindGlobalImportIndexForPackage : -1;
+
+			FBufferWriter ZenAr(ZenSummaryBuffer, ZenSummarySize);
+			ZenAr.Seek(sizeof(FZenPackageSummary));
+
+			// NameMap data
+			{
+				ZenSummary->NameMapOffset = ZenAr.Tell();
+				ZenAr.Serialize(Package->NameIndices.GetData(), Package->NameMapSize);
+			}
+
+			// ImportMap data
+			{
+				check(ImportMapArchive.Tell() == Package->ImportMapSize);
+				ZenSummary->ImportMapOffset = ZenAr.Tell();
+				ZenAr.Serialize(ImportMapArchive.GetWriterData(), ImportMapArchive.Tell());
+			}
+
+			// ExportMap data
+			{
+				check(ZenExportMapArchive.Tell() == Package->ExportMapSize);
+				ZenSummary->ExportMapOffset = ZenAr.Tell();
+				ZenAr.Serialize(ZenExportMapArchive.GetWriterData(), ZenExportMapArchive.Tell());
+			}
+
+			// ExportBundle data
+			{
+				check(ZenExportBundlesArchive.Tell() == Package->ExportBundlesSize);
+				ZenSummary->ExportBundlesOffset = ZenAr.Tell();
+				ZenAr.Serialize(ZenExportBundlesArchive.GetWriterData(), ZenExportBundlesArchive.Tell());
+			}
+
+			// Graph data
+			{
+				check(ZenGraphArchive.Tell() == Package->UGraphSize);
+				ZenSummary->GraphDataOffset = ZenAr.Tell();
+				ZenAr.Serialize(ZenGraphArchive.GetWriterData(), ZenGraphArchive.Tell());
+			}
+
+			// Export bundle chunks
+			{
+				check(Package->ExportBundles.Num());
+
+				FString UExpFileName = FPaths::ChangeExtension(Package->FileName, TEXT(".uexp"));
+				TUniquePtr<FArchive> ExpAr(IFileManager::Get().CreateFileReader(*UExpFileName));
+				Package->UExpSize = ExpAr->TotalSize();
+#if !SKIP_WRITE_CONTAINER
+				uint8* ExportsBuffer = static_cast<uint8*>(FMemory::Malloc(ExpAr->TotalSize()));
+				ExpAr->Serialize(ExportsBuffer, ExpAr->TotalSize());
+#endif
+				ExpAr->Close();
+
+				uint64 BundleBufferSize = ZenSummarySize;
+				for (int32 ExportBundleIndex = 0; ExportBundleIndex < Package->ExportBundles.Num(); ++ExportBundleIndex)
+				{
+					FExportBundle& ExportBundle = Package->ExportBundles[ExportBundleIndex];
+					FExportBundleMeta* ExportBundleMetaEntry = ExportBundleMetaEntries + Package->FirstExportBundleMetaEntry + ExportBundleIndex;
+					ExportBundleMetaEntry->LoadOrder = ExportBundle.LoadOrder;
+					for (FExportGraphNode* Node : ExportBundle.Nodes)
+					{
+						if (Node->BundleEntry.CommandType == FExportBundleEntry::ExportCommandType_Serialize)
+						{
+							const FObjectExport& ObjectExport = ObjectExports[Package->ExportIndexOffset + Node->BundleEntry.LocalExportIndex];
+							BundleBufferSize += ObjectExport.SerialSize;
+						}
+					}
+					if (ExportBundleIndex == 0)
+					{
+						ExportBundleMetaEntry->PayloadSize = BundleBufferSize;
+					}
+				}
+
+#if !SKIP_WRITE_CONTAINER
+				uint8* BundleBuffer = static_cast<uint8*>(FMemory::Malloc(BundleBufferSize));
+				FMemory::Memcpy(BundleBuffer, ZenSummaryBuffer, ZenSummarySize);
+#endif
+				FMemory::Free(ZenSummaryBuffer);
+				uint64 BundleBufferOffset = ZenSummarySize;
+				for (FExportBundle& ExportBundle : Package->ExportBundles)
+				{
+					for (FExportGraphNode* Node : ExportBundle.Nodes)
+					{
+						if (Node->BundleEntry.CommandType == FExportBundleEntry::ExportCommandType_Serialize)
+						{
+							const FObjectExport& ObjectExport = ObjectExports[Package->ExportIndexOffset + Node->BundleEntry.LocalExportIndex];
+							const int64 Offset = ObjectExport.SerialOffset - Package->UAssetSize;
+#if !SKIP_WRITE_CONTAINER
+							FMemory::Memcpy(BundleBuffer + BundleBufferOffset, ExportsBuffer + Offset, ObjectExport.SerialSize);
+#endif
+							BundleBufferOffset += ObjectExport.SerialSize;
+						}
+					}
+				}
+
+#if !SKIP_WRITE_CONTAINER
+				FIoBuffer IoBuffer(FIoBuffer::Wrap, BundleBuffer, BundleBufferSize);
+				IoStoreWriter->Append(CreateChunkId(Package->GlobalPackageId, 0, EIoChunkType::ExportBundleData, *Package->FileName), IoBuffer);
+				FMemory::Free(BundleBuffer);
+				FMemory::Free(ExportsBuffer);
+#endif
+			}
+
+#if !SKIP_BULKDATA
+			// Bulk chunks
+			if (bWithBulkDataManifest)
+			{
+				FString BulkFileName = FPaths::ChangeExtension(Package->FileName, TEXT(".ubulk"));
+				FPaths::NormalizeFilename(BulkFileName);
+
+				WriteBulkData(BulkFileName, EIoChunkType::BulkData, *Package, BulkDataManifest, IoStoreWriter);
+
+				FString OptionalBulkFileName = FPaths::ChangeExtension(Package->FileName, TEXT(".uptnl"));
+				FPaths::NormalizeFilename(OptionalBulkFileName);
+
+				WriteBulkData(OptionalBulkFileName, EIoChunkType::OptionalBulkData, *Package, BulkDataManifest, IoStoreWriter);
+			}
+#endif
+		}
+	}
+}
+
 int32 CreateTarget(const FContainerTarget& Target)
 {
 	TGuardValue<int32> GuardAllowUnversionedContentInEditor(GAllowUnversionedContentInEditor, 1);
@@ -1069,17 +1331,11 @@ int32 CreateTarget(const FContainerTarget& Target)
 	ChunkIdCsv.CreateOutputFile(CookedDir);
 #endif
 
-	TArray<FString> FileNames;
-	UE_LOG(LogIoStore, Display, TEXT("Searching for .uasset and .umap files..."));
-	IFileManager::Get().FindFilesRecursive(FileNames, *CookedDir, TEXT("*.uasset"), true, false, false);
-	IFileManager::Get().FindFilesRecursive(FileNames, *CookedDir, TEXT("*.umap"), true, false, false);
-	UE_LOG(LogIoStore, Display, TEXT("Found '%d' files"), FileNames.Num());
-
 	FNameMapBuilder NameMapBuilder;
 
 	uint64 NameSize = 0;
-	TArray<FObjectImport> Imports;
-	TArray<FObjectExport> Exports;
+	TArray<FObjectImport> ObjectImports;
+	TArray<FObjectExport> ObjectExports;
 	TArray<FImportData> GlobalImports;
 	TArray<FExportData> GlobalExports;
 	TMap<FString, int32> GlobalImportsByFullName;
@@ -1087,84 +1343,54 @@ int32 CreateTarget(const FContainerTarget& Target)
 	TArray<FString> TempFullNames;
 	TArray<FPackageIndex> PreloadDependencies;
 	uint64 UPackageImports = 0;
-	TArray<FPackageFileSummary> Summaries;
 	TArray<int32> ImportPreloadCounts;
 	TArray<int32> ExportPreloadCounts;
-	Summaries.AddDefaulted(FileNames.Num());
-	ImportPreloadCounts.AddDefaulted(FileNames.Num());
-	ExportPreloadCounts.AddDefaulted(FileNames.Num());
 	uint64 ImportPreloadCount = 0;
 	uint64 ExportPreloadCount = 0;
-
-	TUniquePtr<FLargeMemoryWriter> StoreTocArchive = MakeUnique<FLargeMemoryWriter>(0, true);
-	TUniquePtr<FLargeMemoryWriter> ImportedPackagesArchive = MakeUnique<FLargeMemoryWriter>(0, true);
-	TUniquePtr<FLargeMemoryWriter> GlobalImportNamesArchive = MakeUnique<FLargeMemoryWriter>(0, true);
-	TUniquePtr<FLargeMemoryWriter> InitialLoadArchive = MakeUnique<FLargeMemoryWriter>(0, true);
-
-	FIoStoreEnvironment IoStoreEnv;
-#if !SKIP_WRITE_CONTAINER
-	TUniquePtr<FIoStoreWriter> IoStoreWriter;
-
-	{
-		IoStoreEnv.InitializeFileEnvironment(OutputDir);
-		IoStoreWriter = MakeUnique<FIoStoreWriter>(IoStoreEnv);
-		FIoStatus IoStatus = IoStoreWriter->Initialize();
-		check(IoStatus.IsOk());
-	}
-#endif
-
-	TMap<FName, FPackage*> PackageMap;
 	FExportGraph ExportGraph;
+	TArray<FExportBundleMeta> ExportBundleMetaEntries;
 
-	for (int FileIndex = 0; FileIndex < FileNames.Num(); ++FileIndex)
+	TArray<FPackage*> Packages;
+	TMap<FName, FPackage*> PackageMap;
 	{
-		const FString& FileName = FileNames[FileIndex];
-		FPackageFileSummary& Summary = Summaries[FileIndex];
-		TUniquePtr<FArchive> Ar(IFileManager::Get().CreateFileReader(*FileName));
-		check(Ar);
-
-		UE_CLOG(FileIndex % 1000 == 0, LogIoStore, Display, TEXT("Parsing %d/%d: '%s'"), FileIndex, FileNames.Num(), *FileName);
-
-		FString RelativeFileName = FileName;
-		RelativeFileName.RemoveFromStart(*CookedDir);
-		RelativeFileName.RemoveFromStart(TEXT("/"));
-		RelativeFileName = TEXT("../../../") / RelativeFileName;
-
-		FString PackageName;
-		FString ErrorMessage;
-		if (!FPackageName::TryConvertFilenameToLongPackageName(RelativeFileName, PackageName, &ErrorMessage))
+		UE_LOG(LogIoStore, Display, TEXT("Searching for .uasset and .umap files..."));
+		TArray<FString> FileNames;
+		IFileManager::Get().FindFilesRecursive(FileNames, *CookedDir, TEXT("*.uasset"), true, false, false);
+		IFileManager::Get().FindFilesRecursive(FileNames, *CookedDir, TEXT("*.umap"), true, false, false);
+		UE_LOG(LogIoStore, Display, TEXT("Found '%d' files"), FileNames.Num());
+		for (FString& FileName : FileNames)
 		{
-			UE_LOG(LogIoStore, Error, TEXT("Failed to convert file name from file name '%s'"), *ErrorMessage);
-			return -1;
+			AddPackage(*FileName, *CookedDir, Packages, PackageMap);
 		}
+	}
 
-		FName PackageFName = *PackageName;
+	ImportPreloadCounts.AddDefaulted(Packages.Num());
+	ExportPreloadCounts.AddDefaulted(Packages.Num());
 
-		uint64 SummaryStartPos = Ar->Tell();
+	for (FPackage* Package : Packages)
+	{
+		UE_CLOG(Package->GlobalPackageId % 1000 == 0, LogIoStore, Display, TEXT("Parsing %d/%d: '%s'"), Package->GlobalPackageId, Packages.Num(), *Package->FileName);
+
+		FPackageFileSummary Summary;
+		TUniquePtr<FArchive> Ar(IFileManager::Get().CreateFileReader(*Package->FileName));
+		check(Ar);
 		*Ar << Summary;
 
-		FPackage* NewPackage = new FPackage();
-		PackageMap.Add(PackageFName, NewPackage);
-		FPackage& Package = *NewPackage;
-		Package.Name = PackageFName;
-		Package.FileName = FileName;
-		Package.RelativeFileName = MoveTemp(RelativeFileName);
-		Package.GlobalPackageId = FileIndex;
-		Package.UAssetSize = Ar->TotalSize();
-		Package.SummarySize = Ar->Tell() - SummaryStartPos;
-		Package.NameCount = Summary.NameCount;
-		Package.ImportCount = Summary.ImportCount;
-		Package.ExportCount = Summary.ExportCount;
-		Package.PackageFlags = Summary.PackageFlags;
-		Package.BulkDataStartOffset = Summary.BulkDataStartOffset;
+		Package->UAssetSize = Ar->TotalSize();
+		Package->SummarySize = Ar->Tell();
+		Package->NameCount = Summary.NameCount;
+		Package->ImportCount = Summary.ImportCount;
+		Package->ExportCount = Summary.ExportCount;
+		Package->PackageFlags = Summary.PackageFlags;
+		Package->BulkDataStartOffset = Summary.BulkDataStartOffset;
 
 		if (Summary.NameCount > 0)
 		{
 			Ar->Seek(Summary.NameOffset);
 			uint64 LastOffset = Summary.NameOffset;
 
-			Package.NameMap.Reserve(Summary.NameCount);
-			Package.NameIndices.Reserve(Summary.NameCount);
+			Package->NameMap.Reserve(Summary.NameCount);
+			Package->NameIndices.Reserve(Summary.NameCount);
 			FNameEntrySerialized NameEntry(ENAME_LinkerConstructor);
 
 			for (int32 I = 0; I < Summary.NameCount; ++I)
@@ -1172,8 +1398,8 @@ int32 CreateTarget(const FContainerTarget& Target)
 				*Ar << NameEntry;
 				FName Name(NameEntry);
 				NameMapBuilder.MarkNameAsReferenced(Name);
-				Package.NameMap.Emplace(Name.GetDisplayIndex());
-				Package.NameIndices.Add(NameMapBuilder.MapName(Name));
+				Package->NameMap.Emplace(Name.GetDisplayIndex());
+				Package->NameIndices.Add(NameMapBuilder.MapName(Name));
 			}
 
 			NameSize += Ar->Tell() - Summary.NameOffset;
@@ -1183,7 +1409,7 @@ int32 CreateTarget(const FContainerTarget& Target)
 		{
 			int32 DisplayIndex, NameNumber;
 			A << DisplayIndex << NameNumber;
-			FNameEntryId DisplayEntry = Package.NameMap[DisplayIndex];
+			FNameEntryId DisplayEntry = Package->NameMap[DisplayIndex];
 			N = FName::CreateFromDisplayId(DisplayEntry, NameNumber);
 		};
 
@@ -1192,13 +1418,13 @@ int32 CreateTarget(const FContainerTarget& Target)
 			Ar->Seek(Summary.ImportOffset);
 
 			int32 NumPackages = 0;
-			int32 BaseIndex = Imports.Num();
-			Imports.AddUninitialized(Summary.ImportCount);
+			int32 BaseIndex = ObjectImports.Num();
+			ObjectImports.AddUninitialized(Summary.ImportCount);
 			TArray<FString> ImportNames;
 			ImportNames.Reserve(Summary.ImportCount);
 			for (int32 I = 0; I < Summary.ImportCount; ++I)
 			{
-				FObjectImport& ObjectImport = Imports[BaseIndex + I];
+				FObjectImport& ObjectImport = ObjectImports[BaseIndex + I];
 				DeserializeName(*Ar, ObjectImport.ClassPackage);
 				DeserializeName(*Ar, ObjectImport.ClassName);
 				*Ar << ObjectImport.OuterIndex;
@@ -1214,14 +1440,14 @@ int32 CreateTarget(const FContainerTarget& Target)
 
 			UPackageImports += NumPackages;
 
-			Package.ImportedFullNames.SetNum(Summary.ImportCount);
+			Package->ImportedFullNames.SetNum(Summary.ImportCount);
 			for (int32 I = 0; I < Summary.ImportCount; ++I)
 			{
-				FindImport(GlobalImports, GlobalImportsByFullName, Package.ImportedFullNames, Imports.GetData() + BaseIndex, I);
+				FindImport(GlobalImports, GlobalImportsByFullName, Package->ImportedFullNames, ObjectImports.GetData() + BaseIndex, I);
 			}
 		}
 
-		Package.PreloadIndexOffset = PreloadDependencies.Num();
+		Package->PreloadIndexOffset = PreloadDependencies.Num();
 		int32 PreloadDependenciesBaseIndex = -1;
 		if (Summary.PreloadDependencyCount > 0)
 		{
@@ -1234,27 +1460,27 @@ int32 CreateTarget(const FContainerTarget& Target)
 				*Ar << Index;
 				if (Index.IsImport())
 				{
-					++ImportPreloadCounts[FileIndex];
+					++ImportPreloadCounts[Package->GlobalPackageId];
 					++ImportPreloadCount;
 				}
 				else
 				{
-					++ExportPreloadCounts[FileIndex];
+					++ExportPreloadCounts[Package->GlobalPackageId];
 					++ExportPreloadCount;
 				}
 			}
 		}
 
-		Package.ExportIndexOffset = Exports.Num();
+		Package->ExportIndexOffset = ObjectExports.Num();
 		if (Summary.ExportCount > 0)
 		{
 			Ar->Seek(Summary.ExportOffset);
 
-			int32 BaseIndex = Exports.Num();
-			Exports.AddUninitialized(Summary.ExportCount);
+			int32 BaseIndex = ObjectExports.Num();
+			ObjectExports.AddUninitialized(Summary.ExportCount);
 			for (int32 I = 0; I < Summary.ExportCount; ++I)
 			{
-				FObjectExport& ObjectExport = Exports[BaseIndex + I];
+				FObjectExport& ObjectExport = ObjectExports[BaseIndex + I];
 				*Ar << ObjectExport.ClassIndex;
 				*Ar << ObjectExport.SuperIndex;
 				*Ar << ObjectExport.TemplateIndex;
@@ -1283,14 +1509,14 @@ int32 CreateTarget(const FContainerTarget& Target)
 			TempFullNames.SetNum(Summary.ExportCount, false);
 			for (int32 I = 0; I < Summary.ExportCount; ++I)
 			{
-				FindExport(GlobalExports, GlobalExportsByFullName, TempFullNames, Exports.GetData() + BaseIndex, I, Package.Name);
+				FindExport(GlobalExports, GlobalExportsByFullName, TempFullNames, ObjectExports.GetData() + BaseIndex, I, Package->Name);
 
 				FExportData& ExportData = GlobalExports[*GlobalExportsByFullName.Find(TempFullNames[I])];
-				Package.Exports.Add(ExportData.GlobalIndex);
-				ExportData.CreateNode = ExportGraph.AddNode(&Package, { uint32(I), FExportBundleEntry::ExportCommandType_Create });
-				ExportData.SerializeNode = ExportGraph.AddNode(&Package, { uint32(I), FExportBundleEntry::ExportCommandType_Serialize });
-				Package.CreateExportNodes.Add(ExportData.CreateNode);
-				Package.SerializeExportNodes.Add(ExportData.SerializeNode);
+				Package->Exports.Add(ExportData.GlobalIndex);
+				ExportData.CreateNode = ExportGraph.AddNode(Package, { uint32(I), FExportBundleEntry::ExportCommandType_Create });
+				ExportData.SerializeNode = ExportGraph.AddNode(Package, { uint32(I), FExportBundleEntry::ExportCommandType_Serialize });
+				Package->CreateExportNodes.Add(ExportData.CreateNode);
+				Package->SerializeExportNodes.Add(ExportData.SerializeNode);
 				ExportGraph.AddInternalDependency(ExportData.CreateNode, ExportData.SerializeNode);
 			}
 		}
@@ -1364,10 +1590,6 @@ int32 CreateTarget(const FContainerTarget& Target)
 		GlobalExport.GlobalImportIndex = FindGlobalImport ? *FindGlobalImport : -1;
 	}
 
-	uint64 ImportSize = Imports.Num() * 28;
-	uint64 ExportSize = Exports.Num() * 104;
-	uint64 PreloadDependenciesSize = PreloadDependencies.Num() * 4;
-
 #if OUTPUT_NAMEMAP_CSV
 	FString CsvFilePath = OutputDir / TEXT("AllImports.csv");
 	TUniquePtr<FArchive> CsvArchive(IFileManager::Get().CreateFileWriter(*CsvFilePath));
@@ -1392,58 +1614,49 @@ int32 CreateTarget(const FContainerTarget& Target)
 	}
 #endif
 
-	uint64 SummaryChunkCount = 0;
-	uint64 ExportChunkCount = 0;
-	uint64 BulkChunkCount = 0;
-	uint64 BulkPartialChunkCount = 0;
-
 	TSet<FString> MissingExports;
 	TSet<FPackage*> Visited;
 	TSet<FCircularImportChain> CircularChains;
-	int32 PackageIndex = 0;
 
 	// Lookup global indices and package pointers for all imports before adding preload and postload arcs
 	UE_LOG(LogIoStore, Display, TEXT("Looking up import packages..."));
-	for (auto& PackageKV : PackageMap)
+	for (FPackage* Package : Packages)
 	{
-		FPackage& Package = *PackageKV.Value;
-
-		Package.ImportedPackages.Reserve(Package.ImportCount);
-		for (int32 I = 0; I < Package.ImportCount; ++I)
+		Package->ImportedPackages.Reserve(Package->ImportCount);
+		for (int32 I = 0; I < Package->ImportCount; ++I)
 		{
-			int32 GlobalImportIndex = *GlobalImportsByFullName.Find(Package.ImportedFullNames[I]);
+			int32 GlobalImportIndex = *GlobalImportsByFullName.Find(Package->ImportedFullNames[I]);
 
 			FImportData& ImportData = GlobalImports[GlobalImportIndex];
-			Package.Imports.Add(ImportData.GlobalIndex);
+			Package->Imports.Add(ImportData.GlobalIndex);
 			if (ImportData.Package)
 			{
-				Package.ImportedPackages.Add(ImportData.Package);
+				Package->ImportedPackages.Add(ImportData.Package);
 			}
 		}
 	}
 
 	UE_LOG(LogIoStore, Display, TEXT("Converting export map import indices..."));
-	for (auto& PackageKV : PackageMap)
+	for (FPackage* Package : Packages)
 	{
-		FPackage& Package = *PackageKV.Value;
-		for (int32 I = 0; I < Package.ExportCount; ++I)
+		for (int32 I = 0; I < Package->ExportCount; ++I)
 		{
-			FObjectExport& ObjectExport = Exports[Package.ExportIndexOffset + I];
-			FExportData& ExportData = GlobalExports[Package.Exports[I]];
+			FObjectExport& ObjectExport = ObjectExports[Package->ExportIndexOffset + I];
+			FExportData& ExportData = GlobalExports[Package->Exports[I]];
 
 			check(!ObjectExport.OuterIndex.IsImport());
 			ExportData.OuterIndex = ObjectExport.OuterIndex;
 			ExportData.ClassIndex =
 				ObjectExport.ClassIndex.IsImport() ?
-				FPackageIndex::FromImport(Package.Imports[ObjectExport.ClassIndex.ToImport()]) :
+				FPackageIndex::FromImport(Package->Imports[ObjectExport.ClassIndex.ToImport()]) :
 				ObjectExport.ClassIndex;
 			ExportData.SuperIndex =
 				ObjectExport.SuperIndex.IsImport() ?
-				FPackageIndex::FromImport(Package.Imports[ObjectExport.SuperIndex.ToImport()]) :
+				FPackageIndex::FromImport(Package->Imports[ObjectExport.SuperIndex.ToImport()]) :
 				ObjectExport.SuperIndex;
 			ExportData.TemplateIndex =
 				ObjectExport.TemplateIndex.IsImport() ?
-				FPackageIndex::FromImport(Package.Imports[ObjectExport.TemplateIndex.ToImport()]) :
+				FPackageIndex::FromImport(Package->Imports[ObjectExport.TemplateIndex.ToImport()]) :
 				ObjectExport.TemplateIndex;
 		}
 	}
@@ -1477,26 +1690,22 @@ int32 CreateTarget(const FContainerTarget& Target)
 #endif
 
 	UE_LOG(LogIoStore, Display, TEXT("Adding optimized postload dependencies..."));
-	for (auto& PackageKV : PackageMap)
+	for (FPackage* Package : Packages)
 	{
-		FPackage& Package = *PackageKV.Value;
-
 		Visited.Reset();
-		Visited.Add(&Package);
-		AddPostLoadDependencies(Package, Visited, CircularChains);
+		Visited.Add(Package);
+		AddPostLoadDependencies(*Package, Visited, CircularChains);
 	}
 		
 	UE_LOG(LogIoStore, Display, TEXT("Adding preload dependencies..."));
-	for (auto& PackageKV : PackageMap)
+	for (FPackage* Package : Packages)
 	{
-		FPackage& Package = *PackageKV.Value;
-
 		// Convert PreloadDependencies to arcs
-		for (int32 I = 0; I < Package.ExportCount; ++I)
+		for (int32 I = 0; I < Package->ExportCount; ++I)
 		{
-			FObjectExport& ObjectExport = Exports[Package.ExportIndexOffset + I];
-			FExportData& ExportData = GlobalExports[Package.Exports[I]];
-			int32 PreloadDependenciesBaseIndex = Package.PreloadIndexOffset;
+			FObjectExport& ObjectExport = ObjectExports[Package->ExportIndexOffset + I];
+			FExportData& ExportData = GlobalExports[Package->Exports[I]];
+			int32 PreloadDependenciesBaseIndex = Package->PreloadIndexOffset;
 
 			FPackageIndex ExportPackageIndex = FPackageIndex::FromExport(I);
 
@@ -1504,16 +1713,16 @@ int32 CreateTarget(const FContainerTarget& Target)
 			{
 				if (Dep.IsExport())
 				{
-					AddInternalExportArc(ExportGraph, Package, Dep.ToExport(), PhaseFrom, I, PhaseTo);
+					AddInternalExportArc(ExportGraph, *Package, Dep.ToExport(), PhaseFrom, I, PhaseTo);
 				}
 				else
 				{
-					FImportData& Import = GlobalImports[Package.Imports[Dep.ToImport()]];
+					FImportData& Import = GlobalImports[Package->Imports[Dep.ToImport()]];
 					check(!Import.bIsPackage);
 					if (Import.bIsScript)
 					{
 						// Add script arc with null package and global import index as node index
-						AddScriptArc(Package, Import.GlobalIndex, I, PhaseTo);
+						AddScriptArc(*Package, Import.GlobalIndex, I, PhaseTo);
 					}
 					else
 					{
@@ -1522,8 +1731,8 @@ int32 CreateTarget(const FContainerTarget& Target)
 						FPackage* SourcePackage = PackageMap.FindRef(Export.SourcePackageName);
 						check(SourcePackage);
 						FPackageIndex SourceDep = FPackageIndex::FromExport(Export.SourceIndex);
-						AddExternalExportArc(ExportGraph, *SourcePackage, Export.SourceIndex, PhaseFrom, Package, I, PhaseTo);
-						Package.ImportedPreloadPackages.Add(SourcePackage);
+						AddExternalExportArc(ExportGraph, *SourcePackage, Export.SourceIndex, PhaseFrom, *Package, I, PhaseTo);
+						Package->ImportedPreloadPackages.Add(SourcePackage);
 					}
 				}
 			};
@@ -1564,8 +1773,13 @@ int32 CreateTarget(const FContainerTarget& Target)
 	}
 
 	UE_LOG(LogIoStore, Display, TEXT("Building bundles..."));
-	BuildBundles(ExportGraph, PackageMap);
+	BuildBundles(ExportGraph, Packages);
+
 
+	TUniquePtr<FLargeMemoryWriter> StoreTocArchive = MakeUnique<FLargeMemoryWriter>(0, true);
+	TUniquePtr<FLargeMemoryWriter> ImportedPackagesArchive = MakeUnique<FLargeMemoryWriter>(0, true);
+	TUniquePtr<FLargeMemoryWriter> GlobalImportNamesArchive = MakeUnique<FLargeMemoryWriter>(0, true);
+	TUniquePtr<FLargeMemoryWriter> InitialLoadArchive = MakeUnique<FLargeMemoryWriter>(0, true);
 
 	UE_LOG(LogIoStore, Display, TEXT("Serializing global import names..."));
 	if (GlobalImportNamesArchive)
@@ -1577,7 +1791,6 @@ int32 CreateTarget(const FContainerTarget& Target)
 	}
 
 	UE_LOG(LogIoStore, Display, TEXT("Serializing initial load..."));
-#if 1
 	// Separate file for script arcs that are only required during initial loading
 	if (InitialLoadArchive)
 	{
@@ -1617,275 +1830,55 @@ int32 CreateTarget(const FContainerTarget& Target)
 
 		InitialLoadArchive->Serialize(ScriptArcsArchive.GetWriterData(), ScriptArcsArchive.Tell());
 	}
-#endif
-
-	struct FExportBundleMeta
-	{
-		uint32 LoadOrder = -1;
-		uint32 PayloadSize = -1;
-	};
 
-	UE_LOG(LogIoStore, Display, TEXT("Serializing..."));
-	TArray<FExportBundleMeta> ExportBundleMetaEntries;
-	int32 LogIndex = 0;
-	for (auto& PackageKV : PackageMap)
+	for (FPackage* Package : Packages)
 	{
-		FPackage& Package = *PackageKV.Value;
+		Package->FirstExportBundleMetaEntry = ExportBundleMetaEntries.Num();
+		ExportBundleMetaEntries.AddDefaulted(Package->ExportBundles.Num());
 
-		UE_CLOG(LogIndex++ % 1000 == 0, LogIoStore, Display, TEXT("Serializing %d/%d: '%s'"), LogIndex, PackageMap.Num(), *Package.Name.ToString());
+		NameMapBuilder.MarkNameAsReferenced(Package->Name);
+		NameMapBuilder.SerializeName(*StoreTocArchive, Package->Name);
+		*StoreTocArchive << Package->ExportCount;
+		int32 ExportBundleCount = Package->ExportBundles.Num();
+		*StoreTocArchive << ExportBundleCount;
+		*StoreTocArchive << Package->FirstExportBundleMetaEntry;
+		*StoreTocArchive << Package->FirstGlobalImport;
+		*StoreTocArchive << Package->GlobalImportCount;
+		int32 ImportedPackagesCount = Package->ImportedPackages.Num();
+		*StoreTocArchive << ImportedPackagesCount;
+		int32 ImportedPackagesOffset = ImportedPackagesArchive->Tell();
+		*StoreTocArchive << ImportedPackagesOffset;
 
-		// ImportedPackages
-		Package.ImportedPackagesCount = Package.ImportedPackages.Num();
-		Package.ImportedPackagesOffset = ImportedPackagesArchive->Tell();
-		for (FPackage* ImportedPackage : Package.ImportedPackages)
+		for (FPackage* ImportedPackage : Package->ImportedPackages)
 		{
 			*ImportedPackagesArchive << ImportedPackage->GlobalPackageId;
 		}
-
-		// Temporary Archive for serializing ImportMap
-		FBufferWriter ImportMapArchive(nullptr, 0, EBufferWriterFlags::AllowResize | EBufferWriterFlags::TakeOwnership );
-		for (int32 GlobalImportIndex : Package.Imports)
-		{
-			ImportMapArchive << GlobalImportIndex;
-		}
-		Package.ImportMapSize = ImportMapArchive.Tell();
-
-		// Temporary Archive for serializing EDL graph data
-		FBufferWriter ZenGraphArchive(nullptr, 0, EBufferWriterFlags::AllowResize | EBufferWriterFlags::TakeOwnership );
-
-		int32 InternalArcCount = Package.InternalArcs.Num();
-		ZenGraphArchive << InternalArcCount;
-		for (FArc& InternalArc : Package.InternalArcs)
-		{
-			ZenGraphArchive << InternalArc.FromNodeIndex;
-			ZenGraphArchive << InternalArc.ToNodeIndex;
-		}
-
-		int32 ImportedPackagesCount = Package.ExternalArcs.Num();
-		ZenGraphArchive << ImportedPackagesCount;
-		for (auto& KV : Package.ExternalArcs)
-		{
-			FPackage* ImportedPackage = KV.Key;
-			TArray<FArc>& Arcs = KV.Value;
-			int32 ExternalArcCount = Arcs.Num();
-
-			ZenGraphArchive << ImportedPackage->GlobalPackageId;
-			ZenGraphArchive << ExternalArcCount;
-			for (FArc& ExternalArc : Arcs)
-			{
-				ZenGraphArchive << ExternalArc.FromNodeIndex;
-				ZenGraphArchive << ExternalArc.ToNodeIndex;
-			}
-		}
-		Package.UGraphSize = ZenGraphArchive.Tell();
-
-		// Temporary Archive for serializing export map data
-		FBufferWriter ZenExportMapArchive(nullptr, 0, EBufferWriterFlags::AllowResize | EBufferWriterFlags::TakeOwnership );
-		for (int32 I = 0; I < Package.ExportCount; ++I)
-		{
-			FObjectExport& ObjectExport = Exports[Package.ExportIndexOffset + I];
-			FExportData& ExportData = GlobalExports[Package.Exports[I]];
-
-			ZenExportMapArchive << ObjectExport.SerialSize;
-			NameMapBuilder.SerializeName(ZenExportMapArchive, ObjectExport.ObjectName);
-			ZenExportMapArchive << ExportData.OuterIndex;
-			ZenExportMapArchive << ExportData.ClassIndex;
-			ZenExportMapArchive << ExportData.SuperIndex;
-			ZenExportMapArchive << ExportData.TemplateIndex;
-			ZenExportMapArchive << ExportData.GlobalImportIndex;
-			ZenExportMapArchive << (uint32&)ObjectExport.ObjectFlags;
-		}
-		Package.ExportMapSize = ZenExportMapArchive.Tell();
-
-		// Temporary archive for serializing export bundle data
-		FBufferWriter ZenExportBundlesArchive(nullptr, 0, EBufferWriterFlags::AllowResize | EBufferWriterFlags::TakeOwnership);
-		int32 ExportBundleEntryIndex = 0;
-		for (FExportBundle& ExportBundle : Package.ExportBundles)
-		{
-			ZenExportBundlesArchive << ExportBundleEntryIndex;
-			int32 EntryCount = ExportBundle.Nodes.Num();
-			ZenExportBundlesArchive << EntryCount;
-			ExportBundleEntryIndex += ExportBundle.Nodes.Num();
-		}
-		for (FExportBundle& ExportBundle : Package.ExportBundles)
-		{
-			for (FExportGraphNode* ExportNode : ExportBundle.Nodes)
-			{
-				uint32 CommandType = uint32(ExportNode->BundleEntry.CommandType);
-				ZenExportBundlesArchive << CommandType;
-				ZenExportBundlesArchive << ExportNode->BundleEntry.LocalExportIndex;
-			}
-		}
-		Package.ExportBundlesSize = ZenExportBundlesArchive.Tell();
-
-
-		ensure(Package.GlobalPackageId == PackageIndex++);
-		{
-			int32 ExportBundleCount = Package.ExportBundles.Num();
-			int32 FirstExportBundleIndex = ExportBundleMetaEntries.Num();
-
-			NameMapBuilder.MarkNameAsReferenced(Package.Name);
-
-			NameMapBuilder.SerializeName(*StoreTocArchive, Package.Name);
-			*StoreTocArchive << Package.ExportCount;
-			*StoreTocArchive << ExportBundleCount;
-			*StoreTocArchive << FirstExportBundleIndex;
-			*StoreTocArchive << Package.FirstGlobalImport;
-			*StoreTocArchive << Package.GlobalImportCount;
-			*StoreTocArchive << Package.ImportedPackagesCount;
-			*StoreTocArchive << Package.ImportedPackagesOffset;
-		}
-
-		Package.NameMapSize = Package.NameIndices.Num() * Package.NameIndices.GetTypeSize();
-
-		{
-			const uint64 ZenSummarySize = 
-						sizeof(FZenPackageSummary)
-						+ Package.NameMapSize
-						+ Package.ImportMapSize
-						+ Package.ExportMapSize
-						+ Package.ExportBundlesSize
-						+ Package.UGraphSize;
-
-			uint8* ZenSummaryBuffer = static_cast<uint8*>(FMemory::Malloc(ZenSummarySize));
-			FZenPackageSummary* ZenSummary = reinterpret_cast<FZenPackageSummary*>(ZenSummaryBuffer);
-
-			ZenSummary->PackageFlags = Package.PackageFlags;
-			ZenSummary->GraphDataSize = Package.UGraphSize;
-			ZenSummary->BulkDataStartOffset = Package.BulkDataStartOffset;
-			int32* FindGlobalImportIndexForPackage = GlobalImportsByFullName.Find(Package.Name.ToString());
-			ZenSummary->GlobalImportIndex = FindGlobalImportIndexForPackage ? *FindGlobalImportIndexForPackage : -1;
-
-			FBufferWriter ZenAr(ZenSummaryBuffer, ZenSummarySize); 
-			ZenAr.Seek(sizeof(FZenPackageSummary));
-
-			// NameMap data
-			{
-				ZenSummary->NameMapOffset = ZenAr.Tell();
-				ZenAr.Serialize(Package.NameIndices.GetData(), Package.NameMapSize);
-			}
-
-			// ImportMap data
-			{
-				check(ImportMapArchive.Tell() == Package.ImportMapSize);	
-				ZenSummary->ImportMapOffset = ZenAr.Tell();
-				ZenAr.Serialize(ImportMapArchive.GetWriterData(), ImportMapArchive.Tell());
-			}
-
-			// ExportMap data
-			{
-				check(ZenExportMapArchive.Tell() == Package.ExportMapSize);	
-				ZenSummary->ExportMapOffset = ZenAr.Tell();
-				ZenAr.Serialize(ZenExportMapArchive.GetWriterData(), ZenExportMapArchive.Tell());
-			}
-
-			// ExportBundle data
-			{
-				check(ZenExportBundlesArchive.Tell() == Package.ExportBundlesSize);
-				ZenSummary->ExportBundlesOffset = ZenAr.Tell();
-				ZenAr.Serialize(ZenExportBundlesArchive.GetWriterData(), ZenExportBundlesArchive.Tell());
-			}
-
-			// Graph data
-			{
-				check(ZenGraphArchive.Tell() == Package.UGraphSize);	
-				ZenSummary->GraphDataOffset = ZenAr.Tell();
-				ZenAr.Serialize(ZenGraphArchive.GetWriterData(), ZenGraphArchive.Tell());
-			}
-
-			// Export bundle chunks
-			{
-				check(Package.ExportBundles.Num());
-
-				FString UExpFileName = FPaths::ChangeExtension(Package.FileName, TEXT(".uexp"));
-				TUniquePtr<FArchive> ExpAr(IFileManager::Get().CreateFileReader(*UExpFileName));
-				Package.UExpSize = ExpAr->TotalSize();
-#if !SKIP_WRITE_CONTAINER
-				uint8* ExportsBuffer = static_cast<uint8*>(FMemory::Malloc(ExpAr->TotalSize()));
-				ExpAr->Serialize(ExportsBuffer, ExpAr->TotalSize());
-#endif
-				ExpAr->Close();
-
-				uint64 BundleBufferSize = ZenSummarySize;
-				for (int32 I = 0; I < Package.ExportBundles.Num(); ++I)
-				{
-					FExportBundle& ExportBundle = Package.ExportBundles[I];
-					FExportBundleMeta& ExportBundleMeta = ExportBundleMetaEntries.AddDefaulted_GetRef();
-					ExportBundleMeta.LoadOrder = ExportBundle.LoadOrder;
-					for (FExportGraphNode* Node : ExportBundle.Nodes)
-					{
-						if (Node->BundleEntry.CommandType == FExportBundleEntry::ExportCommandType_Serialize)
-						{
-							FObjectExport& ObjectExport = Exports[Package.ExportIndexOffset + Node->BundleEntry.LocalExportIndex];
-							BundleBufferSize += ObjectExport.SerialSize;
-						}
-					}
-					if (I == 0)
-					{
-						ExportBundleMeta.PayloadSize = BundleBufferSize;
-					}
-				}
+	}
 
 #if !SKIP_WRITE_CONTAINER
-				uint8* BundleBuffer = static_cast<uint8*>(FMemory::Malloc(BundleBufferSize));
-				FMemory::Memcpy(BundleBuffer, ZenSummaryBuffer, ZenSummarySize);
-#endif
-				FMemory::Free(ZenSummaryBuffer);
-				uint64 BundleBufferOffset = ZenSummarySize;
-				for (FExportBundle& ExportBundle : Package.ExportBundles)
-				{
-					for (FExportGraphNode* Node : ExportBundle.Nodes)
-					{
-						if (Node->BundleEntry.CommandType == FExportBundleEntry::ExportCommandType_Serialize)
-						{
-							FObjectExport& ObjectExport = Exports[Package.ExportIndexOffset + Node->BundleEntry.LocalExportIndex];
-							const int64 Offset = ObjectExport.SerialOffset - Package.UAssetSize;
-#if !SKIP_WRITE_CONTAINER
-							FMemory::Memcpy(BundleBuffer + BundleBufferOffset, ExportsBuffer + Offset, ObjectExport.SerialSize);
-#endif
-							BundleBufferOffset += ObjectExport.SerialSize;
-						}
-					}
-				}
+	TUniquePtr<FIoStoreWriter> IoStoreWriter;
 
-#if !SKIP_WRITE_CONTAINER
-				FIoBuffer IoBuffer(FIoBuffer::Wrap, BundleBuffer, BundleBufferSize);
-				IoStoreWriter->Append(CreateChunkId(Package.GlobalPackageId, 0, EIoChunkType::ExportBundleData, *Package.FileName), IoBuffer);
-				FMemory::Free(BundleBuffer);
-				FMemory::Free(ExportsBuffer);
+	{
+		FIoStoreEnvironment IoStoreEnv;
+		IoStoreEnv.InitializeFileEnvironment(OutputDir);
+		IoStoreWriter = MakeUnique<FIoStoreWriter>(IoStoreEnv);
+		FIoStatus IoStatus = IoStoreWriter->Initialize();
+		check(IoStatus.IsOk());
+	}
 #endif
-				++ExportChunkCount;
-			}
-
-#if !SKIP_BULKDATA
-			// Bulk chunks
-			if (bWithBulkDataManifest)
-			{
-				FString BulkFileName = FPaths::ChangeExtension(Package.FileName, TEXT(".ubulk"));
-				FPaths::NormalizeFilename(BulkFileName);
-
-				if (!WriteBulkData(BulkFileName, EIoChunkType::BulkData, Package, BulkDataManifest, IoStoreWriter.Get(), BulkChunkCount, BulkPartialChunkCount))
-				{
-					return 1;
-				}
 
-				FString OptionalBulkFileName = FPaths::ChangeExtension(Package.FileName, TEXT(".uptnl"));
-				FPaths::NormalizeFilename(OptionalBulkFileName);
-
-				if (!WriteBulkData(OptionalBulkFileName, EIoChunkType::OptionalBulkData, Package, BulkDataManifest, IoStoreWriter.Get(), BulkChunkCount, BulkPartialChunkCount))
-				{
-					return 1;
-				}
-			}
-#endif
-		}
-	}
+	UE_LOG(LogIoStore, Display, TEXT("Serializing..."));
 
-	auto AppendToContainer = [](FIoStoreWriter& IoWriter, FLargeMemoryWriter& Ar, const FIoChunkId& ChunkId) -> FIoStatus
-	{
-		FIoBuffer IoBuffer(FIoBuffer::Wrap, Ar.GetData(), Ar.TotalSize());
-		return IoWriter.Append(ChunkId, IoBuffer);
-	};
+	SerializePackageData(
+		IoStoreWriter.Get(),
+		Packages,
+		NameMapBuilder,
+		ObjectExports,
+		GlobalExports,
+		GlobalImportsByFullName,
+		ExportBundleMetaEntries.GetData(),
+		BulkDataManifest,
+		bWithBulkDataManifest);
 
 	int32 StoreTocByteCount = StoreTocArchive->TotalSize();
 	int32 ImportedPackagesByteCount = ImportedPackagesArchive->TotalSize();
@@ -1907,7 +1900,7 @@ int32 CreateTarget(const FContainerTarget& Target)
 		GlobalMetaArchive << ExportBundleMetaByteCount;
 		GlobalMetaArchive.Serialize(ExportBundleMetaEntries.GetData(), ExportBundleMetaByteCount);
 
-		const FIoStatus Status = AppendToContainer(*IoStoreWriter, GlobalMetaArchive, CreateIoChunkId(0, 0, EIoChunkType::LoaderGlobalMeta));
+		const FIoStatus Status = IoStoreWriter->Append(CreateIoChunkId(0, 0, EIoChunkType::LoaderGlobalMeta), FIoBuffer(FIoBuffer::Wrap, GlobalMetaArchive.GetData(), GlobalMetaArchive.TotalSize()));
 		if (!Status.IsOk())
 		{
 			UE_LOG(LogIoStore, Error, TEXT("Failed to save global meta data to container file"));
@@ -1916,7 +1909,7 @@ int32 CreateTarget(const FContainerTarget& Target)
 
 	{
 		UE_LOG(LogIoStore, Display, TEXT("Saving initial load meta data to container file"));
-		const FIoStatus Status = AppendToContainer(*IoStoreWriter, *InitialLoadArchive, CreateIoChunkId(0, 0, EIoChunkType::LoaderInitialLoadMeta));
+		const FIoStatus Status = IoStoreWriter->Append(CreateIoChunkId(0, 0, EIoChunkType::LoaderInitialLoadMeta), FIoBuffer(FIoBuffer::Wrap, InitialLoadArchive->GetData(), InitialLoadArchive->TotalSize()));
 		if (!Status.IsOk())
 		{
 			UE_LOG(LogIoStore, Error, TEXT("Failed to save initial load meta data to container file"));
@@ -1959,7 +1952,7 @@ int32 CreateTarget(const FContainerTarget& Target)
 	uint64 ExportMapSize = 0;
 	uint64 NameMapSize = 0;
 	uint64 NameMapCount = 0;
-	uint64 ZenPackageSummarySize = FileNames.Num() * sizeof(FZenPackageSummary);
+	uint64 ZenPackageSummarySize = Packages.Num() * sizeof(FZenPackageSummary);
 	uint64 ImportedPackagesCount = 0;
 	uint64 InitialLoadSize = InitialLoadArchive->Tell();
 	uint64 ScriptArcsCount = 0;
@@ -2022,16 +2015,6 @@ int32 CreateTarget(const FContainerTarget& Target)
 	UE_LOG(LogIoStore, Display, TEXT("Packages: %8d total, %d circular dependencies, %d no preload dependencies, %d no import dependencies"),
 		PackageMap.Num(), PackagesWithCircularDependenciesCount, PackagesWithoutPreloadDependenciesCount, PackagesWithoutImportDependenciesCount);
 	UE_LOG(LogIoStore, Display, TEXT("Bundles:  %8d total, %d entries, %d export objects"), BundleCount, BundleEntryCount, GlobalExports.Num());
-	UE_LOG(LogIoStore, Display, TEXT("Chunks:   %8d summary chunks, %d export chunks, %d bulk data chunks, %d partial bulk data chunks"), SummaryChunkCount, ExportChunkCount, BulkChunkCount, BulkPartialChunkCount);
-
-	UE_LOG(LogIoStore, Display, TEXT("Pak: %12.2f MB UExp, %d files"), (double)UExpSize / 1024.0 / 1024.0, FileNames.Num());
-	UE_LOG(LogIoStore, Display, TEXT("Pak: %12.2f MB UAsset/UMap, %d files"), (double)UAssetSize / 1024.0 / 1024.0, FileNames.Num());
-	UE_LOG(LogIoStore, Display, TEXT("Pak: %12.2f MB FPackageFileSummary"), (double)SummarySize / 1024.0 / 1024.0);
-	UE_LOG(LogIoStore, Display, TEXT("Pak: %12.2f MB NameMap, %d names"), (double)NameSize / 1024.0 / 1024.0, NameCount);
-	UE_LOG(LogIoStore, Display, TEXT("Pak: %12.2f MB ExportMap, %d exports"), (double)ExportSize / 1024.0 / 1024.0, Exports.Num());
-	UE_LOG(LogIoStore, Display, TEXT("Pak: %12.2f MB ImportMap, %d imports (%d UPackages)"), (double)ImportSize / 1024.0 / 1024.0, Imports.Num(), UPackageImports);
-	UE_LOG(LogIoStore, Display, TEXT("Pak: %12.2f MB PreloadDependencies, %d imports + %d exports = %d"),
-		(double)PreloadDependenciesSize / 1024.0 / 1024.0, ImportPreloadCount, ExportPreloadCount, PreloadDependencies.Num());
 
 	UE_LOG(LogIoStore, Display, TEXT("IoStore: %8.2f MB GlobalNames, %d unique names"), (double)GlobalNamesMB, NameMapBuilder.GetNameMap().Num());
 	UE_LOG(LogIoStore, Display, TEXT("IoStore: %8.2f MB GlobalNameHashes"), (double)GlobalNameHashesMB);
@@ -2040,8 +2023,8 @@ int32 CreateTarget(const FContainerTarget& Target)
 	UE_LOG(LogIoStore, Display, TEXT("IoStore: %8.2f MB GlobalBundleMeta, %d bundles"), (double)ExportBundleMetaByteCount / 1024.0 / 1024.0, ExportBundleMetaEntries.Num());
 	UE_LOG(LogIoStore, Display, TEXT("IoStore: %8.2f MB GlobalImportNames, %d total imports, %d script imports, %d UPackage imports"),
 		(double)GlobalImportNamesByteCount / 1024.0 / 1024.0, GlobalImportsByFullName.Num(), NumScriptImports, UniqueImportPackages);
-	UE_LOG(LogIoStore, Display, TEXT("IoStore: %8.2f MB InitialLoadData, %d script arcs, %d script outers, %d packages"), (double)InitialLoadSize / 1024.0 / 1024.0, ScriptArcsCount, NumScriptImports, FileNames.Num());
-	UE_LOG(LogIoStore, Display, TEXT("IoStore: %8.2f MB PackageHeader, %d packages"), (double)ZenPackageHeaderSize / 1024.0 / 1024.0, FileNames.Num());
+	UE_LOG(LogIoStore, Display, TEXT("IoStore: %8.2f MB InitialLoadData, %d script arcs, %d script outers, %d packages"), (double)InitialLoadSize / 1024.0 / 1024.0, ScriptArcsCount, NumScriptImports, Packages.Num());
+	UE_LOG(LogIoStore, Display, TEXT("IoStore: %8.2f MB PackageHeader, %d packages"), (double)ZenPackageHeaderSize / 1024.0 / 1024.0, Packages.Num());
 	UE_LOG(LogIoStore, Display, TEXT("IoStore: %8.2f MB PackageSummary"), (double)ZenPackageSummarySize / 1024.0 / 1024.0);
 	UE_LOG(LogIoStore, Display, TEXT("IoStore: %8.2f MB PackageNameMap, %d indices"), (double)NameMapSize / 1024.0 / 1024.0, NameMapCount);
 	UE_LOG(LogIoStore, Display, TEXT("IoStore: %8.2f MB PackageImportMap"), (double)ImportMapSize / 1024.0 / 1024.0);
