IoDispatcher:
Split container files based on pak chunking

diff --git a/Engine/Source/Developer/IoStoreUtilities/Private/IoStoreUtilities.cpp b/Engine/Source/Developer/IoStoreUtilities/Private/IoStoreUtilities.cpp
--- a/Engine/Source/Developer/IoStoreUtilities/Private/IoStoreUtilities.cpp
+++ b/Engine/Source/Developer/IoStoreUtilities/Private/IoStoreUtilities.cpp
@@ -25,6 +25,7 @@
 #include "UObject/Package.h"
 #include "UObject/SavePackage.h"
 #include "Algo/Find.h"
+#include "Misc/FileHelper.h"
 
 //PRAGMA_DISABLE_OPTIMIZATION
 
@@ -44,6 +45,7 @@ struct FContainerTarget
 	FString CookedDirectory;
 	FString CookedProjectDirectory;
 	FString OutputDirectory;
+	FString ChunkListFile;
 };
 
 class FNameMapBuilder
@@ -472,6 +474,13 @@ struct FCircularImportChain
 	}
 };
 
+struct FInstallChunk
+{
+	FString Name;
+	int32 ChunkId;
+	TArray<FPackage*> Packages;
+};
+
 TArray<FPackage*> FPackageGraph::TopologicalSort() const
 {
 	TMultiMap<FPackageGraphNode*, FPackageGraphNode*> EdgesCopy = Edges;
@@ -1065,13 +1074,21 @@ FPackage* AddPackage(const TCHAR* FileName, const TCHAR* CookedDir, TArray<FPack
 
 	FName PackageFName = *PackageName;
 
-	FPackage* Package = new FPackage();
-	Package->Name = PackageFName;
-	Package->FileName = FileName;
-	Package->RelativeFileName = MoveTemp(RelativeFileName);
-	Package->GlobalPackageId = Packages.Num();
-	Packages.Add(Package);
-	PackageMap.Add(PackageFName, Package);
+	FPackage* Package = PackageMap.FindRef(PackageFName);
+	if (Package)
+	{
+		UE_LOG(LogIoStore, Warning, TEXT("Package in multiple pakchunks: '%s'"), *PackageFName.ToString());
+	}
+	else
+	{
+		Package = new FPackage();
+		Package->Name = PackageFName;
+		Package->FileName = FileName;
+		Package->RelativeFileName = MoveTemp(RelativeFileName);
+		Package->GlobalPackageId = Packages.Num();
+		Packages.Add(Package);
+		PackageMap.Add(PackageFName, Package);
+	}
 
 	return Package;
 }
@@ -1350,17 +1367,82 @@ int32 CreateTarget(const FContainerTarget& Target)
 	FExportGraph ExportGraph;
 	TArray<FExportBundleMeta> ExportBundleMetaEntries;
 
+	TArray<FInstallChunk> InstallChunks;
 	TArray<FPackage*> Packages;
 	TMap<FName, FPackage*> PackageMap;
+	if (Target.ChunkListFile.IsEmpty())
 	{
 		UE_LOG(LogIoStore, Display, TEXT("Searching for .uasset and .umap files..."));
 		TArray<FString> FileNames;
 		IFileManager::Get().FindFilesRecursive(FileNames, *CookedDir, TEXT("*.uasset"), true, false, false);
 		IFileManager::Get().FindFilesRecursive(FileNames, *CookedDir, TEXT("*.umap"), true, false, false);
 		UE_LOG(LogIoStore, Display, TEXT("Found '%d' files"), FileNames.Num());
+		FInstallChunk& InstallChunk = InstallChunks.AddDefaulted_GetRef();
+		InstallChunk.Name = TEXT("container0");
 		for (FString& FileName : FileNames)
 		{
-			AddPackage(*FileName, *CookedDir, Packages, PackageMap);
+			if (FPackage* Package = AddPackage(*FileName, *CookedDir, Packages, PackageMap))
+			{
+				InstallChunk.Packages.Add(Package);
+			}
+		}
+	}
+	else
+	{
+		TArray<FString> ChunkManifestFileNames;
+		FString ChunkFilesDirectory = FPaths::GetPath(Target.ChunkListFile);
+		FFileHelper::LoadFileToStringArray(ChunkManifestFileNames, *Target.ChunkListFile);
+		for (const FString& ChunkManifestFileName : ChunkManifestFileNames)
+		{
+			const TCHAR* PakChunkPrefix = TEXT("pakchunk");
+			const int32 PakChunkPrefixLength = 8;
+			if (!ChunkManifestFileName.StartsWith(PakChunkPrefix))
+			{
+				UE_LOG(LogIoStore, Warning, TEXT("Unexpected file name '%s' in '%s'"), *ChunkManifestFileName, *Target.ChunkListFile);
+				continue;
+			}
+			int32 Index = PakChunkPrefixLength;
+			int32 DigitCount = 0;
+			while (Index < ChunkManifestFileName.Len() && FChar::IsDigit(ChunkManifestFileName[Index]))
+			{
+				++DigitCount;
+				++Index;
+			}
+			if (DigitCount <= 0)
+			{
+				UE_LOG(LogIoStore, Warning, TEXT("Unexpected file name '%s' in '%s'"), *ChunkManifestFileName, *Target.ChunkListFile);
+				continue;
+			}
+			FString ChunkIdString = ChunkManifestFileName.Mid(PakChunkPrefixLength, DigitCount);
+			check(ChunkIdString.IsNumeric());
+			int32 ChunkId;
+			TTypeFromString<int32>::FromString(ChunkId, *ChunkIdString);
+			FInstallChunk& InstallChunk = InstallChunks.AddDefaulted_GetRef();
+			InstallChunk.Name = TEXT("container") + FPaths::GetBaseFilename(ChunkManifestFileName.RightChop(8));
+			InstallChunk.ChunkId = ChunkId;
+			FString ChunkManifestFullPath = ChunkFilesDirectory / ChunkManifestFileName;
+			TArray<FString> ChunkManifest;
+			FFileHelper::LoadFileToStringArray(ChunkManifest, *ChunkManifestFullPath);
+			for (const FString& FileNameWithoutExtension : ChunkManifest)
+			{
+				FString RelativePathWithoutExtension = IFileManager::Get().ConvertToRelativePath(*FileNameWithoutExtension);
+				FString FileName = RelativePathWithoutExtension + ".uasset";
+				if (!IFileManager::Get().FileExists(*FileName))
+				{
+					FileName = RelativePathWithoutExtension + ".umap";
+				}
+				if (!IFileManager::Get().FileExists(*FileName))
+				{
+					UE_LOG(LogIoStore, Warning, TEXT("Couldn't find package file (uexp/umap) for package '%s'"), *FileName);
+				}
+				else
+				{
+					if (FPackage* Package = AddPackage(*FileName, *CookedDir, Packages, PackageMap))
+					{
+						InstallChunk.Packages.Add(Package);
+					}
+				}
+			}
 		}
 	}
 
@@ -1775,7 +1857,6 @@ int32 CreateTarget(const FContainerTarget& Target)
 	UE_LOG(LogIoStore, Display, TEXT("Building bundles..."));
 	BuildBundles(ExportGraph, Packages);
 
-
 	TUniquePtr<FLargeMemoryWriter> StoreTocArchive = MakeUnique<FLargeMemoryWriter>(0, true);
 	TUniquePtr<FLargeMemoryWriter> ImportedPackagesArchive = MakeUnique<FLargeMemoryWriter>(0, true);
 	TUniquePtr<FLargeMemoryWriter> GlobalImportNamesArchive = MakeUnique<FLargeMemoryWriter>(0, true);
@@ -1855,30 +1936,47 @@ int32 CreateTarget(const FContainerTarget& Target)
 		}
 	}
 
+	UE_LOG(LogIoStore, Display, TEXT("Serializing..."));
+
+	TUniquePtr<FIoStoreWriter> GlobalIoStoreWriter;
+	FIoStoreEnvironment GlobalIoStoreEnv;
+	{
+		GlobalIoStoreEnv.InitializeFileEnvironment(OutputDir);
+		GlobalIoStoreWriter = MakeUnique<FIoStoreWriter>(GlobalIoStoreEnv);
 #if !SKIP_WRITE_CONTAINER
-	TUniquePtr<FIoStoreWriter> IoStoreWriter;
+		FIoStatus IoStatus = GlobalIoStoreWriter->Initialize();
+		check(IoStatus.IsOk());
+#endif
+	}
 
+	FIoStoreInstallManifest Manifest;
+	for (const FInstallChunk& InstallChunk : InstallChunks)
 	{
-		FIoStoreEnvironment IoStoreEnv;
-		IoStoreEnv.InitializeFileEnvironment(OutputDir);
-		IoStoreWriter = MakeUnique<FIoStoreWriter>(IoStoreEnv);
+		TUniquePtr<FIoStoreWriter> IoStoreWriter;
+		FIoStoreEnvironment InstallChunkIoStoreEnv(GlobalIoStoreEnv, InstallChunk.Name);
+		IoStoreWriter = MakeUnique<FIoStoreWriter>(InstallChunkIoStoreEnv);
+#if !SKIP_WRITE_CONTAINER
 		FIoStatus IoStatus = IoStoreWriter->Initialize();
 		check(IoStatus.IsOk());
-	}
 #endif
-
-	UE_LOG(LogIoStore, Display, TEXT("Serializing..."));
-
-	SerializePackageData(
-		IoStoreWriter.Get(),
-		Packages,
-		NameMapBuilder,
-		ObjectExports,
-		GlobalExports,
-		GlobalImportsByFullName,
-		ExportBundleMetaEntries.GetData(),
-		BulkDataManifest,
-		bWithBulkDataManifest);
+		SerializePackageData(
+			IoStoreWriter.Get(),
+			InstallChunk.Packages,
+			NameMapBuilder,
+			ObjectExports,
+			GlobalExports,
+			GlobalImportsByFullName,
+			ExportBundleMetaEntries.GetData(),
+			BulkDataManifest,
+			bWithBulkDataManifest);
+		FIoStoreInstallManifest::FEntry& ManifestEntry = Manifest.EditEntries().AddDefaulted_GetRef();
+		ManifestEntry.InstallChunkId = InstallChunk.ChunkId;
+		ManifestEntry.PartitionName = InstallChunk.Name;
+	}
+	TUniquePtr<FLargeMemoryWriter> ManifestArchive = MakeUnique<FLargeMemoryWriter>(0, true);
+	*ManifestArchive << Manifest;
+
+	GlobalIoStoreWriter->Append(CreateIoChunkId(0, 0, EIoChunkType::InstallManifest), FIoBuffer(FIoBuffer::Wrap, ManifestArchive->GetData(), ManifestArchive->TotalSize()));
 
 	int32 StoreTocByteCount = StoreTocArchive->TotalSize();
 	int32 ImportedPackagesByteCount = ImportedPackagesArchive->TotalSize();
@@ -1900,20 +1998,24 @@ int32 CreateTarget(const FContainerTarget& Target)
 		GlobalMetaArchive << ExportBundleMetaByteCount;
 		GlobalMetaArchive.Serialize(ExportBundleMetaEntries.GetData(), ExportBundleMetaByteCount);
 
-		const FIoStatus Status = IoStoreWriter->Append(CreateIoChunkId(0, 0, EIoChunkType::LoaderGlobalMeta), FIoBuffer(FIoBuffer::Wrap, GlobalMetaArchive.GetData(), GlobalMetaArchive.TotalSize()));
+#if !SKIP_WRITE_CONTAINER
+		const FIoStatus Status = GlobalIoStoreWriter->Append(CreateIoChunkId(0, 0, EIoChunkType::LoaderGlobalMeta), FIoBuffer(FIoBuffer::Wrap, GlobalMetaArchive.GetData(), GlobalMetaArchive.TotalSize()));
 		if (!Status.IsOk())
 		{
 			UE_LOG(LogIoStore, Error, TEXT("Failed to save global meta data to container file"));
 		}
+#endif
 	}
 
 	{
 		UE_LOG(LogIoStore, Display, TEXT("Saving initial load meta data to container file"));
-		const FIoStatus Status = IoStoreWriter->Append(CreateIoChunkId(0, 0, EIoChunkType::LoaderInitialLoadMeta), FIoBuffer(FIoBuffer::Wrap, InitialLoadArchive->GetData(), InitialLoadArchive->TotalSize()));
+#if !SKIP_WRITE_CONTAINER
+		const FIoStatus Status = GlobalIoStoreWriter->Append(CreateIoChunkId(0, 0, EIoChunkType::LoaderInitialLoadMeta), FIoBuffer(FIoBuffer::Wrap, InitialLoadArchive->GetData(), InitialLoadArchive->TotalSize()));
 		if (!Status.IsOk())
 		{
 			UE_LOG(LogIoStore, Error, TEXT("Failed to save initial load meta data to container file"));
 		}
+#endif
 	}
 
 	uint64 GlobalNamesMB = 0;
@@ -1928,10 +2030,12 @@ int32 CreateTarget(const FContainerTarget& Target)
 		GlobalNamesMB = Names.Num() >> 20;
 		GlobalNameHashesMB = Hashes.Num() >> 20;
 
-		FIoStatus NameStatus = IoStoreWriter->Append(CreateIoChunkId(0, 0, EIoChunkType::LoaderGlobalNames), 
+#if !SKIP_WRITE_CONTAINER
+		FIoStatus NameStatus = GlobalIoStoreWriter->Append(CreateIoChunkId(0, 0, EIoChunkType::LoaderGlobalNames), 
 													 FIoBuffer(FIoBuffer::Wrap, Names.GetData(), Names.Num()));
-		FIoStatus HashStatus = IoStoreWriter->Append(CreateIoChunkId(0, 0, EIoChunkType::LoaderGlobalNameHashes),
+		FIoStatus HashStatus = GlobalIoStoreWriter->Append(CreateIoChunkId(0, 0, EIoChunkType::LoaderGlobalNameHashes),
 													 FIoBuffer(FIoBuffer::Wrap, Hashes.GetData(), Hashes.Num()));
+#endif
 		
 		if (!NameStatus.IsOk() || !HashStatus.IsOk())
 		{
@@ -2060,6 +2164,12 @@ int32 CreateIoStoreContainerFiles(const TCHAR* CmdLine)
 		UE_LOG(LogIoStore, Display, TEXT("No output directory specified, using project's cooked folder"));
 	}
 
+	FString ChunkListFile;
+	if (FParse::Value(FCommandLine::Get(), TEXT("ChunkListFile="), ChunkListFile))
+	{
+		UE_LOG(LogIoStore, Display, TEXT("Using chunk list file: '%s'"), *ChunkListFile);
+	}
+
 	ITargetPlatformManagerModule& TargetPlatformManager = *GetTargetPlatformManager();
 	ITargetPlatform* TargetPlatform = TargetPlatformManager.FindTargetPlatform(TargetPlatformName);
 
@@ -2074,9 +2184,9 @@ int32 CreateIoStoreContainerFiles(const TCHAR* CmdLine)
 
 	FString TargetOutputDirectory = OutputDirectory.Len() > 0
 		? OutputDirectory
-		: TargetCookedProjectDirectory;
+		: TargetCookedProjectDirectory / TEXT("Content") / TEXT("Containers");
 
-	FContainerTarget Target { TargetPlatform, TargetCookedDirectory, TargetCookedProjectDirectory, TargetOutputDirectory };
+	FContainerTarget Target { TargetPlatform, TargetCookedDirectory, TargetCookedProjectDirectory, TargetOutputDirectory, ChunkListFile };
 
 	UE_LOG(LogIoStore, Display, TEXT("Creating target: '%s' using output directory: '%s'"), *Target.TargetPlatform->PlatformName(), *Target.OutputDirectory);
 
diff --git a/Engine/Source/Programs/AutomationTool/Scripts/CopyBuildToStagingDirectory.Automation.cs b/Engine/Source/Programs/AutomationTool/Scripts/CopyBuildToStagingDirectory.Automation.cs
--- a/Engine/Source/Programs/AutomationTool/Scripts/CopyBuildToStagingDirectory.Automation.cs
+++ b/Engine/Source/Programs/AutomationTool/Scripts/CopyBuildToStagingDirectory.Automation.cs
@@ -2920,11 +2920,12 @@ public partial class Project : CommandUtils
 			{
 				TargetPlatformDescriptor DataPlatformDesc = Params.GetCookedDataPlatformForClientTarget(ClientPlatform);
 				String TargetPlatformName = Platform.Platforms[DataPlatformDesc].GetCookPlatform(false, Params.Client);
-
-				DirectoryReference OutputDirectory = DirectoryReference.Combine(SC.RuntimeRootDir, SC.RelativeProjectRootForStage.Name);
-
+				DirectoryReference OutputDirectory = DirectoryReference.Combine(SC.RuntimeRootDir, SC.RelativeProjectRootForStage.Name, "Content", "Containers");
 				String CommandletParams = String.Format("-TargetPlatform={0} -OutputDirectory={1}", TargetPlatformName, OutputDirectory.FullName);
-
+				if (SC.PlatformUsesChunkManifests && DoesChunkPakManifestExist(Params, SC))
+				{
+					CommandletParams += String.Format(" -ChunkListFile={0}", GetChunkPakManifestListFilename(Params, SC));
+				}
 				RunCommandlet(SC.RawProjectPath, Params.UE4Exe, "IoStore", CommandletParams);
 			}
 		}
diff --git a/Engine/Source/Runtime/Core/Private/IO/IoDispatcher.cpp b/Engine/Source/Runtime/Core/Private/IO/IoDispatcher.cpp
--- a/Engine/Source/Runtime/Core/Private/IO/IoDispatcher.cpp
+++ b/Engine/Source/Runtime/Core/Private/IO/IoDispatcher.cpp
@@ -11,6 +11,8 @@
 #include "HAL/Runnable.h"
 #include "HAL/RunnableThread.h"
 #include "HAL/PlatformProcess.h"
+#include "Serialization/LargeMemoryReader.h"
+#include "GenericPlatform/GenericPlatformChunkInstall.h"
 
 DEFINE_LOG_CATEGORY(LogIoDispatcher);
 
@@ -141,6 +143,64 @@ public:
 		Thread = FRunnableThread::Create(this, TEXT("IoDispatcher"), 0, TPri_AboveNormal);
 	}
 
+	~FIoDispatcherImpl()
+	{
+		IPlatformChunkInstall* ChunkInstall = FPlatformMisc::GetPlatformChunkInstall();
+		if (ChunkInstall && ChunkDownloadedDelegateHandle.IsValid())
+		{
+			ChunkInstall->RemoveChunkInstallDelegate(ChunkDownloadedDelegateHandle);
+		}
+	}
+
+	FIoStatus Initialize(const FIoStoreEnvironment& InitialEnvironment)
+	{
+		FIoStatus IoStatus = Mount(InitialEnvironment);
+		if (!IoStatus.IsOk())
+		{
+			return IoStatus;
+		}
+
+		FIoBuffer IoBuffer;
+		FEvent* Event = FPlatformProcess::GetSynchEventFromPool();
+		ReadWithCallback(
+			CreateIoChunkId(0, 0, EIoChunkType::InstallManifest),
+			FIoReadOptions(),
+			[Event, &IoBuffer](TIoStatusOr<FIoBuffer> Result)
+		{
+			if (Result.IsOk())
+			{
+				IoBuffer = Result.ConsumeValueOrDie();
+			}
+			Event->Trigger();
+		});
+		Event->Wait();
+		if (!IoBuffer.DataSize())
+		{
+			return FIoStatusBuilder(EIoErrorCode::NotFound) << TEXT("Failed to open install manifest");
+		}
+		FLargeMemoryReader InstallManifestAr(IoBuffer.Data(), IoBuffer.DataSize());
+
+		FIoStoreInstallManifest InstallManifest;
+		InstallManifestAr << InstallManifest;
+
+		IPlatformChunkInstall* ChunkInstall = FPlatformMisc::GetPlatformChunkInstall();
+		ChunkDownloadedDelegateHandle = ChunkInstall->AddChunkInstallDelegate(FPlatformChunkInstallDelegate::CreateRaw(this, &FIoDispatcherImpl::OnChunkDownloaded));
+		for (const FIoStoreInstallManifest::FEntry& InstallManifestEntry : InstallManifest.ReadEntries())
+		{
+			FIoStoreEnvironment PartitionEnvironment(InitialEnvironment, InstallManifestEntry.PartitionName);
+			if (ChunkInstall && ChunkInstall->GetChunkLocation(InstallManifestEntry.InstallChunkId) == EChunkLocation::NotAvailable)
+			{
+				PendingInstallChunks.Add(MakeTuple(InstallManifestEntry.InstallChunkId, PartitionEnvironment));
+			}
+			else
+			{
+				Mount(PartitionEnvironment);
+			}
+		}
+
+		return FIoStatus::Ok;
+	}
+
 	FIoRequestImpl* AllocRequest(const FIoChunkId& ChunkId, FIoReadOptions Options)
 	{
 		FIoRequestImpl* Request = RequestAllocator.Construct();
@@ -208,9 +268,9 @@ public:
 		EventQueue.Notify();
 	}
 
-	FIoStatus Mount(FIoStoreEnvironment& Environment)
+	FIoStatus Mount(const FIoStoreEnvironment& Environment)
 	{
-		return FileIoStore.Mount(*Environment.GetRootPath());
+		return FileIoStore.Mount(Environment);
 	}
 
 	bool DoesChunkExist(const FIoChunkId& ChunkId) const
@@ -354,6 +414,23 @@ private:
 		}
 	}
 
+	void OnChunkDownloaded(uint32 ChunkId, bool bSuccess)
+	{
+		if (bSuccess)
+		{
+			for (auto It = PendingInstallChunks.CreateIterator(); It; ++It)
+			{
+				int32 PendingChunkId = It->Get<0>();
+				if (PendingChunkId == ChunkId)
+				{
+					const FIoStoreEnvironment& PendingEnvironment = It->Get<1>();
+					Mount(PendingEnvironment);
+					It.RemoveCurrent();
+				}
+			}
+		}
+	}
+
 	virtual bool Init()
 	{
 		return true;
@@ -392,6 +469,8 @@ private:
 	int32 CurrentRequestsToSubmitIndex = 0;
 	FIoRequestImpl* SubmittedRequestsHead = nullptr;
 	FIoRequestImpl* SubmittedRequestsTail = nullptr;
+	FDelegateHandle ChunkDownloadedDelegateHandle;
+	TArray<TTuple<int32, FIoStoreEnvironment>> PendingInstallChunks;
 };
 
 FIoDispatcher::FIoDispatcher()
@@ -404,7 +483,7 @@ FIoDispatcher::~FIoDispatcher()
 	delete Impl;
 }
 
-FIoStatus FIoDispatcher::Mount(FIoStoreEnvironment& Environment)
+FIoStatus FIoDispatcher::Mount(const FIoStoreEnvironment& Environment)
 {
 	return Impl->Mount(Environment);
 }
@@ -440,12 +519,24 @@ FIoDispatcher::GetSizeForChunk(const FIoChunkId& ChunkId) const
 	return Impl->GetSizeForChunk(ChunkId);
 }
 
+bool
+FIoDispatcher::IsInitialized()
+{
+	return GIoDispatcher.IsValid();
+}
+
+bool
+FIoDispatcher::IsValidEnvironment(const FIoStoreEnvironment& Environment)
+{
+	return FFileIoStore::IsValidEnvironment(Environment);
+}
+
 FIoStatus
-FIoDispatcher::Initialize()
+FIoDispatcher::Initialize(const FIoStoreEnvironment& InitialEnvironment)
 {
 	GIoDispatcher = MakeUnique<FIoDispatcher>();
 
-	return EIoErrorCode::Ok;
+	return GIoDispatcher->Impl->Initialize(InitialEnvironment);
 }
 
 void
diff --git a/Engine/Source/Runtime/Core/Private/IO/IoDispatcherFileBackend.cpp b/Engine/Source/Runtime/Core/Private/IO/IoDispatcherFileBackend.cpp
--- a/Engine/Source/Runtime/Core/Private/IO/IoDispatcherFileBackend.cpp
+++ b/Engine/Source/Runtime/Core/Private/IO/IoDispatcherFileBackend.cpp
@@ -15,22 +15,29 @@ FFileIoStoreReader::FFileIoStoreReader(FFileIoStoreImpl& InPlatformImpl)
 {
 }
 
-FIoStatus FFileIoStoreReader::Initialize(const TCHAR* RootPath)
+FIoStatus FFileIoStoreReader::Initialize(const FIoStoreEnvironment& Environment)
 {
 	IPlatformFile& Ipf = FPlatformFileManager::Get().GetPlatformFile();
 
 	TStringBuilder<256> ContainerFilePath;
-	ContainerFilePath.Append(RootPath);
+	ContainerFilePath.Append(Environment.GetBasePath());
 	if (ContainerFilePath.LastChar() != '/')
 	{
 		ContainerFilePath.Append(TEXT('/'));
 	}
-
+	if (Environment.GetPartitionName().IsEmpty())
+	{
+		ContainerFilePath.Append(TEXT("global"));
+	}
+	else
+	{
+		ContainerFilePath.Append(Environment.GetPartitionName());
+	}
 	TStringBuilder<256> TocFilePath;
 	TocFilePath.Append(ContainerFilePath);
-	TocFilePath.Append(TEXT("Container.utoc"));
 
-	ContainerFilePath.Append(TEXT("Container.ucas"));
+	ContainerFilePath.Append(TEXT(".ucas"));
+	TocFilePath.Append(TEXT(".utoc"));
 
 	if (!PlatformImpl.OpenContainer(*ContainerFilePath, ContainerFileHandle, ContainerFileSize))
 	{
@@ -145,10 +152,10 @@ FFileIoStore::FFileIoStore(FIoDispatcherEventQueue& InEventQueue)
 	InitCache();
 }
 
-FIoStatus FFileIoStore::Mount(const TCHAR* RootPath)
+FIoStatus FFileIoStore::Mount(const FIoStoreEnvironment& Environment)
 {
 	TUniquePtr<FFileIoStoreReader> Reader(new FFileIoStoreReader(PlatformImpl));
-	FIoStatus IoStatus = Reader->Initialize(RootPath);
+	FIoStatus IoStatus = Reader->Initialize(Environment);
 	if (IoStatus.IsOk())
 	{
 		FWriteScopeLock _(IoStoreReadersLock);
@@ -230,6 +237,26 @@ TIoStatusOr<uint64> FFileIoStore::GetSizeForChunk(const FIoChunkId& ChunkId) con
 	return FIoStatus(EIoErrorCode::NotFound);
 }
 
+bool FFileIoStore::IsValidEnvironment(const FIoStoreEnvironment& Environment)
+{
+	TStringBuilder<256> TocFilePath;
+	TocFilePath.Append(Environment.GetBasePath());
+	if (TocFilePath.LastChar() != '/')
+	{
+		TocFilePath.Append(TEXT('/'));
+	}
+	if (Environment.GetPartitionName().IsEmpty())
+	{
+		TocFilePath.Append(TEXT("global"));
+	}
+	else
+	{
+		TocFilePath.Append(Environment.GetPartitionName());
+	}
+	TocFilePath.Append(TEXT(".utoc"));
+	return FPlatformFileManager::Get().GetPlatformFile().FileExists(*TocFilePath);
+}
+
 void FFileIoStore::ProcessIncomingBlocks()
 {
 	while (FFileIoStoreReadBlock* CompletedBlock = PlatformImpl.GetNextCompletedBlock())
diff --git a/Engine/Source/Runtime/Core/Private/IO/IoDispatcherFileBackend.h b/Engine/Source/Runtime/Core/Private/IO/IoDispatcherFileBackend.h
--- a/Engine/Source/Runtime/Core/Private/IO/IoDispatcherFileBackend.h
+++ b/Engine/Source/Runtime/Core/Private/IO/IoDispatcherFileBackend.h
@@ -56,7 +56,7 @@ class FFileIoStoreReader
 {
 public:
 	FFileIoStoreReader(FFileIoStoreImpl& InPlatformImpl);
-	FIoStatus Initialize(const TCHAR* RootPath);
+	FIoStatus Initialize(const FIoStoreEnvironment& Environment);
 	bool DoesChunkExist(const FIoChunkId& ChunkId) const;
 	TIoStatusOr<uint64> GetSizeForChunk(const FIoChunkId& ChunkId) const;
 	bool Resolve(FFileIoStoreResolvedRequest& ResolvedRequest);
@@ -73,12 +73,14 @@ class FFileIoStore
 {
 public:
 	FFileIoStore(FIoDispatcherEventQueue& InEventQueue);
-	FIoStatus Mount(const TCHAR* RootPath);
+	FIoStatus Mount(const FIoStoreEnvironment& Environment);
 	EIoStoreResolveResult Resolve(FIoRequestImpl* Request);
 	bool DoesChunkExist(const FIoChunkId& ChunkId) const;
 	TIoStatusOr<uint64> GetSizeForChunk(const FIoChunkId& ChunkId) const;
 	void ProcessIncomingBlocks();
 
+	static bool IsValidEnvironment(const FIoStoreEnvironment& Environment);
+
 private:
 	void InitCache();
 	void ReadBlockCached(uint32 BlockIndex, const FFileIoStoreResolvedRequest& ResolvedRequest);
diff --git a/Engine/Source/Runtime/Core/Private/IO/IoStore.cpp b/Engine/Source/Runtime/Core/Private/IO/IoStore.cpp
--- a/Engine/Source/Runtime/Core/Private/IO/IoStore.cpp
+++ b/Engine/Source/Runtime/Core/Private/IO/IoStore.cpp
@@ -16,13 +16,19 @@ FIoStoreEnvironment::FIoStoreEnvironment()
 {
 }
 
+FIoStoreEnvironment::FIoStoreEnvironment(const FIoStoreEnvironment& BaseEnvironment, FStringView InPartitionName)
+{
+	BasePath = BaseEnvironment.BasePath;
+	PartitionName = InPartitionName;
+}
+
 FIoStoreEnvironment::~FIoStoreEnvironment()
 {
 }
 
-void FIoStoreEnvironment::InitializeFileEnvironment(FStringView InRootPath)
+void FIoStoreEnvironment::InitializeFileEnvironment(FStringView InBasePath)
 {
-	RootPath = InRootPath;
+	BasePath = InBasePath;
 }
 
 //////////////////////////////////////////////////////////////////////////
@@ -39,19 +45,15 @@ public:
 	{
 		IPlatformFile& Ipf = IPlatformFile::GetPlatformPhysical();
 
-		const FString& RootPath = Environment.GetRootPath();
-		Ipf.CreateDirectoryTree(*RootPath);
-
-		TStringBuilder<256> ContainerFilePath;
-		ContainerFilePath.Append(RootPath);
-		if (ContainerFilePath.LastChar() != '/')
-			ContainerFilePath.Append(TEXT('/'));
-
-		TStringBuilder<256> TocFilePath;
-		TocFilePath.Append(ContainerFilePath);
+		FString PartitionName = Environment.GetPartitionName();
+		if (PartitionName.IsEmpty())
+		{
+			PartitionName = TEXT("global");
+		}
+		FString TocFilePath = Environment.GetBasePath() / PartitionName + TEXT(".utoc");
+		FString ContainerFilePath = Environment.GetBasePath() / PartitionName + TEXT(".ucas");
 
-		ContainerFilePath.Append(TEXT("Container.ucas"));
-		TocFilePath.Append(TEXT("Container.utoc"));
+		Ipf.CreateDirectoryTree(*Environment.GetBasePath());
 
 		ContainerFileHandle.Reset(Ipf.OpenWrite(*ContainerFilePath, /* append */ false, /* allowread */ true));
 
diff --git a/Engine/Source/Runtime/Core/Public/IO/IoDispatcher.h b/Engine/Source/Runtime/Core/Public/IO/IoDispatcher.h
--- a/Engine/Source/Runtime/Core/Public/IO/IoDispatcher.h
+++ b/Engine/Source/Runtime/Core/Public/IO/IoDispatcher.h
@@ -570,14 +570,58 @@ private:
 	uint8	Id[12];
 };
 
+class FIoStoreInstallManifest
+{
+public:
+	struct FEntry
+	{
+		FString PartitionName;
+		int32 InstallChunkId;
+	};
+
+	inline const TArray<FEntry>& ReadEntries() const
+	{
+		return Entries;
+	}
+
+	inline TArray<FEntry>& EditEntries()
+	{
+		return Entries;
+	}
+
+	friend FArchive& operator<<(FArchive& Ar, FIoStoreInstallManifest& IoStoreManifest)
+	{
+		int32 Version = CurrentVersion;
+		Ar << Version;
+		check(Version == CurrentVersion);
+		int32 EntryCount = IoStoreManifest.Entries.Num();
+		Ar << EntryCount;
+		if (Ar.IsLoading())
+		{
+			IoStoreManifest.Entries.SetNum(EntryCount);
+		}
+		for (int32 Index = 0; Index < EntryCount; ++Index)
+		{
+			FEntry& Entry = IoStoreManifest.Entries[Index];
+			Ar << Entry.PartitionName;
+			Ar << Entry.InstallChunkId;
+		}
+		return Ar;
+	}
+
+private:
+	static constexpr int32 CurrentVersion = 1;
+
+	TArray<FEntry> Entries;
+};
+
 /**
  * Addressable chunk types.
  */
 enum class EIoChunkType : uint8
 {
 	Invalid,
-	PackageSummary,
-	ExportData,
+	InstallManifest,
 	ExportBundleData,
 	BulkData,
 	OptionalBulkData,
@@ -732,7 +776,7 @@ public:
 	CORE_API						FIoDispatcher();
 	CORE_API virtual				~FIoDispatcher();
 
-	CORE_API FIoStatus				Mount(FIoStoreEnvironment& Environment);
+	CORE_API FIoStatus				Mount(const FIoStoreEnvironment& Environment);
 
 	CORE_API FIoBatch				NewBatch();
 	CORE_API void					FreeBatch(FIoBatch Batch);
@@ -746,8 +790,9 @@ public:
 	FIoDispatcher(const FIoDispatcher&) = default;
 	FIoDispatcher& operator=(const FIoDispatcher&) = delete;
 
+	static CORE_API bool IsValidEnvironment(const FIoStoreEnvironment& Environment);
 	static CORE_API bool IsInitialized();
-	static CORE_API FIoStatus Initialize();
+	static CORE_API FIoStatus Initialize(const FIoStoreEnvironment& InitialEnvironment);
 	static CORE_API void Shutdown();
 	static CORE_API FIoDispatcher& Get();
 
@@ -767,14 +812,17 @@ class FIoStoreEnvironment
 {
 public:
 	CORE_API FIoStoreEnvironment();
+	CORE_API FIoStoreEnvironment(const FIoStoreEnvironment& BaseEnvironment, FStringView PartitionName);
 	CORE_API ~FIoStoreEnvironment();
 
-	CORE_API void InitializeFileEnvironment(FStringView InRootPath);
+	CORE_API void InitializeFileEnvironment(FStringView InBasePath);
 
-	CORE_API const FString& GetRootPath() const { return RootPath; }
+	CORE_API const FString& GetBasePath() const { return BasePath; }
+	CORE_API const FString& GetPartitionName() const { return PartitionName; }
 
 private:
-	FString			RootPath;
+	FString			BasePath;
+	FString			PartitionName;
 };
 
 //////////////////////////////////////////////////////////////////////////
diff --git a/Engine/Source/Runtime/CoreUObject/Private/Serialization/AsyncLoading2.cpp b/Engine/Source/Runtime/CoreUObject/Private/Serialization/AsyncLoading2.cpp
--- a/Engine/Source/Runtime/CoreUObject/Private/Serialization/AsyncLoading2.cpp
+++ b/Engine/Source/Runtime/CoreUObject/Private/Serialization/AsyncLoading2.cpp
@@ -1472,7 +1472,7 @@ class FAsyncLoadingThread2Impl
 {
 	friend struct FAsyncPackage2;
 public:
-	FAsyncLoadingThread2Impl(IEDLBootNotificationManager& InEDLBootNotificationManager);
+	FAsyncLoadingThread2Impl(FIoDispatcher& IoDispatcher, IEDLBootNotificationManager& InEDLBootNotificationManager);
 	virtual ~FAsyncLoadingThread2Impl();
 
 private:
@@ -1880,18 +1880,9 @@ struct FAsyncLoadingTickScope2
 
 void FAsyncLoadingThread2Impl::InitializeLoading()
 {
-	{
-		TRACE_CPUPROFILER_EVENT_SCOPE(InitIoDispatcher);
-
-		FIoStoreEnvironment IoStoreEnvironment;
-		IoStoreEnvironment.InitializeFileEnvironment(FPaths::ProjectDir());
-		FIoStatus ReaderStatus = IoDispatcher.Mount(IoStoreEnvironment);
-		UE_CLOG(!ReaderStatus.IsOk(), LogStreaming, Fatal, TEXT("Failed to initialize I/O dispatcher: '%s'"), *ReaderStatus.ToString());
-
 #if USE_NEW_BULKDATA
-		FBulkDataBase::SetIODispatcher(&IoDispatcher);
+	FBulkDataBase::SetIODispatcher(&IoDispatcher);
 #endif
-	}
 
 	{
 		TRACE_CPUPROFILER_EVENT_SCOPE(LoadGlobalNameMap);
@@ -3419,10 +3410,10 @@ EAsyncPackageState::Type FAsyncLoadingThread2Impl::TickAsyncLoadingFromGameThrea
 	return Result;
 }
 
-FAsyncLoadingThread2Impl::FAsyncLoadingThread2Impl(IEDLBootNotificationManager& InEDLBootNotificationManager)
+FAsyncLoadingThread2Impl::FAsyncLoadingThread2Impl(FIoDispatcher& InIoDispatcher, IEDLBootNotificationManager& InEDLBootNotificationManager)
 	: Thread(nullptr)
 	, EDLBootNotificationManager(InEDLBootNotificationManager)
-	, IoDispatcher(FIoDispatcher::Get())
+	, IoDispatcher(InIoDispatcher)
 {
 	GEventDrivenLoaderEnabled = true;
 
@@ -3469,8 +3460,6 @@ FAsyncLoadingThread2Impl::~FAsyncLoadingThread2Impl()
 
 void FAsyncLoadingThread2Impl::ShutdownLoading()
 {
-	FIoDispatcher::Shutdown();
-
 	FCoreUObjectDelegates::GetPreGarbageCollectDelegate().RemoveAll(this);
 	FCoreUObjectDelegates::GetPostGarbageCollect().RemoveAll(this);
 
@@ -4594,9 +4583,9 @@ EAsyncPackageState::Type FAsyncLoadingThread2Impl::ProcessLoadingUntilCompleteFr
 	return TimeLimit <= 0 ? EAsyncPackageState::TimeOut : EAsyncPackageState::Complete;
 }
 
-FAsyncLoadingThread2::FAsyncLoadingThread2(IEDLBootNotificationManager& InEDLBootNotificationManager)
+FAsyncLoadingThread2::FAsyncLoadingThread2(FIoDispatcher& InIoDispatcher, IEDLBootNotificationManager& InEDLBootNotificationManager)
 {
-	Impl = new FAsyncLoadingThread2Impl(InEDLBootNotificationManager);
+	Impl = new FAsyncLoadingThread2Impl(InIoDispatcher, InEDLBootNotificationManager);
 }
 
 FAsyncLoadingThread2::~FAsyncLoadingThread2()
diff --git a/Engine/Source/Runtime/CoreUObject/Private/Serialization/AsyncLoading2.h b/Engine/Source/Runtime/CoreUObject/Private/Serialization/AsyncLoading2.h
--- a/Engine/Source/Runtime/CoreUObject/Private/Serialization/AsyncLoading2.h
+++ b/Engine/Source/Runtime/CoreUObject/Private/Serialization/AsyncLoading2.h
@@ -10,12 +10,13 @@
 #include "Serialization/AsyncPackageLoader.h"
 
 class FAsyncLoadingThread2Impl;
+class FIoDispatcher;
 
 class FAsyncLoadingThread2
 	: public IAsyncPackageLoader
 {
 public:
-	FAsyncLoadingThread2(IEDLBootNotificationManager& InEDLBootNotificationManager);
+	FAsyncLoadingThread2(FIoDispatcher& InIoDispatcher, IEDLBootNotificationManager& InEDLBootNotificationManager);
 	virtual ~FAsyncLoadingThread2();
 
 	void InitializeLoading() override;
diff --git a/Engine/Source/Runtime/CoreUObject/Private/Serialization/AsyncPackageLoader.cpp b/Engine/Source/Runtime/CoreUObject/Private/Serialization/AsyncPackageLoader.cpp
--- a/Engine/Source/Runtime/CoreUObject/Private/Serialization/AsyncPackageLoader.cpp
+++ b/Engine/Source/Runtime/CoreUObject/Private/Serialization/AsyncPackageLoader.cpp
@@ -6,6 +6,7 @@
 #include "UObject/GCObject.h"
 #include "UObject/LinkerLoad.h"
 #include "Misc/CoreDelegates.h"
+#include "IO/IoDispatcher.h"
 
 volatile int32 GIsLoaderCreated;
 TUniquePtr<IAsyncPackageLoader> GPackageLoader;
@@ -518,9 +519,9 @@ IAsyncPackageLoader& GetAsyncPackageLoader()
 
 void InitAsyncThread()
 {
-	if (FParse::Param(FCommandLine::Get(), TEXT("zenloader")))
+	if (FIoDispatcher::IsInitialized())
 	{
-		GPackageLoader = MakeUnique<FAsyncLoadingThread2>(GetGEDLBootNotificationManager());
+		GPackageLoader = MakeUnique<FAsyncLoadingThread2>(FIoDispatcher::Get(), GetGEDLBootNotificationManager());
 	}
 	else
 	{
diff --git a/Engine/Source/Runtime/CoreUObject/Private/Serialization/BulkData2.cpp b/Engine/Source/Runtime/CoreUObject/Private/Serialization/BulkData2.cpp
--- a/Engine/Source/Runtime/CoreUObject/Private/Serialization/BulkData2.cpp
+++ b/Engine/Source/Runtime/CoreUObject/Private/Serialization/BulkData2.cpp
@@ -558,11 +558,10 @@ void FBulkDataBase::Serialize(FArchive& Ar, UObject* Owner, int32 /*Index*/, boo
 		Ar << BulkDataOffsetInFile;
 
 #if ENABLE_IO_DISPATCHER
-		static const bool bUseZenLoader = FParse::Param(FCommandLine::Get(), TEXT("zenloader"));
+		const bool bUseIoDispatcher = FIoDispatcher::IsInitialized();
 #else
-		const bool bUseZenLoader = false;
+		const bool bUseIoDispatcher = false;
 #endif
-
 		if ((BulkDataFlags & BULKDATA_BadDataVersion) != 0)
 		{
 			uint16 DummyValue;
@@ -576,7 +575,7 @@ void FBulkDataBase::Serialize(FArchive& Ar, UObject* Owner, int32 /*Index*/, boo
 		const UPackage* Package = Owner->GetOutermost();
 		check(Package != nullptr);
 
-		if (!IsInlined() && bUseZenLoader)
+		if (!IsInlined() && bUseIoDispatcher)
 		{
 			const EIoChunkType Type = IsOptional() ? EIoChunkType::OptionalBulkData : EIoChunkType::BulkData;
 			ChunkID = CreateBulkdataChunkId(Package->GetPackageId().ToIndex(), BulkDataOffsetInFile, Type);
@@ -594,7 +593,7 @@ void FBulkDataBase::Serialize(FArchive& Ar, UObject* Owner, int32 /*Index*/, boo
 		FName PackageName;
 		const FLinkerLoad* Linker = nullptr;
 
-		if (bUseZenLoader == false)
+		if (bUseIoDispatcher == false)
 		{
 			Linker = FLinkerLoad::FindExistingLinkerForPackage(Package);
 			check(Linker != nullptr);
@@ -638,7 +637,7 @@ void FBulkDataBase::Serialize(FArchive& Ar, UObject* Owner, int32 /*Index*/, boo
 					{
 						SerializeDuplicateData(Ar, BulkDataFlags, BulkDataSizeOnDisk, BulkDataOffsetInFile);
 
-						if (!IsInlined() && bUseZenLoader)
+						if (!IsInlined() && bUseIoDispatcher)
 						{
 							// Regenerate the FIoChunkId to find the optional BulkData instead!
 							ChunkID = CreateBulkdataChunkId(Package->GetPackageId().ToIndex(), BulkDataOffsetInFile, EIoChunkType::OptionalBulkData);
@@ -668,7 +667,7 @@ void FBulkDataBase::Serialize(FArchive& Ar, UObject* Owner, int32 /*Index*/, boo
 			}
 
 			// If we are not using the FIoDispatcher then we need to make sure we can retrieve the filename later
-			if (bUseZenLoader == false)
+			if (bUseIoDispatcher == false)
 			{
 				check(Filename != nullptr);
 				Fallback.Token = FileTokenSystem::RegisterFileToken( PackageName, *Filename, BulkDataOffsetInFile);
diff --git a/Engine/Source/Runtime/Launch/Private/LaunchEngineLoop.cpp b/Engine/Source/Runtime/Launch/Private/LaunchEngineLoop.cpp
--- a/Engine/Source/Runtime/Launch/Private/LaunchEngineLoop.cpp
+++ b/Engine/Source/Runtime/Launch/Private/LaunchEngineLoop.cpp
@@ -1581,14 +1581,6 @@ int32 FEngineLoop::PreInitPreStartupScreen(const TCHAR* CmdLine)
 		IFileManager::Get().ProcessCommandLineOptions();
 	}
 
-#if !(IS_PROGRAM || WITH_EDITOR)
-	// Initialize I/O dispatcher when using the new package loader
-	if (FParse::Param(FCommandLine::Get(), TEXT("zenloader")))
-	{
-		FIoDispatcher::Initialize();
-	}
-#endif
-
 	if (GIsGameAgnosticExe)
 	{
 		// If we launched without a project file, but with a game name that is incomplete, warn about the improper use of a Game suffix
@@ -5487,6 +5479,22 @@ bool FEngineLoop::AppInit( )
 
 	FEmbeddedCommunication::ForceTick(18);
 
+#if !(IS_PROGRAM || WITH_EDITOR)
+	// Initialize I/O dispatcher if available
+	if (!FParse::Param(FCommandLine::Get(), TEXT("noiodispatcher")))
+	{
+		TRACE_CPUPROFILER_EVENT_SCOPE(InitIoDispatcher);
+
+		FIoStoreEnvironment IoStoreEnvironment;
+		IoStoreEnvironment.InitializeFileEnvironment(FPaths::ProjectDir() / TEXT("Content") / TEXT("Containers"));
+		if (FIoDispatcher::IsValidEnvironment(IoStoreEnvironment))
+		{
+			FIoStatus IoDispatcherInitStatus = FIoDispatcher::Initialize(IoStoreEnvironment);
+			UE_CLOG(!IoDispatcherInitStatus.IsOk(), LogInit, Fatal, TEXT("Failed to initialize I/O dispatcher: '%s'"), *IoDispatcherInitStatus.ToString());
+		}
+	}
+#endif
+
 	// Init other systems.
 	{
 		SCOPED_BOOT_TIMING("FCoreDelegates::OnInit.Broadcast");
@@ -5564,6 +5572,10 @@ void FEngineLoop::AppPreExit( )
 
 void FEngineLoop::AppExit( )
 {
+#if !(IS_PROGRAM || WITH_EDITOR)
+	FIoDispatcher::Shutdown();
+#endif
+
 #if !WITH_ENGINE
 	// when compiled WITH_ENGINE, this will happen in FEngineLoop::Exit()
 	FTaskGraphInterface::Shutdown();
