BulkData2 will now be able to work with 0 sized BulkData, which maintains the exisiting behaviour of the original BulkData api.

- If the data is missing then a lot of features just won't work but since we were silently allowing it with the old BulkData api we should maintain this for now.
- When saving out a 0 byte BulkData entry to the manifest we instead add a special entry that will be 0 size and at the start of the bulkdata file. We also do not warn if this is a dupe.
- Moved around the code when adding BulkData to the IoStore so that we still create mappings for BulkData entries even in the case that all BulkData for a package is 0 sized and we don't have the file at all.
- Check for the size when creating the FIoChunkID at runtime and if 0 then use the special id instead of the offset.

#rnx

diff --git a/Engine/Source/Developer/IoStoreUtilities/Private/IoStoreUtilities.cpp b/Engine/Source/Developer/IoStoreUtilities/Private/IoStoreUtilities.cpp
--- a/Engine/Source/Developer/IoStoreUtilities/Private/IoStoreUtilities.cpp
+++ b/Engine/Source/Developer/IoStoreUtilities/Private/IoStoreUtilities.cpp
@@ -853,47 +853,53 @@ static void BuildBundles(FExportGraph& ExportGraph, const TArray<FPackage*>& Pac
 static bool WriteBulkData(	const FString& Filename, EIoChunkType Type, const FPackage& Package, const FPackageStoreBulkDataManifest& BulkDataManifest,
 							FIoStoreWriter* IoStoreWriter)
 {
-	TUniquePtr<FArchive> BulkAr(IFileManager::Get().CreateFileReader(*Filename));
-	if (BulkAr != nullptr)
+
+	const FPackageStoreBulkDataManifest::PackageDesc* PackageDesc = BulkDataManifest.Find(Package.FileName);
+	if (PackageDesc != nullptr)
 	{
-		const FIoChunkId BulkDataChunkId = CreateChunkIdForBulkData(Package.GlobalPackageId, INDEX_NONE, Type, *Package.FileName);
-#if !SKIP_WRITE_CONTAINER
-		uint8* BulkBuffer = static_cast<uint8*>(FMemory::Malloc(BulkAr->TotalSize()));
-		BulkAr->Serialize(BulkBuffer, BulkAr->TotalSize());
-		FIoBuffer IoBuffer(FIoBuffer::AssumeOwnership, BulkBuffer, BulkAr->TotalSize());
+		const FIoChunkId BulkDataChunkId = CreateChunkIdForBulkData(Package.GlobalPackageId, TNumericLimits<uint64>::Max()-1, Type, *Package.FileName);
+
+#if !SKIP_WRITE_CONTAINER		
+		FIoBuffer IoBuffer;
+
+		TUniquePtr<FArchive> BulkAr(IFileManager::Get().CreateFileReader(*Filename));
+		if (BulkAr != nullptr)
+		{
+			uint8* BulkBuffer = static_cast<uint8*>(FMemory::Malloc(BulkAr->TotalSize()));
+			BulkAr->Serialize(BulkBuffer, BulkAr->TotalSize());
+			IoBuffer = FIoBuffer(FIoBuffer::AssumeOwnership, BulkBuffer, BulkAr->TotalSize());
+
+			BulkAr->Close();
+		}
+		
 		const FIoStatus AppendResult = IoStoreWriter->Append(BulkDataChunkId, IoBuffer);
 		if (!AppendResult.IsOk())
 		{
 			UE_LOG(LogIoStore, Error, TEXT("Failed to append bulkdata for '%s' due to: %s"), *Package.FileName, *AppendResult.ToString());
 			return false;
 		}
-
-		BulkAr->Close();
 #endif
+
 		// Create additional mapping chunks as needed
-		const FPackageStoreBulkDataManifest::PackageDesc* PackageDesc = BulkDataManifest.Find(Package.FileName);
-		if (PackageDesc != nullptr)
+		for (const FPackageStoreBulkDataManifest::PackageDesc::BulkDataDesc& BulkDataDesc : PackageDesc->GetDataArray())
 		{
-			for (const FPackageStoreBulkDataManifest::PackageDesc::BulkDataDesc& BulkDataDesc : PackageDesc->GetDataArray())
+			if (BulkDataDesc.Type == Type)
 			{
-				if (BulkDataDesc.Type == Type)
+				const FIoChunkId AccessChunkId = CreateChunkIdForBulkData(Package.GlobalPackageId, BulkDataDesc.ChunkId, Type, *Package.FileName);
+#if !SKIP_WRITE_CONTAINER	
+				const FIoStatus PartialResult = IoStoreWriter->MapPartialRange(BulkDataChunkId, BulkDataDesc.Offset, BulkDataDesc.Size, AccessChunkId);
+				if (!PartialResult.IsOk())
 				{
-#if !SKIP_WRITE_CONTAINER
-					const FIoChunkId AccessChunkId = CreateChunkIdForBulkData(Package.GlobalPackageId, BulkDataDesc.ChunkId, Type, *Package.FileName);
-					const FIoStatus PartialResult = IoStoreWriter->MapPartialRange(BulkDataChunkId, BulkDataDesc.Offset, BulkDataDesc.Size, AccessChunkId);
-					if (!PartialResult.IsOk())
-					{
-						UE_LOG(LogIoStore, Warning, TEXT("Failed to map partial range for '%s' due to: %s"), *Package.FileName, *PartialResult.ToString());
-					}
-#endif
+					UE_LOG(LogIoStore, Warning, TEXT("Failed to map partial range for '%s' due to: %s"), *Package.FileName, *PartialResult.ToString());
 				}
+#endif
 			}
 		}
-		else
-		{
-			UE_LOG(LogIoStore, Error, TEXT("Unable to find an entry in the bulkdata manifest for '%s' the file might be out of date!"), *Package.FileName);
-			return false;
-		}
+	}
+	else if(IFileManager::Get().FileExists(*Filename))
+	{
+		UE_LOG(LogIoStore, Error, TEXT("Unable to find an entry in the bulkdata manifest for '%s' the file might be out of date!"), *Package.FileName);
+		return false;
 	}
 
 	return true;
diff --git a/Engine/Source/Runtime/CoreUObject/Private/Serialization/BulkData2.cpp b/Engine/Source/Runtime/CoreUObject/Private/Serialization/BulkData2.cpp
--- a/Engine/Source/Runtime/CoreUObject/Private/Serialization/BulkData2.cpp
+++ b/Engine/Source/Runtime/CoreUObject/Private/Serialization/BulkData2.cpp
@@ -578,7 +578,9 @@ void FBulkDataBase::Serialize(FArchive& Ar, UObject* Owner, int32 /*Index*/, boo
 		if (!IsInlined() && bUseIoDispatcher)
 		{
 			const EIoChunkType Type = IsOptional() ? EIoChunkType::OptionalBulkData : EIoChunkType::BulkData;
-			ChunkID = CreateBulkdataChunkId(Package->GetPackageId().ToIndex(), BulkDataOffsetInFile, Type);
+
+			const int64 BulkDataID = BulkDataSize > 0 ? BulkDataOffsetInFile : TNumericLimits<uint64>::Max();
+			ChunkID = CreateBulkdataChunkId(Package->GetPackageId().ToIndex(), BulkDataID, Type);
 
 			BulkDataFlags |= BULKDATA_UsesIoDispatcher; // Indicates that this BulkData should use the FIoChunkId rather than a filename
 		}
@@ -640,7 +642,8 @@ void FBulkDataBase::Serialize(FArchive& Ar, UObject* Owner, int32 /*Index*/, boo
 						if (!IsInlined() && bUseIoDispatcher)
 						{
 							// Regenerate the FIoChunkId to find the optional BulkData instead!
-							ChunkID = CreateBulkdataChunkId(Package->GetPackageId().ToIndex(), BulkDataOffsetInFile, EIoChunkType::OptionalBulkData);
+							const int64 BulkDataID = BulkDataSize > 0 ? BulkDataOffsetInFile : TNumericLimits<uint64>::Max();
+							ChunkID = CreateBulkdataChunkId(Package->GetPackageId().ToIndex(), BulkDataID, EIoChunkType::OptionalBulkData);
 							BulkDataFlags |= BULKDATA_UsesIoDispatcher; // Indicates that this BulkData should use the FIoChunkId rather than a filename
 						}
 						else
diff --git a/Engine/Source/Runtime/CoreUObject/Private/UObject/SavePackage.cpp b/Engine/Source/Runtime/CoreUObject/Private/UObject/SavePackage.cpp
--- a/Engine/Source/Runtime/CoreUObject/Private/UObject/SavePackage.cpp
+++ b/Engine/Source/Runtime/CoreUObject/Private/UObject/SavePackage.cpp
@@ -477,6 +477,21 @@ void FPackageStoreBulkDataManifest::PackageDesc::AddData(EIoChunkType InType, ui
 	}
 }
 
+void FPackageStoreBulkDataManifest::PackageDesc::AddZeroByteData(EIoChunkType InType)
+{
+	// Make sure we only add one empty read per Bulkdata type!
+	auto func = [=](const BulkDataDesc& Entry) { return Entry.Type == InType && Entry.Size == 0; };
+	if (Data.FindByPredicate(func) == nullptr)
+	{
+		BulkDataDesc& Entry = Data.Emplace_GetRef();
+
+		Entry.ChunkId = TNumericLimits<uint64>::Max();
+		Entry.Offset = 0;
+		Entry.Size = 0;
+		Entry.Type = InType;
+	}
+}
+
 FPackageStoreBulkDataManifest::FPackageStoreBulkDataManifest(const FString& ProjectPath)
 {
 	Filename = ProjectPath / TEXT("Metadata/BulkDataInfo.ubulkmanifest");
@@ -536,7 +551,16 @@ void FPackageStoreBulkDataManifest::AddFileAccess(const FString& PackageFilename
 	const FString NormalizedFilename = FixFilename(PackageFilename);
 	
 	PackageDesc& Entry = GetOrCreateFileAccess(NormalizedFilename);
-	Entry.AddData(InType, InChunkId, InOffset, InSize, NormalizedFilename);
+
+	if (InSize > 0)
+	{
+		Entry.AddData(InType, InChunkId, InOffset, InSize, NormalizedFilename);
+	}
+	else
+	{
+		Entry.AddZeroByteData(InType);
+	}
+	
 }
 
 FPackageStoreBulkDataManifest::PackageDesc& FPackageStoreBulkDataManifest::GetOrCreateFileAccess(const FString& PackageFilename)
diff --git a/Engine/Source/Runtime/CoreUObject/Public/UObject/SavePackage.h b/Engine/Source/Runtime/CoreUObject/Public/UObject/SavePackage.h
--- a/Engine/Source/Runtime/CoreUObject/Public/UObject/SavePackage.h
+++ b/Engine/Source/Runtime/CoreUObject/Public/UObject/SavePackage.h
@@ -38,6 +38,8 @@ public:
 		};
 
 		void AddData(EIoChunkType InType, uint64 InChunkId, uint64 InOffset, uint64 InSize, const FString& DebugFilename);
+		void AddZeroByteData(EIoChunkType InType);
+
 		const TArray<BulkDataDesc>& GetDataArray() const { return Data; }
 	private:
 		friend FArchive& operator<<(FArchive& Ar, PackageDesc& Entry);
