The IoStore commandlet now accepts '-targetplatform=*' in the same format that the cook commandlet does. This allows us to request that we create the IoStore data for multiple platforms at a time ie '-targetplatform=A+B+etc'

- Replace the manual command line parsing with the ITargetPlatformManagerModule instead, this will keep the behaviour similar to the cook commandlet.
- Note that this will change functionality, so that running the IoStore commandlet without '-targetplatform=XXX' will now attempt to run on all available platforms.
- Made a few things const.

diff --git a/Engine/Source/Developer/IoStoreUtilities/Private/IoStoreUtilities.cpp b/Engine/Source/Developer/IoStoreUtilities/Private/IoStoreUtilities.cpp
--- a/Engine/Source/Developer/IoStoreUtilities/Private/IoStoreUtilities.cpp
+++ b/Engine/Source/Developer/IoStoreUtilities/Private/IoStoreUtilities.cpp
@@ -41,11 +41,11 @@ DEFINE_LOG_CATEGORY_STATIC(LogIoStore, Log, All);
 
 struct FContainerTarget 
 {
-	ITargetPlatform* TargetPlatform;
-	FString CookedDirectory;
-	FString CookedProjectDirectory;
-	FString OutputDirectory;
-	FString ChunkListFile;
+	const ITargetPlatform* TargetPlatform;
+	const FString CookedDirectory;
+	const FString CookedProjectDirectory;
+	const FString OutputDirectory;
+	const FString ChunkListFile;
 };
 
 class FNameMapBuilder
@@ -2149,16 +2149,7 @@ int32 CreateIoStoreContainerFiles(const TCHAR* CmdLine)
 {
 	UE_LOG(LogIoStore, Display, TEXT("==================== IoStore Utils ===================="));
 
-	FString TargetPlatformName;
-	if (FParse::Value(FCommandLine::Get(), TEXT("TargetPlatform="), TargetPlatformName))
-	{
-		UE_LOG(LogIoStore, Display, TEXT("Using target platform: '%s'"), *TargetPlatformName);
-	}
-	else
-	{
-		UE_LOG(LogIoStore, Error, TEXT("No target platform specified [-targetplatform]"));
-		return -1;
-	}
+	const TArray<ITargetPlatform*>& Platforms = GetTargetPlatformManagerRef().GetActiveTargetPlatforms();
 
 	FString OutputDirectory;
 	if (FParse::Value(FCommandLine::Get(), TEXT("OutputDirectory="), OutputDirectory))
@@ -2176,25 +2167,25 @@ int32 CreateIoStoreContainerFiles(const TCHAR* CmdLine)
 		UE_LOG(LogIoStore, Display, TEXT("Using chunk list file: '%s'"), *ChunkListFile);
 	}
 
-	ITargetPlatformManagerModule& TargetPlatformManager = *GetTargetPlatformManager();
-	ITargetPlatform* TargetPlatform = TargetPlatformManager.FindTargetPlatform(TargetPlatformName);
-
-	if (!TargetPlatform)
+	for (const ITargetPlatform* TargetPlatform : Platforms)
 	{
-		UE_LOG(LogIoStore, Display, TEXT("Unknown target platform '%s'"), *TargetPlatformName);
-		return -1;
-	}
+		const FString TargetCookedDirectory = FPaths::ProjectSavedDir() / TEXT("Cooked") / TargetPlatform->PlatformName();
+		const FString TargetCookedProjectDirectory = TargetCookedDirectory / FApp::GetProjectName();
 
-	FString TargetCookedDirectory = FPaths::ProjectSavedDir() / TEXT("Cooked") / TargetPlatform->PlatformName();
-	FString TargetCookedProjectDirectory = TargetCookedDirectory / FApp::GetProjectName();
+		const FString TargetOutputDirectory = OutputDirectory.Len() > 0
+			? OutputDirectory
+			: TargetCookedProjectDirectory / TEXT("Content") / TEXT("Containers");
 
-	FString TargetOutputDirectory = OutputDirectory.Len() > 0
-		? OutputDirectory
-		: TargetCookedProjectDirectory / TEXT("Content") / TEXT("Containers");
+		FContainerTarget Target{ TargetPlatform, TargetCookedDirectory, TargetCookedProjectDirectory, TargetOutputDirectory, ChunkListFile };
 
-	FContainerTarget Target { TargetPlatform, TargetCookedDirectory, TargetCookedProjectDirectory, TargetOutputDirectory, ChunkListFile };
+		UE_LOG(LogIoStore, Display, TEXT("Creating target: '%s' using output directory: '%s'"), *Target.TargetPlatform->PlatformName(), *Target.OutputDirectory);
 
-	UE_LOG(LogIoStore, Display, TEXT("Creating target: '%s' using output directory: '%s'"), *Target.TargetPlatform->PlatformName(), *Target.OutputDirectory);
+		int32 ReturnValue = CreateTarget(Target);
+		if (ReturnValue != 0)
+		{
+			return ReturnValue;
+		}
+	}
 
-	return CreateTarget(Target);
+	return 0;
 }
