IoStore: Fix IoDispatcher shutdown crash by moving shutdown to end of AppPreExit (i.e. after UObjectBaseShutdown and ShutdownAsyncThread, but before UnloadModulesAtShutdown)
Fix parsing of chunkmanifest to handle options and skip creating chunks without files.
Optimize file exists and recursive find to do one single file system traversal only.
Fix "Serializing..." output to display total number of packages

diff --git a/Engine/Source/Developer/IoStoreUtilities/Private/IoStoreUtilities.cpp b/Engine/Source/Developer/IoStoreUtilities/Private/IoStoreUtilities.cpp
--- a/Engine/Source/Developer/IoStoreUtilities/Private/IoStoreUtilities.cpp
+++ b/Engine/Source/Developer/IoStoreUtilities/Private/IoStoreUtilities.cpp
@@ -2,6 +2,7 @@
 
 #include "IoStoreUtilities.h"
 #include "HAL/FileManager.h"
+#include "HAL/PlatformFilemanager.h"
 #include "Hash/CityHash.h"
 #include "Interfaces/IPluginManager.h"
 #include "Interfaces/ITargetPlatformManagerModule.h"
@@ -1101,6 +1102,7 @@ FPackage* AddPackage(const TCHAR* FileName, const TCHAR* CookedDir, TArray<FPack
 
 void SerializePackageData(
 	FIoStoreWriter* IoStoreWriter,
+	int32 TotalPackages,
 	const TArray<FPackage*>& Packages,
 	const FNameMapBuilder& NameMapBuilder,
 	const TArray<FObjectExport>& ObjectExports,
@@ -1112,7 +1114,7 @@ void SerializePackageData(
 {
 	for (FPackage* Package : Packages)
 	{
-		UE_CLOG(Package->GlobalPackageId % 1000 == 0, LogIoStore, Display, TEXT("Serializing %d/%d: '%s'"), Package->GlobalPackageId, Packages.Num(), *Package->Name.ToString());
+		UE_CLOG(Package->GlobalPackageId % 1000 == 0, LogIoStore, Display, TEXT("Serializing %d/%d: '%s'"), Package->GlobalPackageId, TotalPackages, *Package->Name.ToString());
 
 		// Temporary Archive for serializing ImportMap
 		FBufferWriter ImportMapArchive(nullptr, 0, EBufferWriterFlags::AllowResize | EBufferWriterFlags::TakeOwnership);
@@ -1336,6 +1338,26 @@ void SerializePackageData(
 	}
 }
 
+template <class T>
+class FCookedHeaderVisitor : public IPlatformFile::FDirectoryVisitor
+{
+	T& FoundFiles;
+public:
+	FCookedHeaderVisitor(T& InFoundFiles) : FoundFiles(InFoundFiles) {}
+	virtual bool Visit(const TCHAR* FilenameOrDirectory, bool bIsDirectory)
+	{
+		if (bIsDirectory == false)
+		{
+			FString Filename(FilenameOrDirectory);
+			if (Filename.EndsWith(TEXT(".uasset")) || Filename.EndsWith(TEXT(".umap")))
+			{
+				FoundFiles.Add(MoveTemp(Filename));
+			}
+		}
+		return true;
+	}
+};
+
 int32 CreateTarget(const FContainerTarget& Target)
 {
 	TGuardValue<int32> GuardAllowUnversionedContentInEditor(GAllowUnversionedContentInEditor, 1);
@@ -1376,13 +1398,15 @@ int32 CreateTarget(const FContainerTarget& Target)
 	TArray<FInstallChunk> InstallChunks;
 	TArray<FPackage*> Packages;
 	TMap<FName, FPackage*> PackageMap;
+
 	if (Target.ChunkListFile.IsEmpty())
 	{
 		UE_LOG(LogIoStore, Display, TEXT("Searching for .uasset and .umap files..."));
 		TArray<FString> FileNames;
-		IFileManager::Get().FindFilesRecursive(FileNames, *CookedDir, TEXT("*.uasset"), true, false, false);
-		IFileManager::Get().FindFilesRecursive(FileNames, *CookedDir, TEXT("*.umap"), true, false, false);
+		FCookedHeaderVisitor<TArray<FString>> CookedHeaderVistor(FileNames);
+		FPlatformFileManager::Get().GetPlatformFile().IterateDirectoryRecursively(*CookedDir, CookedHeaderVistor);
 		UE_LOG(LogIoStore, Display, TEXT("Found '%d' files"), FileNames.Num());
+
 		FInstallChunk& InstallChunk = InstallChunks.AddDefaulted_GetRef();
 		InstallChunk.Name = TEXT("container0");
 		for (FString& FileName : FileNames)
@@ -1395,59 +1419,126 @@ int32 CreateTarget(const FContainerTarget& Target)
 	}
 	else
 	{
-		TArray<FString> ChunkManifestFileNames;
+		TArray<FString> ChunkFileEntries;
 		FString ChunkFilesDirectory = FPaths::GetPath(Target.ChunkListFile);
-		FFileHelper::LoadFileToStringArray(ChunkManifestFileNames, *Target.ChunkListFile);
-		for (const FString& ChunkManifestFileName : ChunkManifestFileNames)
+		if (!FFileHelper::LoadFileToStringArray(ChunkFileEntries, *Target.ChunkListFile))
 		{
+			UE_LOG(LogIoStore, Error, TEXT("Failed to read chunk list file '%s'."), *Target.ChunkListFile);
+			return -1;
+		}
+
+		UE_LOG(LogIoStore, Display, TEXT("Searching for .uasset and .umap files..."));
+		TSet<FString> FileNames;
+		FCookedHeaderVisitor<TSet<FString>> CookedHeaderVistor(FileNames);
+		FPlatformFileManager::Get().GetPlatformFile().IterateDirectoryRecursively(*CookedDir, CookedHeaderVistor);
+		UE_LOG(LogIoStore, Display, TEXT("Found '%d' files"), FileNames.Num());
+
+		UE_LOG(LogIoStore, Display, TEXT("Parsing chunk list file '%s'"), *Target.ChunkListFile);
+
+		const FString SpaceStr(TEXT(" "));
+		for (FString& ChunkFileEntry : ChunkFileEntries)
+		{
+			FString ChunkFileName;
+			FString Remainder;
+			FString Option;
+			if (!ChunkFileEntry.Split(SpaceStr, &ChunkFileName, &Remainder, ESearchCase::CaseSensitive))
+			{
+				// no options, only a chunk file name
+				ChunkFileName = MoveTemp(ChunkFileEntry);
+				UE_LOG(LogIoStore, Log, TEXT("Parsing chunk file entry for '%s' with no options"), *ChunkFileName);
+			}
+			else
+			{
+				UE_LOG(LogIoStore, Log, TEXT("Parsing chunk file entry for '%s' with options '%s'"), *ChunkFileName, *Remainder);
+			}
+
 			const TCHAR* PakChunkPrefix = TEXT("pakchunk");
 			const int32 PakChunkPrefixLength = 8;
-			if (!ChunkManifestFileName.StartsWith(PakChunkPrefix))
+			if (!ChunkFileName.StartsWith(PakChunkPrefix))
 			{
-				UE_LOG(LogIoStore, Warning, TEXT("Unexpected file name '%s' in '%s'"), *ChunkManifestFileName, *Target.ChunkListFile);
+				UE_LOG(LogIoStore, Error, TEXT("Unexpected file name prefix in '%s'"), *ChunkFileName);
 				continue;
 			}
 			int32 Index = PakChunkPrefixLength;
 			int32 DigitCount = 0;
-			while (Index < ChunkManifestFileName.Len() && FChar::IsDigit(ChunkManifestFileName[Index]))
+			while (Index < ChunkFileName.Len() && FChar::IsDigit(ChunkFileName[Index]))
 			{
 				++DigitCount;
 				++Index;
 			}
 			if (DigitCount <= 0)
 			{
-				UE_LOG(LogIoStore, Warning, TEXT("Unexpected file name '%s' in '%s'"), *ChunkManifestFileName, *Target.ChunkListFile);
+				UE_LOG(LogIoStore, Error, TEXT("Unexpected file name digits in '%s'"), *ChunkFileName);
+				continue;
+			}
+
+			while (Remainder.Split(SpaceStr, &Option, &Remainder, ESearchCase::CaseSensitive) || Remainder.Len() > 0)
+			{
+				if (Option.Len() == 0)
+				{
+					Option = MoveTemp(Remainder);
+				}
+
+				if (FPlatformString::Stricmp(*Option, TEXT("compressed")))
+				{
+					UE_LOG(LogIoStore, Log, TEXT("Ignored option '%s' for chunk '%s'"), *Option, *ChunkFileName);
+				}
+				else if (Option.StartsWith(TEXT("encryptionkeyguid=")))
+				{
+					UE_LOG(LogIoStore, Log, TEXT("Ignored option '%s' for chunk '%s'"), *Option, *ChunkFileName);
+				}
+				else
+				{
+					UE_LOG(LogIoStore, Warning, TEXT("Unexpected option '%s' for chunk '%s'"), *Option, *ChunkFileName);
+				}
+			}
+
+			FString ChunkManifestFullPath = ChunkFilesDirectory / ChunkFileName;
+			TArray<FString> ChunkManifest;
+			if (!FFileHelper::LoadFileToStringArray(ChunkManifest, *ChunkManifestFullPath))
+			{
+				UE_LOG(LogIoStore, Error, TEXT("Failed to read chunk manifest file '%s'."), *ChunkManifestFullPath);
 				continue;
 			}
-			FString ChunkIdString = ChunkManifestFileName.Mid(PakChunkPrefixLength, DigitCount);
+
+			if (ChunkManifest.Num() == 0)
+			{
+				UE_LOG(LogIoStore, Verbose, TEXT("Skipped zero size chunk manifest file '%s'."), *ChunkManifestFullPath);
+				continue;
+			}
+
+			FString ChunkIdString = ChunkFileName.Mid(PakChunkPrefixLength, DigitCount);
 			check(ChunkIdString.IsNumeric());
 			int32 ChunkId;
 			TTypeFromString<int32>::FromString(ChunkId, *ChunkIdString);
 			FInstallChunk& InstallChunk = InstallChunks.AddDefaulted_GetRef();
-			InstallChunk.Name = TEXT("container") + FPaths::GetBaseFilename(ChunkManifestFileName.RightChop(8));
+			InstallChunk.Name = TEXT("container") + FPaths::GetBaseFilename(ChunkFileName.RightChop(8));
 			InstallChunk.ChunkId = ChunkId;
-			FString ChunkManifestFullPath = ChunkFilesDirectory / ChunkManifestFileName;
-			TArray<FString> ChunkManifest;
-			FFileHelper::LoadFileToStringArray(ChunkManifest, *ChunkManifestFullPath);
+
+			UE_LOG(LogIoStore, Log, TEXT("Parsing chunk manifest file '%s'"), *ChunkManifestFullPath);
 			for (const FString& FileNameWithoutExtension : ChunkManifest)
 			{
 				FString RelativePathWithoutExtension = IFileManager::Get().ConvertToRelativePath(*FileNameWithoutExtension);
 				FString FileName = RelativePathWithoutExtension + ".uasset";
-				if (!IFileManager::Get().FileExists(*FileName))
+				if (!FileNames.Contains(FileName))
 				{
 					FileName = RelativePathWithoutExtension + ".umap";
+					if (!FileNames.Contains(FileName))
+					{
+						FileName.Empty();
+					}
 				}
-				if (!IFileManager::Get().FileExists(*FileName))
-				{
-					UE_LOG(LogIoStore, Warning, TEXT("Couldn't find package file (uexp/umap) for package '%s'"), *FileName);
-				}
-				else
+				if (FileName.Len() > 0)
 				{
 					if (FPackage* Package = AddPackage(*FileName, *CookedDir, Packages, PackageMap))
 					{
 						InstallChunk.Packages.Add(Package);
 					}
 				}
+				else
+				{
+					UE_LOG(LogIoStore, Log, TEXT("Ignored file '%s' since it has no corresponding package header file (.map/.uasset) in '%s'."), *FileNameWithoutExtension, *CookedDir);
+				}
 			}
 		}
 	}
@@ -1967,6 +2058,7 @@ int32 CreateTarget(const FContainerTarget& Target)
 #endif
 		SerializePackageData(
 			IoStoreWriter.Get(),
+			PackageMap.Num(),
 			InstallChunk.Packages,
 			NameMapBuilder,
 			ObjectExports,
diff --git a/Engine/Source/Runtime/Core/Private/IO/IoDispatcher.cpp b/Engine/Source/Runtime/Core/Private/IO/IoDispatcher.cpp
--- a/Engine/Source/Runtime/Core/Private/IO/IoDispatcher.cpp
+++ b/Engine/Source/Runtime/Core/Private/IO/IoDispatcher.cpp
@@ -146,6 +146,7 @@ public:
 
 	~FIoDispatcherImpl()
 	{
+		delete Thread;
 		IPlatformChunkInstall* ChunkInstall = FPlatformMisc::GetPlatformChunkInstall();
 		if (ChunkInstall && ChunkDownloadedDelegateHandle.IsValid())
 		{
@@ -442,7 +443,7 @@ private:
 
 	virtual uint32 Run()
 	{
-		while (true)
+		while (!bStopRequested)
 		{
 			EventQueue.ProcessEvents();
 
@@ -455,6 +456,8 @@ private:
 
 	virtual void Stop()
 	{
+		bStopRequested = true;
+		EventQueue.Notify();
 	}
 
 	using FRequestAllocator = TBlockAllocator<FIoRequestImpl, 4096>;
@@ -475,6 +478,7 @@ private:
 	FIoRequestImpl* SubmittedRequestsTail = nullptr;
 	FDelegateHandle ChunkDownloadedDelegateHandle;
 	TArray<TTuple<int32, FIoStoreEnvironment>> PendingInstallChunks;
+	TAtomic<bool> bStopRequested { false };
 };
 
 FIoDispatcher::FIoDispatcher()
