AsyncLoading2: Remove all PostLoad arcs for non-circular packages.
Rename Event_StartPostLoad to Event_PostLoad and let it run all Event_Tick logic.
Change processing of ExternalReadQueue to do at most most one blocking wait per iteration.

diff --git a/Engine/Source/Developer/IoStoreUtilities/Private/IoStoreUtilities.cpp b/Engine/Source/Developer/IoStoreUtilities/Private/IoStoreUtilities.cpp
--- a/Engine/Source/Developer/IoStoreUtilities/Private/IoStoreUtilities.cpp
+++ b/Engine/Source/Developer/IoStoreUtilities/Private/IoStoreUtilities.cpp
@@ -217,8 +217,7 @@ enum EPreloadDependencyType
 enum EEventLoadNode2 : uint8
 {
 	Package_ExportsSerialized,
-	Package_StartPostLoad,
-	Package_Tick,
+	Package_PostLoad,
 	Package_Delete,
 	Package_NumPhases,
 
@@ -388,7 +387,6 @@ struct FPackage
 	int64 ExportBundlesSize = 0;
 
 	bool bHasCircularImportDependencies = false;
-	bool bHasCircularPostLoadDependencies = false;
 
 	TArray<FString> ImportedFullNames;
 
@@ -615,25 +613,19 @@ static void AddScriptArc(FPackage& Package, uint32 GlobalImportIndex, uint32 Exp
 }
 
 static void AddPostLoadArc(FPackage& FromPackage, FPackage& ToPackage)
-{
-	// TArray<FArc>& ExternalArcs = ToPackage.ExternalArcs.FindOrAdd(FromPackage.Name);
-	// ExternalArcs.Add({ EEventLoadNode2::Package_ExportsSerialized, EEventLoadNode2::Package_Tick });
-}
-
-static void AddStartPostLoadArc(FPackage& FromPackage, FPackage& ToPackage)
 {
 	TArray<FArc>& ExternalArcs = ToPackage.ExternalArcs.FindOrAdd(&FromPackage);
-	check(!ExternalArcs.Contains(FArc({EEventLoadNode2::Package_ExportsSerialized, EEventLoadNode2::Package_StartPostLoad})));
-	check(!ExternalArcs.Contains(FArc({EEventLoadNode2::Package_StartPostLoad, EEventLoadNode2::Package_StartPostLoad})));
-	ExternalArcs.Add({ EEventLoadNode2::Package_StartPostLoad, EEventLoadNode2::Package_StartPostLoad });
+	check(!ExternalArcs.Contains(FArc({EEventLoadNode2::Package_ExportsSerialized, EEventLoadNode2::Package_PostLoad})));
+	check(!ExternalArcs.Contains(FArc({EEventLoadNode2::Package_PostLoad, EEventLoadNode2::Package_PostLoad})));
+	ExternalArcs.Add({ EEventLoadNode2::Package_PostLoad, EEventLoadNode2::Package_PostLoad });
 }
 
 static void AddExportsDoneArc(FPackage& FromPackage, FPackage& ToPackage)
 {
 	TArray<FArc>& ExternalArcs = ToPackage.ExternalArcs.FindOrAdd(&FromPackage);
-	check(!ExternalArcs.Contains(FArc({EEventLoadNode2::Package_ExportsSerialized, EEventLoadNode2::Package_StartPostLoad})));
-	check(!ExternalArcs.Contains(FArc({EEventLoadNode2::Package_StartPostLoad, EEventLoadNode2::Package_StartPostLoad})));
-	ExternalArcs.Add({ EEventLoadNode2::Package_ExportsSerialized, EEventLoadNode2::Package_StartPostLoad });
+	check(!ExternalArcs.Contains(FArc({EEventLoadNode2::Package_ExportsSerialized, EEventLoadNode2::Package_PostLoad})));
+	check(!ExternalArcs.Contains(FArc({EEventLoadNode2::Package_PostLoad, EEventLoadNode2::Package_PostLoad})));
+	ExternalArcs.Add({ EEventLoadNode2::Package_ExportsSerialized, EEventLoadNode2::Package_PostLoad });
 }
 
 static void AddUniqueExternalBundleArc(FPackage& FromPackage, uint32 FromBundleIndex, FPackage& ToPackage, uint32 ToBundleIndex)
@@ -650,29 +642,6 @@ static void AddUniqueScriptBundleArc(FPackage& Package, uint32 GlobalImportIndex
 	Package.ScriptArcs.AddUnique({ GlobalImportIndex, NodeIndex });
 }
 
-static void AddPostLoadDependenciesRecursive(FPackage& Package, FPackage& ImportedPackage, TSet<FPackage*>& Visited)
-{
-	if (&ImportedPackage == &Package)
-	{
-		Package.bHasCircularPostLoadDependencies = true;
-		return;
-	}
-
-	bool bIsVisited = false;
-	Visited.Add(&ImportedPackage, &bIsVisited);
-	if (bIsVisited)
-	{
-		return;
-	}
-
-	AddPostLoadArc(ImportedPackage, Package);
-	
-	for (FPackage* DependentPackage : ImportedPackage.ImportedPackages)
-	{
-		AddPostLoadDependenciesRecursive(Package, *DependentPackage, Visited);
-	}
-}
-
 static void AddReachablePackagesRecursive(FPackage& Package, FPackage& PackageWithImports, TSet<FPackage*>& Visited, bool bFirst)
 {
 	if (!bFirst)
@@ -713,6 +682,7 @@ static bool FindNewCircularImportChains(
 {
 	if (&ImportedPackage == &Package)
 	{
+		Package.bHasCircularImportDependencies = true;
 		CurrentChain.SortAndGenerateHash();
 		bool bAlreadyFound = true;
 		CircularChains.AddByHash(CurrentChain.Hash, CurrentChain, &bAlreadyFound);
@@ -765,18 +735,21 @@ static void AddPostLoadDependencies(
 		}
 	}
 
-	for (FPackage* ImportedPackage : Package.ImportedPackages)
+	if (Package.bHasCircularImportDependencies)
 	{
-		if (!DependentPackages.Contains(ImportedPackage))
+		for (FPackage* ImportedPackage : Package.ImportedPackages)
 		{
-			AddStartPostLoadArc(*ImportedPackage, Package);
+			if (!DependentPackages.Contains(ImportedPackage))
+			{
+				AddPostLoadArc(*ImportedPackage, Package);
+			}
 		}
-	}
 
-	DependentPackages.Remove(&Package);
-	for (FPackage* DependentPackage : DependentPackages)
-	{
-		AddExportsDoneArc(*DependentPackage, Package);
+		DependentPackages.Remove(&Package);
+		for (FPackage* DependentPackage : DependentPackages)
+		{
+			AddExportsDoneArc(*DependentPackage, Package);
+		}
 	}
 
 	/*
@@ -1840,34 +1813,6 @@ int32 CreateTarget(const FContainerTarget& Target)
 		}
 	}
 
-#if 0
-	UE_LOG(LogIoStore, Display, TEXT("Simulating postload dependencies for stats..."));
-	for (auto& PackageKV : PackageMap)
-	{
-		FPackage& Package = *PackageKV.Value;
-
-		Visited.Reset();
-		for (FPackage* DependentPackage : Package.ImportedPackages)
-		{
-			AddPostLoadDependenciesRecursive(Package, *DependentPackage, Visited);
-		}
-	}
-
-	UE_LOG(LogIoStore, Display, TEXT("Adding reachable packages for stats..."));
-	for (auto& PackageKV : PackageMap)
-	{
-		FPackage& Package = *PackageKV.Value;
-
-		AddReachablePackagesRecursive(Package, Package, Package.AllReachablePackages, /*bFirst*/ true);
-		Package.bHasCircularImportDependencies = Package.AllReachablePackages.Contains(&Package);
-		// if (Package.bHasCircularImportDependencies)
-		// {
-		// 	UE_LOG(LogIoStore, Display, TEXT("Circular %s: %d reachable packages"), *Package.Name.ToString(),
-		// 		Package.AllReachablePackages.Num());
-		// }
-	}
-#endif
-
 	UE_LOG(LogIoStore, Display, TEXT("Adding optimized postload dependencies..."));
 	for (FPackage* Package : Packages)
 	{
@@ -2163,7 +2108,6 @@ int32 CreateTarget(const FContainerTarget& Target)
 	uint64 TotalExternalArcCount = 0;
 	uint64 NameCount = 0;
 
-	uint64 PackagesWithCircularDependenciesCount = 0;
 	uint64 PackagesWithoutImportDependenciesCount = 0;
 	uint64 PackagesWithoutPreloadDependenciesCount = 0;
 	uint64 BundleCount = 0;
@@ -2194,7 +2138,6 @@ int32 CreateTarget(const FContainerTarget& Target)
 		TotalInternalArcCount += Package.InternalArcs.Num();
 		ImportedPackagesCount += Package.ImportedPackages.Num();
 		NameCount += Package.NameMap.Num();
-		PackagesWithCircularDependenciesCount += Package.bHasCircularPostLoadDependencies;
 		PackagesWithoutPreloadDependenciesCount += Package.ImportedPreloadPackages.Num() == 0;
 		PackagesWithoutImportDependenciesCount += Package.ImportedPackages.Num() == 0;
 
@@ -2215,7 +2158,7 @@ int32 CreateTarget(const FContainerTarget& Target)
 
 	UE_LOG(LogIoStore, Display, TEXT("-------------------- IoStore Summary: %s --------------------"), *Target.TargetPlatform->PlatformName());
 	UE_LOG(LogIoStore, Display, TEXT("Packages: %8d total, %d circular dependencies, %d no preload dependencies, %d no import dependencies"),
-		PackageMap.Num(), PackagesWithCircularDependenciesCount, PackagesWithoutPreloadDependenciesCount, PackagesWithoutImportDependenciesCount);
+		PackageMap.Num(), CircularPackagesCount, PackagesWithoutPreloadDependenciesCount, PackagesWithoutImportDependenciesCount);
 	UE_LOG(LogIoStore, Display, TEXT("Bundles:  %8d total, %d entries, %d export objects"), BundleCount, BundleEntryCount, GlobalExports.Num());
 
 	UE_LOG(LogIoStore, Display, TEXT("IoStore: %8.2f MB GlobalNames, %d unique names"), (double)GlobalNamesMB, NameMapBuilder.GetNameMap().Num());
diff --git a/Engine/Source/Runtime/CoreUObject/Private/Serialization/AsyncLoading2.cpp b/Engine/Source/Runtime/CoreUObject/Private/Serialization/AsyncLoading2.cpp
--- a/Engine/Source/Runtime/CoreUObject/Private/Serialization/AsyncLoading2.cpp
+++ b/Engine/Source/Runtime/CoreUObject/Private/Serialization/AsyncLoading2.cpp
@@ -828,8 +828,7 @@ enum class EAsyncPackageLoadingState2 : uint8
 enum EEventLoadNode2 : uint8
 {
 	Package_ExportsSerialized,
-	Package_StartPostLoad,
-	Package_Tick,
+	Package_PostLoad,
 	Package_Delete,
 	Package_NumPhases,
 
@@ -1339,8 +1338,7 @@ public:
 
 	static EAsyncPackageState::Type Event_ProcessExportBundle(FAsyncPackage2* Package, int32 ExportBundleIndex);
 	static EAsyncPackageState::Type Event_ExportsDone(FAsyncPackage2* Package, int32);
-	static EAsyncPackageState::Type Event_StartPostLoad(FAsyncPackage2* Package, int32);
-	static EAsyncPackageState::Type Event_Tick(FAsyncPackage2* Package, int32);
+	static EAsyncPackageState::Type Event_PostLoad(FAsyncPackage2* Package, int32);
 	static EAsyncPackageState::Type Event_Delete(FAsyncPackage2* Package, int32);
 
 	void EventDrivenCreateExport(int32 LocalExportIndex);
@@ -2988,20 +2986,13 @@ EAsyncPackageState::Type FAsyncPackage2::Event_ExportsDone(FAsyncPackage2* Packa
 {
 	TRACE_CPUPROFILER_EVENT_SCOPE(Event_ExportsDone);
 
-	Package->GetNode(EEventLoadNode2::Package_StartPostLoad)->ReleaseBarrier();
+	Package->GetNode(EEventLoadNode2::Package_PostLoad)->ReleaseBarrier();
 	return EAsyncPackageState::Complete;
 }
 
-EAsyncPackageState::Type FAsyncPackage2::Event_StartPostLoad(FAsyncPackage2* Package, int32)
+EAsyncPackageState::Type FAsyncPackage2::Event_PostLoad(FAsyncPackage2* Package, int32)
 {
-	TRACE_CPUPROFILER_EVENT_SCOPE(Event_StartPostLoad);
-	Package->GetNode(EEventLoadNode2::Package_Tick)->ReleaseBarrier();
-	return EAsyncPackageState::Complete;
-}
-
-EAsyncPackageState::Type FAsyncPackage2::Event_Tick(FAsyncPackage2* Package, int32)
-{
-	TRACE_CPUPROFILER_EVENT_SCOPE(Event_Tick);
+	TRACE_CPUPROFILER_EVENT_SCOPE(Event_PostLoad);
 
 	check(!Package->HasFinishedLoading());
 	check(Package->ExternalReadDependencies.Num() == 0);
@@ -3475,8 +3466,7 @@ FAsyncLoadingThread2Impl::FAsyncLoadingThread2Impl(FIoDispatcher& InIoDispatcher
 
 	EventSpecs.AddDefaulted(EEventLoadNode2::Package_NumPhases + EEventLoadNode2::ExportBundle_NumPhases);
 	EventSpecs[EEventLoadNode2::Package_ExportsSerialized] = { &FAsyncPackage2::Event_ExportsDone, &AsyncEventQueue, true };
-	EventSpecs[EEventLoadNode2::Package_StartPostLoad] = { &FAsyncPackage2::Event_StartPostLoad, &AsyncEventQueue, true };
-	EventSpecs[EEventLoadNode2::Package_Tick] = { &FAsyncPackage2::Event_Tick, &EventQueue, false };
+	EventSpecs[EEventLoadNode2::Package_PostLoad] = { &FAsyncPackage2::Event_PostLoad, &AsyncEventQueue, true };
 	EventSpecs[EEventLoadNode2::Package_Delete] = { &FAsyncPackage2::Event_Delete, &AsyncEventQueue, false };
 
 	EventSpecs[EEventLoadNode2::Package_NumPhases + EEventLoadNode2::ExportBundle_Process] = { &FAsyncPackage2::Event_ProcessExportBundle, &ProcessExportBundlesEventQueue, false };
@@ -3657,33 +3647,8 @@ uint32 FAsyncLoadingThread2Impl::Run()
 				TRACE_CPUPROFILER_EVENT_SCOPE(AsyncLoadingTime);
 				do
 				{
-					const bool bPreviousDidSomething = bDidSomething;
 					bDidSomething = false;
 
-					bool bDidExternalRead = false;
-					do
-					{
-						bDidExternalRead = false;
-						FAsyncPackage2* Package = nullptr;
-						if (ExternalReadQueue.Peek(Package))
-						{
-							TRACE_CPUPROFILER_EVENT_SCOPE(ProcessExternalReads);
-
-							FAsyncPackage2::EExternalReadAction Action =
-								bPreviousDidSomething ?
-								FAsyncPackage2::ExternalReadAction_Poll :
-								FAsyncPackage2::ExternalReadAction_Wait;
-
-							EAsyncPackageState::Type Result = Package->ProcessExternalReads(Action);
-							if (Result == EAsyncPackageState::Complete)
-							{
-								ExternalReadQueue.Pop();
-								bDidExternalRead = true;
-								bDidSomething = true;
-							}
-						}
-					} while (bDidExternalRead);
-
 					if (QueuedPackagesCounter)
 					{
 						if (CreateAsyncPackagesFromQueue())
@@ -3722,6 +3687,33 @@ uint32 FAsyncLoadingThread2Impl::Run()
 						bDidSomething = true;
 						break;
 					}
+
+					{
+						bool bDidExternalRead = false;
+						do
+						{
+							bDidExternalRead = false;
+							FAsyncPackage2* Package = nullptr;
+							if (ExternalReadQueue.Peek(Package))
+							{
+								TRACE_CPUPROFILER_EVENT_SCOPE(ProcessExternalReads);
+
+								FAsyncPackage2::EExternalReadAction Action =
+									bDidSomething ?
+									FAsyncPackage2::ExternalReadAction_Poll :
+									FAsyncPackage2::ExternalReadAction_Wait;
+
+								EAsyncPackageState::Type Result = Package->ProcessExternalReads(Action);
+								if (Result == EAsyncPackageState::Complete)
+								{
+									ExternalReadQueue.Pop();
+									bDidExternalRead = true;
+									bDidSomething = true;
+								}
+							}
+						} while (bDidExternalRead);
+					}
+
 				} while (bDidSomething);
 			}
 
@@ -4048,11 +4040,9 @@ void FAsyncPackage2::CreateNodes(const FAsyncLoadEventSpec* EventSpecs)
 		}
 
 		FEventLoadNode2* ExportsSerializedNode = PackageNodes + EEventLoadNode2::Package_ExportsSerialized;
-		FEventLoadNode2* StartPostLoadNode = PackageNodes + EEventLoadNode2::Package_StartPostLoad;
-		FEventLoadNode2* TickNode = PackageNodes + EEventLoadNode2::Package_Tick;
+		FEventLoadNode2* StartPostLoadNode = PackageNodes + EEventLoadNode2::Package_PostLoad;
 
 		StartPostLoadNode->AddBarrier();
-		TickNode->AddBarrier();
 
 		FEventLoadNode2* DeleteNode = PackageNodes + EEventLoadNode2::Package_Delete;
 		DeleteNode->AddBarrier();
