Create one IoDispatcher container file per pak file and let the pak file system handle mounting of IoStores

diff --git a/Engine/Source/Developer/IoStoreUtilities/Private/IoStoreUtilities.cpp b/Engine/Source/Developer/IoStoreUtilities/Private/IoStoreUtilities.cpp
--- a/Engine/Source/Developer/IoStoreUtilities/Private/IoStoreUtilities.cpp
+++ b/Engine/Source/Developer/IoStoreUtilities/Private/IoStoreUtilities.cpp
@@ -39,15 +39,31 @@ DEFINE_LOG_CATEGORY_STATIC(LogIoStore, Log, All);
 #define OUTPUT_NAMEMAP_CSV 0
 #define OUTPUT_IMPORTMAP_CSV 0
 #define SKIP_WRITE_CONTAINER 0
-#define SKIP_BULKDATA 0
 
-struct FContainerTarget 
+struct FContainerSourceFile 
 {
-	const ITargetPlatform* TargetPlatform;
-	const FString CookedDirectory;
-	const FString CookedProjectDirectory;
-	const FString OutputDirectory;
-	const FString ChunkListFile;
+	FString Path;
+};
+
+struct FContainerSourceSpec
+{
+	FString OutputPath;
+	TArray<FContainerSourceFile> SourceFiles;
+};
+
+struct FPackage;
+
+struct FContainerTargetFile
+{
+	FPackage* Package;
+	FString SourcePath;
+	bool bIsBulkData;
+};
+
+struct FContainerTargetSpec
+{
+	FIoStoreWriter* IoStoreWriter;
+	TArray<FContainerTargetFile> TargetFiles;
 };
 
 class FNameMapBuilder
@@ -220,8 +236,6 @@ struct FExportBundle
 	uint32 LoadOrder;
 };
 
-struct FPackage;
-
 struct FPackageGraphNode
 {
 	FPackage* Package = nullptr;
@@ -329,7 +343,6 @@ struct FPackage
 {
 	FName Name;
 	FString FileName;
-	FString RelativeFileName;
 	int32 GlobalPackageId = 0;
 	uint32 PackageFlags = 0;
 	int32 NameCount = 0;
@@ -438,13 +451,6 @@ struct FCircularImportChain
 	}
 };
 
-struct FInstallChunk
-{
-	FString Name;
-	int32 ChunkId;
-	TArray<FPackage*> Packages;
-};
-
 TArray<FPackage*> FPackageGraph::TopologicalSort() const
 {
 	TMultiMap<FPackageGraphNode*, FPackageGraphNode*> EdgesCopy = Edges;
@@ -996,7 +1002,7 @@ static void FindExport(TArray<FExportData>& GlobalExports, TMap<FString, int32>&
 	}
 };
 
-FPackage* AddPackage(const TCHAR* FileName, const TCHAR* CookedDir, TArray<FPackage*>& Packages, TMap<FName, FPackage*>& PackageMap)
+FPackage* FindOrAddPackage(const TCHAR* FileName, const TCHAR* CookedDir, TArray<FPackage*>& Packages, TMap<FName, FPackage*>& PackageMap)
 {
 	FString RelativeFileName = FileName;
 	RelativeFileName.RemoveFromStart(CookedDir);
@@ -1014,16 +1020,10 @@ FPackage* AddPackage(const TCHAR* FileName, const TCHAR* CookedDir, TArray<FPack
 	FName PackageFName = *PackageName;
 
 	FPackage* Package = PackageMap.FindRef(PackageFName);
-	if (Package)
-	{
-		UE_LOG(LogIoStore, Warning, TEXT("Package in multiple pakchunks: '%s'"), *PackageFName.ToString());
-	}
-	else
+	if (!Package)
 	{
 		Package = new FPackage();
 		Package->Name = PackageFName;
-		Package->FileName = FileName;
-		Package->RelativeFileName = MoveTemp(RelativeFileName);
 		Package->GlobalPackageId = Packages.Num();
 		Packages.Add(Package);
 		PackageMap.Add(PackageFName, Package);
@@ -1033,20 +1033,15 @@ FPackage* AddPackage(const TCHAR* FileName, const TCHAR* CookedDir, TArray<FPack
 }
 
 void SerializePackageData(
-	FIoStoreWriter* IoStoreWriter,
-	int32 TotalPackages,
-	const TArray<FPackage*>& Packages,
+	const FContainerTargetFile& TargetFile,
 	const FNameMapBuilder& NameMapBuilder,
 	const TArray<FObjectExport>& ObjectExports,
 	const TArray<FExportData>& GlobalExports,
 	const TMap<FString, int32>& GlobalImportsByFullName,
 	FExportBundleMetaEntry* ExportBundleMetaEntries,
-	const FPackageStoreBulkDataManifest& BulkDataManifest,
-	bool bWithBulkDataManifest)
+	FIoStoreWriter* IoStoreWriter)
 {
-	for (FPackage* Package : Packages)
-	{
-		UE_CLOG(Package->GlobalPackageId % 1000 == 0, LogIoStore, Display, TEXT("Serializing %d/%d: '%s'"), Package->GlobalPackageId, TotalPackages, *Package->Name.ToString());
+	FPackage* Package = TargetFile.Package;
 
 		// Temporary Archive for serializing ImportMap
 		FBufferWriter ImportMapArchive(nullptr, 0, EBufferWriterFlags::AllowResize | EBufferWriterFlags::TakeOwnership);
@@ -1262,54 +1257,14 @@ void SerializePackageData(
 				FMemory::Free(ExportsBuffer);
 #endif
 			}
-
-#if !SKIP_BULKDATA
-			// Bulk chunks
-			if (bWithBulkDataManifest)
-			{
-				FString BulkFileName = FPaths::ChangeExtension(Package->FileName, TEXT(".ubulk"));
-				FPaths::NormalizeFilename(BulkFileName);
-
-				WriteBulkData(BulkFileName, EIoChunkType::BulkData, *Package, BulkDataManifest, IoStoreWriter);
-
-				FString OptionalBulkFileName = FPaths::ChangeExtension(Package->FileName, TEXT(".uptnl"));
-				FPaths::NormalizeFilename(OptionalBulkFileName);
-
-				WriteBulkData(OptionalBulkFileName, EIoChunkType::OptionalBulkData, *Package, BulkDataManifest, IoStoreWriter);
-			}
-#endif
-		}
 	}
 }
 
-template <class T>
-class FCookedHeaderVisitor : public IPlatformFile::FDirectoryVisitor
-{
-	T& FoundFiles;
-public:
-	FCookedHeaderVisitor(T& InFoundFiles) : FoundFiles(InFoundFiles) {}
-	virtual bool Visit(const TCHAR* FilenameOrDirectory, bool bIsDirectory)
-	{
-		if (bIsDirectory == false)
-		{
-			FString Filename(FilenameOrDirectory);
-			if (Filename.EndsWith(TEXT(".uasset")) || Filename.EndsWith(TEXT(".umap")))
-			{
-				FoundFiles.Add(MoveTemp(Filename));
-			}
-		}
-		return true;
-	}
-};
-
-int32 CreateTarget(const FContainerTarget& Target)
+int32 CreateTarget(const TCHAR* OutputDir, const TCHAR* CookedDir, const TArray<FContainerSourceSpec> Containers)
 {
 	TGuardValue<int32> GuardAllowUnversionedContentInEditor(GAllowUnversionedContentInEditor, 1);
 
-	const FString CookedDir = Target.CookedDirectory;
-	const FString OutputDir = Target.OutputDirectory;
-	
-	FPackageStoreBulkDataManifest BulkDataManifest(Target.CookedProjectDirectory);
+	FPackageStoreBulkDataManifest BulkDataManifest(FString(CookedDir) / FApp::GetProjectName());
 	const bool bWithBulkDataManifest = BulkDataManifest.Load();
 	if (bWithBulkDataManifest)
 	{
@@ -1339,149 +1294,61 @@ int32 CreateTarget(const FContainerTarget& Target)
 	FExportGraph ExportGraph;
 	TArray<FExportBundleMetaEntry> ExportBundleMetaEntries;
 
-	TArray<FInstallChunk> InstallChunks;
 	TArray<FPackage*> Packages;
 	TMap<FName, FPackage*> PackageMap;
 
-	if (Target.ChunkListFile.IsEmpty())
+	TArray<FIoStoreWriter*> IoStoreWriters;
+	FIoStoreWriter* GlobalIoStoreWriter;
+	FIoStoreEnvironment GlobalIoStoreEnv;
 	{
-		UE_LOG(LogIoStore, Display, TEXT("Searching for .uasset and .umap files..."));
-		TArray<FString> FileNames;
-		FCookedHeaderVisitor<TArray<FString>> CookedHeaderVistor(FileNames);
-		FPlatformFileManager::Get().GetPlatformFile().IterateDirectoryRecursively(*CookedDir, CookedHeaderVistor);
-		UE_LOG(LogIoStore, Display, TEXT("Found '%d' files"), FileNames.Num());
+		FString GlobalOutputPath = FString(OutputDir) / TEXT("global");
+		GlobalIoStoreEnv.InitializeFileEnvironment(*GlobalOutputPath);
+		GlobalIoStoreWriter = new FIoStoreWriter(GlobalIoStoreEnv);
+		IoStoreWriters.Add(GlobalIoStoreWriter);
+#if !SKIP_WRITE_CONTAINER
+		FIoStatus IoStatus = GlobalIoStoreWriter->Initialize();
+		check(IoStatus.IsOk());
+#endif
+	}
+
+	TArray<FContainerTargetSpec> ContainerTargets;
 
-		FInstallChunk& InstallChunk = InstallChunks.AddDefaulted_GetRef();
-		InstallChunk.Name = TEXT("container0");
-		for (FString& FileName : FileNames)
+	for (const FContainerSourceSpec& Container : Containers)
+	{
+		FContainerTargetSpec& ContainerTarget = ContainerTargets.AddDefaulted_GetRef();
+		if (Container.OutputPath.IsEmpty())
 		{
-			if (FPackage* Package = AddPackage(*FileName, *CookedDir, Packages, PackageMap))
-			{
-				InstallChunk.Packages.Add(Package);
-			}
+			ContainerTarget.IoStoreWriter = GlobalIoStoreWriter;
 		}
-	}
-	else
-	{
-		TArray<FString> ChunkFileEntries;
-		FString ChunkFilesDirectory = FPaths::GetPath(Target.ChunkListFile);
-		if (!FFileHelper::LoadFileToStringArray(ChunkFileEntries, *Target.ChunkListFile))
+		else
 		{
-			UE_LOG(LogIoStore, Error, TEXT("Failed to read chunk list file '%s'."), *Target.ChunkListFile);
-			return -1;
+			FIoStoreEnvironment ContainerIoStoreEnv;
+			ContainerIoStoreEnv.InitializeFileEnvironment(Container.OutputPath);
+			ContainerTarget.IoStoreWriter = new FIoStoreWriter(ContainerIoStoreEnv);
+			IoStoreWriters.Add(ContainerTarget.IoStoreWriter);
+#if !SKIP_WRITE_CONTAINER
+			FIoStatus IoStatus = ContainerTarget.IoStoreWriter->Initialize();
+			check(IoStatus.IsOk());
+#endif
 		}
-
-		UE_LOG(LogIoStore, Display, TEXT("Searching for .uasset and .umap files..."));
-		TSet<FString> FileNames;
-		FCookedHeaderVisitor<TSet<FString>> CookedHeaderVistor(FileNames);
-		FPlatformFileManager::Get().GetPlatformFile().IterateDirectoryRecursively(*CookedDir, CookedHeaderVistor);
-		UE_LOG(LogIoStore, Display, TEXT("Found '%d' files"), FileNames.Num());
-
-		UE_LOG(LogIoStore, Display, TEXT("Parsing chunk list file '%s'"), *Target.ChunkListFile);
-
-		const FString SpaceStr(TEXT(" "));
-		for (FString& ChunkFileEntry : ChunkFileEntries)
+		for (const FContainerSourceFile& SourceFile : Container.SourceFiles)
 		{
-			FString ChunkFileName;
-			FString Remainder;
-			FString Option;
-			if (!ChunkFileEntry.Split(SpaceStr, &ChunkFileName, &Remainder, ESearchCase::CaseSensitive))
-			{
-				// no options, only a chunk file name
-				ChunkFileName = MoveTemp(ChunkFileEntry);
-				UE_LOG(LogIoStore, Log, TEXT("Parsing chunk file entry for '%s' with no options"), *ChunkFileName);
-			}
-			else
-			{
-				UE_LOG(LogIoStore, Log, TEXT("Parsing chunk file entry for '%s' with options '%s'"), *ChunkFileName, *Remainder);
-			}
-
-			const TCHAR* PakChunkPrefix = TEXT("pakchunk");
-			const int32 PakChunkPrefixLength = 8;
-			if (!ChunkFileName.StartsWith(PakChunkPrefix))
-			{
-				UE_LOG(LogIoStore, Error, TEXT("Unexpected file name prefix in '%s'"), *ChunkFileName);
-				continue;
-			}
-			int32 Index = PakChunkPrefixLength;
-			int32 DigitCount = 0;
-			while (Index < ChunkFileName.Len() && FChar::IsDigit(ChunkFileName[Index]))
-			{
-				++DigitCount;
-				++Index;
-			}
-			if (DigitCount <= 0)
-			{
-				UE_LOG(LogIoStore, Error, TEXT("Unexpected file name digits in '%s'"), *ChunkFileName);
-				continue;
-			}
-
-			while (Remainder.Split(SpaceStr, &Option, &Remainder, ESearchCase::CaseSensitive) || Remainder.Len() > 0)
-			{
-				if (Option.Len() == 0)
-				{
-					Option = MoveTemp(Remainder);
-				}
-
-				if (FPlatformString::Stricmp(*Option, TEXT("compressed")))
-				{
-					UE_LOG(LogIoStore, Log, TEXT("Ignored option '%s' for chunk '%s'"), *Option, *ChunkFileName);
-				}
-				else if (Option.StartsWith(TEXT("encryptionkeyguid=")))
-				{
-					UE_LOG(LogIoStore, Log, TEXT("Ignored option '%s' for chunk '%s'"), *Option, *ChunkFileName);
-				}
-				else
-				{
-					UE_LOG(LogIoStore, Warning, TEXT("Unexpected option '%s' for chunk '%s'"), *Option, *ChunkFileName);
-				}
-			}
-
-			FString ChunkManifestFullPath = ChunkFilesDirectory / ChunkFileName;
-			TArray<FString> ChunkManifest;
-			if (!FFileHelper::LoadFileToStringArray(ChunkManifest, *ChunkManifestFullPath))
-			{
-				UE_LOG(LogIoStore, Error, TEXT("Failed to read chunk manifest file '%s'."), *ChunkManifestFullPath);
-				continue;
-			}
-
-			if (ChunkManifest.Num() == 0)
+			FPackage* Package = FindOrAddPackage(*SourceFile.Path, CookedDir, Packages, PackageMap);
+			if (Package)
 			{
-				UE_LOG(LogIoStore, Verbose, TEXT("Skipped zero size chunk manifest file '%s'."), *ChunkManifestFullPath);
-				continue;
-			}
-
-			FString ChunkIdString = ChunkFileName.Mid(PakChunkPrefixLength, DigitCount);
-			check(ChunkIdString.IsNumeric());
-			int32 ChunkId;
-			TTypeFromString<int32>::FromString(ChunkId, *ChunkIdString);
-			FInstallChunk& InstallChunk = InstallChunks.AddDefaulted_GetRef();
-			InstallChunk.Name = TEXT("container") + FPaths::GetBaseFilename(ChunkFileName.RightChop(8));
-			InstallChunk.ChunkId = ChunkId;
+				FContainerTargetFile& TargetFile = ContainerTarget.TargetFiles.AddDefaulted_GetRef();
+				TargetFile.SourcePath = SourceFile.Path;
+				TargetFile.Package = Package;
 
-			UE_LOG(LogIoStore, Log, TEXT("Parsing chunk manifest file '%s'"), *ChunkManifestFullPath);
-			for (const FString& FileNameWithoutExtension : ChunkManifest)
-			{
-				FString RelativePathWithoutExtension = IFileManager::Get().ConvertToRelativePath(*FileNameWithoutExtension);
-				FString FileName = RelativePathWithoutExtension + ".uasset";
-				if (!FileNames.Contains(FileName))
-				{
-					FileName = RelativePathWithoutExtension + ".umap";
-					if (!FileNames.Contains(FileName))
-					{
-						FileName.Empty();
-					}
-				}
-				if (FileName.Len() > 0)
+				FString Extension = FPaths::GetExtension(SourceFile.Path);
+				if (Extension == TEXT("uasset") || Extension == TEXT("umap"))
 				{
-					if (FPackage* Package = AddPackage(*FileName, *CookedDir, Packages, PackageMap))
-					{
-						InstallChunk.Packages.Add(Package);
-					}
+					TargetFile.bIsBulkData = false;
+					Package->FileName = SourceFile.Path;
 				}
 				else
 				{
-					UE_LOG(LogIoStore, Log, TEXT("Ignored file '%s' since it has no corresponding package header file (.map/.uasset) in '%s'."), *FileNameWithoutExtension, *CookedDir);
+					TargetFile.bIsBulkData = true;
 				}
 			}
 		}
@@ -1951,47 +1818,50 @@ int32 CreateTarget(const FContainerTarget& Target)
 
 	UE_LOG(LogIoStore, Display, TEXT("Serializing..."));
 
-	TUniquePtr<FIoStoreWriter> GlobalIoStoreWriter;
-	FIoStoreEnvironment GlobalIoStoreEnv;
+	int32 TotalFileCount = 0;
+	for (FContainerTargetSpec& ContainerTarget : ContainerTargets)
 	{
-		GlobalIoStoreEnv.InitializeFileEnvironment(OutputDir);
-		GlobalIoStoreWriter = MakeUnique<FIoStoreWriter>(GlobalIoStoreEnv);
-#if !SKIP_WRITE_CONTAINER
-		FIoStatus IoStatus = GlobalIoStoreWriter->Initialize();
-		check(IoStatus.IsOk());
-#endif
+		TotalFileCount += ContainerTarget.TargetFiles.Num();
 	}
-
-	FIoStoreInstallManifest Manifest;
-	for (const FInstallChunk& InstallChunk : InstallChunks)
+	int32 CurrentFileIndex = 0;
+	for (FContainerTargetSpec& ContainerTarget : ContainerTargets)
 	{
-		TUniquePtr<FIoStoreWriter> IoStoreWriter;
-		FIoStoreEnvironment InstallChunkIoStoreEnv(GlobalIoStoreEnv, InstallChunk.Name);
-		IoStoreWriter = MakeUnique<FIoStoreWriter>(InstallChunkIoStoreEnv);
-#if !SKIP_WRITE_CONTAINER
-		FIoStatus IoStatus = IoStoreWriter->Initialize();
-		check(IoStatus.IsOk());
-#endif
-		SerializePackageData(
-			IoStoreWriter.Get(),
-			PackageMap.Num(),
-			InstallChunk.Packages,
-			NameMapBuilder,
-			ObjectExports,
-			GlobalExports,
-			GlobalImportsByFullName,
-			ExportBundleMetaEntries.GetData(),
-			BulkDataManifest,
-			bWithBulkDataManifest);
-		FIoStoreInstallManifest::FEntry& ManifestEntry = Manifest.EditEntries().AddDefaulted_GetRef();
-		ManifestEntry.InstallChunkId = InstallChunk.ChunkId;
-		ManifestEntry.PartitionName = InstallChunk.Name;
-	}
-	TUniquePtr<FLargeMemoryWriter> ManifestArchive = MakeUnique<FLargeMemoryWriter>(0, true);
-	*ManifestArchive << Manifest;
-
-	GlobalIoStoreWriter->Append(CreateIoChunkId(0, 0, EIoChunkType::InstallManifest), FIoBuffer(FIoBuffer::Wrap, ManifestArchive->GetData(), ManifestArchive->TotalSize()));
+		for (FContainerTargetFile& TargetFile : ContainerTarget.TargetFiles)
+		{
+			UE_CLOG(CurrentFileIndex % 1000 == 0, LogIoStore, Display, TEXT("Serializing %d/%d: '%s'"), CurrentFileIndex, TotalFileCount, *TargetFile.SourcePath);
+			++CurrentFileIndex;
+			if (TargetFile.bIsBulkData)
+			{
+				if (bWithBulkDataManifest)
+				{
+					FString NormalizedSourcePath = TargetFile.SourcePath;
+					FPaths::NormalizeFilename(NormalizedSourcePath);
 
+					FString Extension = FPaths::GetExtension(NormalizedSourcePath);
+					if (Extension == TEXT("uptnl"))
+					{
+						WriteBulkData(NormalizedSourcePath, EIoChunkType::OptionalBulkData, *TargetFile.Package, BulkDataManifest, ContainerTarget.IoStoreWriter);
+					}
+					else
+					{
+						WriteBulkData(NormalizedSourcePath, EIoChunkType::BulkData, *TargetFile.Package, BulkDataManifest, ContainerTarget.IoStoreWriter);
+					}
+				}
+			}
+			else
+			{
+				SerializePackageData(
+					TargetFile,
+					NameMapBuilder,
+					ObjectExports,
+					GlobalExports,
+					GlobalImportsByFullName,
+					ExportBundleMetaEntries.GetData(),
+					ContainerTarget.IoStoreWriter);
+			}
+		}
+	}
+	
 	int32 StoreTocByteCount = StoreTocArchive->TotalSize();
 	int32 ImportedPackagesByteCount = ImportedPackagesArchive->TotalSize();
 	int32 GlobalImportNamesByteCount = GlobalImportNamesArchive->TotalSize();
@@ -2061,6 +1931,12 @@ int32 CreateTarget(const FContainerTarget& Target)
 #endif
 	}
 
+	for (FIoStoreWriter* IoStoreWriter : IoStoreWriters)
+	{
+		delete IoStoreWriter;
+	}
+	IoStoreWriters.Empty();
+
 	UE_LOG(LogIoStore, Display, TEXT("Calculating stats..."));
 	uint64 UExpSize = 0;
 	uint64 UAssetSize = 0;
@@ -2127,7 +2003,7 @@ int32 CreateTarget(const FContainerTarget& Target)
 
 	PackageHeaderSize = PackageSummarySize + NameMapSize + ImportMapSize + ExportMapSize + UGraphSize;
 
-	UE_LOG(LogIoStore, Display, TEXT("-------------------- IoStore Summary: %s --------------------"), *Target.TargetPlatform->PlatformName());
+	UE_LOG(LogIoStore, Display, TEXT("-------------------- IoStore Summary --------------------"));
 	UE_LOG(LogIoStore, Display, TEXT("Packages: %8d total, %d circular dependencies, %d no preload dependencies, %d no import dependencies"),
 		PackageMap.Num(), CircularPackagesCount, PackagesWithoutPreloadDependenciesCount, PackagesWithoutImportDependenciesCount);
 	UE_LOG(LogIoStore, Display, TEXT("Bundles:  %8d total, %d entries, %d export objects"), BundleCount, BundleEntryCount, GlobalExports.Num());
@@ -2155,43 +2031,149 @@ int32 CreateIoStoreContainerFiles(const TCHAR* CmdLine)
 {
 	UE_LOG(LogIoStore, Display, TEXT("==================== IoStore Utils ===================="));
 
-	const TArray<ITargetPlatform*>& Platforms = GetTargetPlatformManagerRef().GetActiveTargetPlatforms();
-
-	FString OutputDirectory;
-	if (FParse::Value(FCommandLine::Get(), TEXT("OutputDirectory="), OutputDirectory))
-	{
-		UE_LOG(LogIoStore, Display, TEXT("Using output directory: '%s'"), *OutputDirectory);
-	}
-	else
+	FString CommandListFile;
+	if (FParse::Value(FCommandLine::Get(), TEXT("Commands="), CommandListFile))
 	{
-		UE_LOG(LogIoStore, Display, TEXT("No output directory specified, using project's cooked folder"));
-	}
+		UE_LOG(LogIoStore, Display, TEXT("Using command list file: '%s'"), *CommandListFile);
 
-	FString ChunkListFile;
-	if (FParse::Value(FCommandLine::Get(), TEXT("ChunkListFile="), ChunkListFile))
-	{
-		UE_LOG(LogIoStore, Display, TEXT("Using chunk list file: '%s'"), *ChunkListFile);
-	}
+		FString OutputDirectory;
+		if (!FParse::Value(FCommandLine::Get(), TEXT("OutputDirectory="), OutputDirectory))
+		{
+			UE_LOG(LogIoStore, Error, TEXT("OutputDirectory must be specified"));
+			return 1;
+		}
 
-	for (const ITargetPlatform* TargetPlatform : Platforms)
-	{
-		const FString TargetCookedDirectory = FPaths::ProjectSavedDir() / TEXT("Cooked") / TargetPlatform->PlatformName();
-		const FString TargetCookedProjectDirectory = TargetCookedDirectory / FApp::GetProjectName();
+		FString CookedDirectory;
+		if (!FParse::Value(FCommandLine::Get(), TEXT("CookedDirectory="), CookedDirectory))
+		{
+			UE_LOG(LogIoStore, Error, TEXT("CookedDirectory must be specified"));
+			return 1;
+		}
 
-		const FString TargetOutputDirectory = OutputDirectory.Len() > 0
-			? OutputDirectory
-			: TargetCookedProjectDirectory / TEXT("Content") / TEXT("Containers");
+		TArray<FString> Commands;
+		if (!FFileHelper::LoadFileToStringArray(Commands, *CommandListFile))
+		{
+			UE_LOG(LogIoStore, Error, TEXT("Failed to read command list file '%s'."), *CommandListFile);
+			return -1;
+		}
 
-		FContainerTarget Target{ TargetPlatform, TargetCookedDirectory, TargetCookedProjectDirectory, TargetOutputDirectory, ChunkListFile };
+		TArray<FContainerSourceSpec> Containers;
+		for (const FString& Command : Commands)
+		{
+			FContainerSourceSpec& ContainerSpec = Containers.AddDefaulted_GetRef();
 
-		UE_LOG(LogIoStore, Display, TEXT("Creating target: '%s' using output directory: '%s'"), *Target.TargetPlatform->PlatformName(), *Target.OutputDirectory);
+			if (!FParse::Value(*Command, TEXT("Output="), ContainerSpec.OutputPath))
+			{
+				UE_LOG(LogIoStore, Error, TEXT("Output argument missing from command '%s'"), *Command);
+				return -1;
+			}
 
-		int32 ReturnValue = CreateTarget(Target);
+			FString ResponseFilePath;
+			if (!FParse::Value(*Command, TEXT("ResponseFile="), ResponseFilePath))
+			{
+				UE_LOG(LogIoStore, Error, TEXT("ResponseFile argument missing from command '%s'"), *Command);
+				return -1;
+			}
+
+			TArray<FString> ResponseFileContents;
+			if (!FFileHelper::LoadFileToStringArray(ResponseFileContents, *ResponseFilePath))
+			{
+				UE_LOG(LogIoStore, Error, TEXT("Failed to read response file '%s'."), *ResponseFilePath);
+				return -1;
+			}
+
+			for (const FString& ResponseLine : ResponseFileContents)
+			{
+				TArray<FString> SourceAndDest;
+				TArray<FString> Switches;
+
+				FString NextToken;
+				const TCHAR* ResponseLinePtr = *ResponseLine;
+				while (FParse::Token(ResponseLinePtr, NextToken, false))
+				{
+					if ((**NextToken == TCHAR('-')))
+					{
+						new(Switches) FString(NextToken.Mid(1));
+					}
+					else
+					{
+						new(SourceAndDest) FString(NextToken);
+					}
+				}
+
+				if (SourceAndDest.Num() == 0)
+				{
+					continue;
+				}
+
+				if (SourceAndDest.Num() != 2)
+				{
+					UE_LOG(LogIoStore, Error, TEXT("Invalid line in response file '%s'."), *ResponseLine);
+					return -1;
+				}
+
+				FContainerSourceFile& FileEntry = ContainerSpec.SourceFiles.AddDefaulted_GetRef();
+				FileEntry.Path = SourceAndDest[0];
+			}
+		}
+
+		int32 ReturnValue = CreateTarget(*OutputDirectory, *CookedDirectory, Containers);
 		if (ReturnValue != 0)
 		{
 			return ReturnValue;
 		}
 	}
+	else
+	{
+		const TArray<ITargetPlatform*>& Platforms = GetTargetPlatformManagerRef().GetActiveTargetPlatforms();
+
+		for (const ITargetPlatform* TargetPlatform : Platforms)
+		{
+			const FString TargetCookedDirectory = FPaths::ProjectSavedDir() / TEXT("Cooked") / TargetPlatform->PlatformName();
+			const FString TargetCookedProjectDirectory = TargetCookedDirectory / FApp::GetProjectName();
+
+			UE_LOG(LogIoStore, Display, TEXT("Searching for cooked files..."));
+			class FCookedFileVisitor
+				: public IPlatformFile::FDirectoryVisitor
+			{
+				FContainerSourceSpec& ContainerSpec;
+
+			public:
+				FCookedFileVisitor(FContainerSourceSpec& InContainerSpec)
+					: ContainerSpec(InContainerSpec)
+				{
+				}
+
+				virtual bool Visit(const TCHAR* FilenameOrDirectory, bool bIsDirectory)
+				{
+					if (bIsDirectory == false)
+					{
+						FString Filename(FilenameOrDirectory);
+						if (Filename.EndsWith(TEXT(".uasset")) || Filename.EndsWith(TEXT(".umap")) || Filename.EndsWith(TEXT(".ubulk")) || Filename.EndsWith(TEXT(".uptnl")))
+						{
+							FContainerSourceFile& FileEntry = ContainerSpec.SourceFiles.AddDefaulted_GetRef();
+							FileEntry.Path = MoveTemp(FilenameOrDirectory);
+						}
+					}
+					return true;
+				}
+			};
+			
+			TArray<FContainerSourceSpec> Containers;
+			FContainerSourceSpec& ContainerSpec = Containers.AddDefaulted_GetRef();
+			FCookedFileVisitor CookedFileVistor(ContainerSpec);
+			FPlatformFileManager::Get().GetPlatformFile().IterateDirectoryRecursively(*TargetCookedDirectory, CookedFileVistor);
+			UE_LOG(LogIoStore, Display, TEXT("Found '%d' files"), ContainerSpec.SourceFiles.Num());
+
+			UE_LOG(LogIoStore, Display, TEXT("Creating target: '%s' using output path: '%s'"), *TargetPlatform->PlatformName(), *TargetCookedProjectDirectory);
+
+			int32 ReturnValue = CreateTarget(*TargetCookedProjectDirectory, *TargetCookedDirectory, Containers);
+			if (ReturnValue != 0)
+			{
+				return ReturnValue;
+			}
+		}
+	}
 
 	return 0;
 }
diff --git a/Engine/Source/Developer/PakFileUtilities/Private/PakFileUtilities.cpp b/Engine/Source/Developer/PakFileUtilities/Private/PakFileUtilities.cpp
--- a/Engine/Source/Developer/PakFileUtilities/Private/PakFileUtilities.cpp
+++ b/Engine/Source/Developer/PakFileUtilities/Private/PakFileUtilities.cpp
@@ -4858,12 +4858,6 @@ bool ExecuteUnrealPak(const TCHAR* CmdLine)
 			return false;
 		}
 
-		if (Entries.Num() == 0)
-		{
-			UE_LOG(LogPakFile, Error, TEXT("No files specified to add to pak file."));
-			return false;
-		}
-
 		int32 LowestSourcePakVersion = 0;
 		TMap<FString, FFileInfo> SourceFileHashes;
 
diff --git a/Engine/Source/Programs/AutomationTool/Scripts/CopyBuildToStagingDirectory.Automation.cs b/Engine/Source/Programs/AutomationTool/Scripts/CopyBuildToStagingDirectory.Automation.cs
--- a/Engine/Source/Programs/AutomationTool/Scripts/CopyBuildToStagingDirectory.Automation.cs
+++ b/Engine/Source/Programs/AutomationTool/Scripts/CopyBuildToStagingDirectory.Automation.cs
@@ -163,6 +163,22 @@ public partial class Project : CommandUtils
 		return CmdLine.ToString();
 	}
 
+	static private string GetIoStoreCommand(Dictionary<string, string> UnrealPakResponseFile, FileReference OutputLocation, FileReference PakOrderFileLocation, bool Compressed, EncryptionAndSigning.CryptoSettings CryptoSettings, string EncryptionKeyGuid, FileReference SecondaryPakOrderFileLocation=null)
+	{
+		StringBuilder CmdLine = new StringBuilder();
+		CmdLine.AppendFormat("-Output={0}", MakePathSafeToUseWithCommandLine(Path.ChangeExtension(OutputLocation.FullName, null)));
+
+		// Force encryption of ALL files if we're using specific encryption key. This should be made an option per encryption key in the settings, but for our initial
+		// implementation we will just assume that we require maximum security for this data.
+		bool bForceEncryption = !string.IsNullOrEmpty(EncryptionKeyGuid);
+		string PakName = Path.GetFileNameWithoutExtension(OutputLocation.FullName);
+		string UnrealPakResponseFileName = CombinePaths(CmdEnv.LogFolder, "PakListIoStore_" + PakName + ".txt");
+		WritePakResponseFile(UnrealPakResponseFileName, UnrealPakResponseFile, Compressed, CryptoSettings, bForceEncryption);
+		CmdLine.AppendFormat(" -ResponseFile={0}", CommandUtils.MakePathSafeToUseWithCommandLine(UnrealPakResponseFileName));
+
+		return CmdLine.ToString();
+	}
+
 	public static FileReference GetUnrealPakLocation()
 	{
 		if (HostPlatform.Current.HostEditorPlatform == UnrealTargetPlatform.Win64)
@@ -1928,6 +1944,7 @@ public partial class Project : CommandUtils
 
 		List<string> Commands = new List<string>();
 		List<string> LogNames = new List<string>();
+		List<string> IoStoreCommands = new List<string>();
 
 		// Calculate target patch index by iterating all pak files in source build
 		int TargetPatchIndex = 0;
@@ -2169,7 +2186,39 @@ public partial class Project : CommandUtils
 						AdditionalArgs += " -fallbackOrderForNonUassetFiles";
 					}
 
-					Commands.Add(GetUnrealPakArguments(Params.RawProjectPath, PakParams.UnrealPakResponseFile, OutputLocation, PrimaryOrderFile, SC.StageTargetPlatform.GetPlatformPakCommandLine(Params, SC) + AdditionalArgs + BulkOption + CompressionFormats + " " + Params.AdditionalPakOptions, PakParams.bCompressed, CryptoSettings, CryptoKeysCacheFilename, PatchSourceContentPath, PakParams.EncryptionKeyGuid, SecondaryOrderFile));
+					Dictionary<string, string> UnrealPakResponseFile;
+					if (ShouldCreateIoStoreContainerFiles(Params, SC))
+					{
+						UnrealPakResponseFile = new Dictionary<string, string>();
+						Dictionary<string, string> IoStoreResponseFile = new Dictionary<string, string>();
+						foreach (var Entry in PakParams.UnrealPakResponseFile)
+						{
+							if (Path.GetExtension(Entry.Key).Contains(".uasset") ||
+								Path.GetExtension(Entry.Key).Contains(".umap") ||
+								Path.GetExtension(Entry.Key).Contains(".ubulk") ||
+								Path.GetExtension(Entry.Key).Contains(".uptnl"))
+							{
+								IoStoreResponseFile.Add(Entry.Key, Entry.Value);
+							}
+							else if (!Path.GetExtension(Entry.Key).Contains(".uexp"))
+							{
+								UnrealPakResponseFile.Add(Entry.Key, Entry.Value);
+							}
+
+							// TODO: Remove when file enumeration has been fixed
+							if (Path.GetExtension(Entry.Key).Contains(".uasset") ||
+								Path.GetExtension(Entry.Key).Contains(".umap"))
+							{
+								UnrealPakResponseFile.Add(Entry.Key, Entry.Value);
+							}
+						}
+						IoStoreCommands.Add(GetIoStoreCommand(IoStoreResponseFile, OutputLocation, PrimaryOrderFile, PakParams.bCompressed, CryptoSettings, PakParams.EncryptionKeyGuid, SecondaryOrderFile));
+					}
+					else
+					{
+						UnrealPakResponseFile = PakParams.UnrealPakResponseFile;
+					}
+					Commands.Add(GetUnrealPakArguments(Params.RawProjectPath, UnrealPakResponseFile, OutputLocation, PrimaryOrderFile, SC.StageTargetPlatform.GetPlatformPakCommandLine(Params, SC) + AdditionalArgs + BulkOption + CompressionFormats + " " + Params.AdditionalPakOptions, PakParams.bCompressed, CryptoSettings, CryptoKeysCacheFilename, PatchSourceContentPath, PakParams.EncryptionKeyGuid, SecondaryOrderFile));
 					LogNames.Add(OutputLocation.GetFileNameWithoutExtension());
 				}
 			}
@@ -2183,6 +2232,19 @@ public partial class Project : CommandUtils
 			RunUnrealPakInParallel( Commands, LogNames, AdditionalCompressionOptionsOnCommandLine);
 		}
 
+		if (IoStoreCommands.Count > 0)
+		{
+			string IoStoreCommandsFileName = CombinePaths(CmdEnv.LogFolder, "IoStoreCommands.txt");
+			using (var Writer = new StreamWriter(IoStoreCommandsFileName, false, new System.Text.UTF8Encoding(true)))
+			{
+				foreach (string Command in IoStoreCommands)
+				{
+					Writer.WriteLine(Command);
+				}
+			}
+			RunIoStore(Params, SC, IoStoreCommandsFileName);
+		}
+
 		// Do any additional processing on the command output
 		for (int Idx = 0; Idx < PakParamsList.Count; Idx++)
 		{
@@ -2324,6 +2386,14 @@ public partial class Project : CommandUtils
 		}
 	}
 
+	private static void RunIoStore(ProjectParams Params, DeploymentContext SC, string CommandsFileName)
+	{
+		FileReference OutputLocation = FileReference.Combine(SC.RuntimeRootDir, SC.RelativeProjectRootForStage.Name);
+		string CommandletParams = String.Format("-OutputDirectory={0} -CookedDirectory={1} -Commands={2}", OutputLocation.FullName, SC.PlatformCookDir, CommandsFileName);
+		LogInformation("Running IoStore commandlet with arguments: {0}", CommandletParams);
+		RunCommandlet(SC.RawProjectPath, Params.UE4Exe, "IoStore", CommandletParams);
+	}
+
 	private static void RunUnrealPakInParallel(List<string> Commands, List<string> LogNames, string AdditionalCompressionOptionsOnCommandLine)
 	{
 		LogInformation("Executing {0} UnrealPak command{1}...", Commands.Count, (Commands.Count > 1)? "s" : "");
@@ -2963,22 +3033,6 @@ public partial class Project : CommandUtils
 
 	public static void ApplyStagingManifest(ProjectParams Params, DeploymentContext SC)
 	{
-		if (ShouldCreateIoStoreContainerFiles(Params, SC))
-		{
-			foreach (var ClientPlatform in Params.ClientTargetPlatforms)
-			{
-				TargetPlatformDescriptor DataPlatformDesc = Params.GetCookedDataPlatformForClientTarget(ClientPlatform);
-				String TargetPlatformName = Platform.Platforms[DataPlatformDesc].GetCookPlatform(false, Params.Client);
-				DirectoryReference OutputDirectory = DirectoryReference.Combine(SC.RuntimeRootDir, SC.RelativeProjectRootForStage.Name, "Content", "Containers");
-				String CommandletParams = String.Format("-TargetPlatform={0} -OutputDirectory={1}", TargetPlatformName, OutputDirectory.FullName);
-				if (SC.PlatformUsesChunkManifests && DoesChunkPakManifestExist(Params, SC))
-				{
-					CommandletParams += String.Format(" -ChunkListFile={0}", GetChunkPakManifestListFilename(Params, SC));
-				}
-				RunCommandlet(SC.RawProjectPath, Params.UE4Exe, "IoStore", CommandletParams);
-			}
-		}
-
 		if (ShouldCreatePak(Params, SC))
 		{
 			if (SC.CrashReporterUFSFiles.Count > 0)
diff --git a/Engine/Source/Runtime/Core/Private/IO/IoDispatcher.cpp b/Engine/Source/Runtime/Core/Private/IO/IoDispatcher.cpp
--- a/Engine/Source/Runtime/Core/Private/IO/IoDispatcher.cpp
+++ b/Engine/Source/Runtime/Core/Private/IO/IoDispatcher.cpp
@@ -147,11 +147,6 @@ public:
 	~FIoDispatcherImpl()
 	{
 		delete Thread;
-		IPlatformChunkInstall* ChunkInstall = FPlatformMisc::GetPlatformChunkInstall();
-		if (ChunkInstall && ChunkDownloadedDelegateHandle.IsValid())
-		{
-			ChunkInstall->RemoveChunkInstallDelegate(ChunkDownloadedDelegateHandle);
-		}
 	}
 
 	FIoStatus Initialize(const FIoStoreEnvironment& InitialEnvironment)
@@ -161,48 +156,6 @@ public:
 		{
 			return IoStatus;
 		}
-
-		FIoBuffer IoBuffer;
-		FEvent* Event = FPlatformProcess::GetSynchEventFromPool();
-		ReadWithCallback(
-			CreateIoChunkId(0, 0, EIoChunkType::InstallManifest),
-			FIoReadOptions(),
-			[Event, &IoBuffer](TIoStatusOr<FIoBuffer> Result)
-		{
-			if (Result.IsOk())
-			{
-				IoBuffer = Result.ConsumeValueOrDie();
-			}
-			Event->Trigger();
-		});
-		Event->Wait();
-		if (!IoBuffer.DataSize())
-		{
-			return FIoStatusBuilder(EIoErrorCode::NotFound) << TEXT("Failed to open install manifest");
-		}
-		FLargeMemoryReader InstallManifestAr(IoBuffer.Data(), IoBuffer.DataSize());
-
-		FIoStoreInstallManifest InstallManifest;
-		InstallManifestAr << InstallManifest;
-
-		IPlatformChunkInstall* ChunkInstall = FPlatformMisc::GetPlatformChunkInstall();
-		if (ChunkInstall)
-		{
-			ChunkDownloadedDelegateHandle = ChunkInstall->AddChunkInstallDelegate(FPlatformChunkInstallDelegate::CreateRaw(this, &FIoDispatcherImpl::OnChunkDownloaded));
-		}
-		for (const FIoStoreInstallManifest::FEntry& InstallManifestEntry : InstallManifest.ReadEntries())
-		{
-			FIoStoreEnvironment PartitionEnvironment(InitialEnvironment, InstallManifestEntry.PartitionName);
-			if (ChunkInstall && ChunkInstall->GetChunkLocation(InstallManifestEntry.InstallChunkId) == EChunkLocation::NotAvailable)
-			{
-				PendingInstallChunks.Add(MakeTuple(InstallManifestEntry.InstallChunkId, PartitionEnvironment));
-			}
-			else
-			{
-				Mount(PartitionEnvironment);
-			}
-		}
-
 		return FIoStatus::Ok;
 	}
 
@@ -419,23 +372,6 @@ private:
 		}
 	}
 
-	void OnChunkDownloaded(uint32 ChunkId, bool bSuccess)
-	{
-		if (bSuccess)
-		{
-			for (auto It = PendingInstallChunks.CreateIterator(); It; ++It)
-			{
-				int32 PendingChunkId = It->Get<0>();
-				if (PendingChunkId == ChunkId)
-				{
-					const FIoStoreEnvironment& PendingEnvironment = It->Get<1>();
-					Mount(PendingEnvironment);
-					It.RemoveCurrent();
-				}
-			}
-		}
-	}
-
 	virtual bool Init()
 	{
 		return true;
@@ -476,8 +412,6 @@ private:
 	int32 CurrentRequestsToSubmitIndex = 0;
 	FIoRequestImpl* SubmittedRequestsHead = nullptr;
 	FIoRequestImpl* SubmittedRequestsTail = nullptr;
-	FDelegateHandle ChunkDownloadedDelegateHandle;
-	TArray<TTuple<int32, FIoStoreEnvironment>> PendingInstallChunks;
 	TAtomic<bool> bStopRequested { false };
 };
 
diff --git a/Engine/Source/Runtime/Core/Private/IO/IoDispatcherFileBackend.cpp b/Engine/Source/Runtime/Core/Private/IO/IoDispatcherFileBackend.cpp
--- a/Engine/Source/Runtime/Core/Private/IO/IoDispatcherFileBackend.cpp
+++ b/Engine/Source/Runtime/Core/Private/IO/IoDispatcherFileBackend.cpp
@@ -20,19 +20,8 @@ FIoStatus FFileIoStoreReader::Initialize(const FIoStoreEnvironment& Environment)
 	IPlatformFile& Ipf = FPlatformFileManager::Get().GetPlatformFile();
 
 	TStringBuilder<256> ContainerFilePath;
-	ContainerFilePath.Append(Environment.GetBasePath());
-	if (ContainerFilePath.LastChar() != '/')
-	{
-		ContainerFilePath.Append(TEXT('/'));
-	}
-	if (Environment.GetPartitionName().IsEmpty())
-	{
-		ContainerFilePath.Append(TEXT("global"));
-	}
-	else
-	{
-		ContainerFilePath.Append(Environment.GetPartitionName());
-	}
+	ContainerFilePath.Append(Environment.GetPath());
+
 	TStringBuilder<256> TocFilePath;
 	TocFilePath.Append(ContainerFilePath);
 
@@ -240,19 +229,7 @@ TIoStatusOr<uint64> FFileIoStore::GetSizeForChunk(const FIoChunkId& ChunkId) con
 bool FFileIoStore::IsValidEnvironment(const FIoStoreEnvironment& Environment)
 {
 	TStringBuilder<256> TocFilePath;
-	TocFilePath.Append(Environment.GetBasePath());
-	if (TocFilePath.LastChar() != '/')
-	{
-		TocFilePath.Append(TEXT('/'));
-	}
-	if (Environment.GetPartitionName().IsEmpty())
-	{
-		TocFilePath.Append(TEXT("global"));
-	}
-	else
-	{
-		TocFilePath.Append(Environment.GetPartitionName());
-	}
+	TocFilePath.Append(Environment.GetPath());
 	TocFilePath.Append(TEXT(".utoc"));
 	return FPlatformFileManager::Get().GetPlatformFile().FileExists(*TocFilePath);
 }
diff --git a/Engine/Source/Runtime/Core/Private/IO/IoStore.cpp b/Engine/Source/Runtime/Core/Private/IO/IoStore.cpp
--- a/Engine/Source/Runtime/Core/Private/IO/IoStore.cpp
+++ b/Engine/Source/Runtime/Core/Private/IO/IoStore.cpp
@@ -16,19 +16,13 @@ FIoStoreEnvironment::FIoStoreEnvironment()
 {
 }
 
-FIoStoreEnvironment::FIoStoreEnvironment(const FIoStoreEnvironment& BaseEnvironment, FStringView InPartitionName)
-{
-	BasePath = BaseEnvironment.BasePath;
-	PartitionName = InPartitionName;
-}
-
 FIoStoreEnvironment::~FIoStoreEnvironment()
 {
 }
 
-void FIoStoreEnvironment::InitializeFileEnvironment(FStringView InBasePath)
+void FIoStoreEnvironment::InitializeFileEnvironment(FStringView InPath)
 {
-	BasePath = InBasePath;
+	Path = InPath;
 }
 
 //////////////////////////////////////////////////////////////////////////
@@ -45,15 +39,10 @@ public:
 	{
 		IPlatformFile& Ipf = IPlatformFile::GetPlatformPhysical();
 
-		FString PartitionName = Environment.GetPartitionName();
-		if (PartitionName.IsEmpty())
-		{
-			PartitionName = TEXT("global");
-		}
-		FString TocFilePath = Environment.GetBasePath() / PartitionName + TEXT(".utoc");
-		FString ContainerFilePath = Environment.GetBasePath() / PartitionName + TEXT(".ucas");
+		FString TocFilePath = Environment.GetPath() + TEXT(".utoc");
+		FString ContainerFilePath = Environment.GetPath() + TEXT(".ucas");
 
-		Ipf.CreateDirectoryTree(*Environment.GetBasePath());
+		Ipf.CreateDirectoryTree(*FPaths::GetPath(ContainerFilePath));
 
 		ContainerFileHandle.Reset(Ipf.OpenWrite(*ContainerFilePath, /* append */ false, /* allowread */ true));
 
diff --git a/Engine/Source/Runtime/Core/Public/IO/IoDispatcher.h b/Engine/Source/Runtime/Core/Public/IO/IoDispatcher.h
--- a/Engine/Source/Runtime/Core/Public/IO/IoDispatcher.h
+++ b/Engine/Source/Runtime/Core/Public/IO/IoDispatcher.h
@@ -570,51 +570,6 @@ private:
 	uint8	Id[12];
 };
 
-class FIoStoreInstallManifest
-{
-public:
-	struct FEntry
-	{
-		FString PartitionName;
-		int32 InstallChunkId;
-	};
-
-	inline const TArray<FEntry>& ReadEntries() const
-	{
-		return Entries;
-	}
-
-	inline TArray<FEntry>& EditEntries()
-	{
-		return Entries;
-	}
-
-	friend FArchive& operator<<(FArchive& Ar, FIoStoreInstallManifest& IoStoreManifest)
-	{
-		int32 Version = CurrentVersion;
-		Ar << Version;
-		check(Version == CurrentVersion);
-		int32 EntryCount = IoStoreManifest.Entries.Num();
-		Ar << EntryCount;
-		if (Ar.IsLoading())
-		{
-			IoStoreManifest.Entries.SetNum(EntryCount);
-		}
-		for (int32 Index = 0; Index < EntryCount; ++Index)
-		{
-			FEntry& Entry = IoStoreManifest.Entries[Index];
-			Ar << Entry.PartitionName;
-			Ar << Entry.InstallChunkId;
-		}
-		return Ar;
-	}
-
-private:
-	static constexpr int32 CurrentVersion = 1;
-
-	TArray<FEntry> Entries;
-};
-
 /**
  * Addressable chunk types.
  */
@@ -812,17 +767,14 @@ class FIoStoreEnvironment
 {
 public:
 	CORE_API FIoStoreEnvironment();
-	CORE_API FIoStoreEnvironment(const FIoStoreEnvironment& BaseEnvironment, FStringView PartitionName);
 	CORE_API ~FIoStoreEnvironment();
 
-	CORE_API void InitializeFileEnvironment(FStringView InBasePath);
+	CORE_API void InitializeFileEnvironment(FStringView InPath);
 
-	CORE_API const FString& GetBasePath() const { return BasePath; }
-	CORE_API const FString& GetPartitionName() const { return PartitionName; }
+	CORE_API const FString& GetPath() const { return Path; }
 
 private:
-	FString			BasePath;
-	FString			PartitionName;
+	FString			Path;
 };
 
 //////////////////////////////////////////////////////////////////////////
diff --git a/Engine/Source/Runtime/Launch/Private/LaunchEngineLoop.cpp b/Engine/Source/Runtime/Launch/Private/LaunchEngineLoop.cpp
--- a/Engine/Source/Runtime/Launch/Private/LaunchEngineLoop.cpp
+++ b/Engine/Source/Runtime/Launch/Private/LaunchEngineLoop.cpp
@@ -1585,6 +1585,22 @@ int32 FEngineLoop::PreInitPreStartupScreen(const TCHAR* CmdLine)
 	}
 #endif
 
+#if !(IS_PROGRAM || WITH_EDITOR)
+	// Initialize I/O dispatcher if available
+	if (!FParse::Param(CmdLine, TEXT("noiodispatcher")))
+	{
+		TRACE_CPUPROFILER_EVENT_SCOPE(InitIoDispatcher);
+
+		FIoStoreEnvironment IoStoreEnvironment;
+		IoStoreEnvironment.InitializeFileEnvironment(FPaths::ProjectDir() / TEXT("global"));
+		if (FIoDispatcher::IsValidEnvironment(IoStoreEnvironment))
+		{
+			FIoStatus IoDispatcherInitStatus = FIoDispatcher::Initialize(IoStoreEnvironment);
+			UE_CLOG(!IoDispatcherInitStatus.IsOk(), LogInit, Fatal, TEXT("Failed to initialize I/O dispatcher: '%s'"), *IoDispatcherInitStatus.ToString());
+		}
+	}
+#endif
+
 	// allow the command line to override the platform file singleton
 	bool bFileOverrideFound = false;
 	{
@@ -5509,22 +5525,6 @@ bool FEngineLoop::AppInit( )
 
 	FEmbeddedCommunication::ForceTick(18);
 
-#if !(IS_PROGRAM || WITH_EDITOR)
-	// Initialize I/O dispatcher if available
-	if (!FParse::Param(FCommandLine::Get(), TEXT("noiodispatcher")))
-	{
-		TRACE_CPUPROFILER_EVENT_SCOPE(InitIoDispatcher);
-
-		FIoStoreEnvironment IoStoreEnvironment;
-		IoStoreEnvironment.InitializeFileEnvironment(FPaths::ProjectDir() / TEXT("Content") / TEXT("Containers"));
-		if (FIoDispatcher::IsValidEnvironment(IoStoreEnvironment))
-		{
-			FIoStatus IoDispatcherInitStatus = FIoDispatcher::Initialize(IoStoreEnvironment);
-			UE_CLOG(!IoDispatcherInitStatus.IsOk(), LogInit, Fatal, TEXT("Failed to initialize I/O dispatcher: '%s'"), *IoDispatcherInitStatus.ToString());
-		}
-	}
-#endif
-
 	// Init other systems.
 	{
 		SCOPED_BOOT_TIMING("FCoreDelegates::OnInit.Broadcast");
diff --git a/Engine/Source/Runtime/PakFile/Private/IPlatformFilePak.cpp b/Engine/Source/Runtime/PakFile/Private/IPlatformFilePak.cpp
--- a/Engine/Source/Runtime/PakFile/Private/IPlatformFilePak.cpp
+++ b/Engine/Source/Runtime/PakFile/Private/IPlatformFilePak.cpp
@@ -32,6 +32,7 @@
 #include "Misc/Fnv.h"
 
 #include "Async/MappedFileHandle.h"
+#include "IO/IoDispatcher.h"
 
 DEFINE_LOG_CATEGORY(LogPakFile);
 
@@ -6330,6 +6331,21 @@ bool FPakPlatformFile::Mount(const TCHAR* InPakFilename, uint32 PakOrder, const
 
 		if (bSuccess)
 		{
+			if (FIoDispatcher::IsInitialized())
+			{
+				FIoStoreEnvironment IoStoreEnvironment;
+				IoStoreEnvironment.InitializeFileEnvironment(FPaths::ChangeExtension(InPakFilename, FString()));
+				FIoStatus IoStatus = FIoDispatcher::Get().Mount(IoStoreEnvironment);
+				if (IoStatus.IsOk())
+				{
+					UE_LOG(LogPakFile, Display, TEXT("Mounted IoStore environment \"%s\""), *IoStoreEnvironment.GetPath());
+				}
+				else
+				{
+					UE_LOG(LogPakFile, Warning, TEXT("Failed to mount IoStore environment \"%s\""), *IoStoreEnvironment.GetPath());
+				}
+			}
+
 			PRAGMA_DISABLE_DEPRECATION_WARNINGS
 			FCoreDelegates::PakFileMountedCallback.Broadcast(InPakFilename);
 			PRAGMA_ENABLE_DEPRECATION_WARNINGS
