IoStore:
Refactored to have access to the size of each chunk before the chunk is added to a container
Basic chunk ordering
Optional output of .csv container manifests

diff --git a/Engine/Source/Developer/IoStoreUtilities/Private/IoStoreUtilities.cpp b/Engine/Source/Developer/IoStoreUtilities/Private/IoStoreUtilities.cpp
--- a/Engine/Source/Developer/IoStoreUtilities/Private/IoStoreUtilities.cpp
+++ b/Engine/Source/Developer/IoStoreUtilities/Private/IoStoreUtilities.cpp
@@ -39,6 +39,7 @@ DEFINE_LOG_CATEGORY_STATIC(LogIoStore, Log, All);
 #define OUTPUT_NAMEMAP_CSV 0
 #define OUTPUT_IMPORTMAP_CSV 0
 #define SKIP_WRITE_CONTAINER 0
+#define OUTPUT_CONTAINER_CSV 0
 
 struct FContainerSourceFile 
 {
@@ -55,9 +56,11 @@ struct FPackage;
 
 struct FContainerTargetFile
 {
-	FPackage* Package;
+	FPackage* Package = nullptr;
 	FString SourcePath;
-	bool bIsBulkData;
+	FString TargetPath;
+	uint64 Size = 0;
+	bool bIsBulkData = false;
 };
 
 struct FContainerTargetSpec
@@ -369,6 +372,7 @@ struct FPackage
 	TArray<FString> ImportedFullNames;
 
 	TArray<FPackage*> ImportedPackages;
+	TArray<FPackage*> ImportedByPackages;
 	TSet<FPackage*> AllReachablePackages;
 	TSet<FPackage*> ImportedPreloadPackages;
 
@@ -389,6 +393,12 @@ struct FPackage
 
 	TArray<FExportGraphNode*> NodesWithNoIncomingEdges;
 	FPackageGraphNode* Node = nullptr;
+
+	TArray<uint8> HeaderData;
+	uint64 ExportsPayloadSize = 0;
+
+	uint64 ClusterIndex = MAX_uint64;
+	uint64 OpenOrder = MAX_uint64;
 };
 
 struct FCircularImportChain
@@ -795,6 +805,38 @@ static void BuildBundles(FExportGraph& ExportGraph, const TArray<FPackage*>& Pac
 	}
 }
 
+static void BuildClustersRecursive(FPackage* Package, uint64 ClusterIndex)
+{
+	Package->ClusterIndex = ClusterIndex;
+	for (FPackage* OtherPackage : Package->ImportedPackages)
+	{
+		if (OtherPackage->ClusterIndex == MAX_uint64)
+		{
+			BuildClustersRecursive(OtherPackage, ClusterIndex);
+		}
+	}
+	for (FPackage* OtherPackage : Package->ImportedByPackages)
+	{
+		if (OtherPackage->ClusterIndex == MAX_uint64)
+		{
+			BuildClustersRecursive(OtherPackage, ClusterIndex);
+		}
+	}
+}
+
+static void BuildClusters(const TArray<FPackage*>& Packages)
+{
+	uint64 NextClusterIndex = 0;
+	for (FPackage* Package : Packages)
+	{
+		if (Package->ClusterIndex == MAX_uint64)
+		{
+			BuildClustersRecursive(Package, NextClusterIndex);
+			++NextClusterIndex;
+		}
+	}
+}
+
 static EIoChunkType BulkdataTypeToChunkIdType(FPackageStoreBulkDataManifest::EBulkdataType Type)
 {
 	switch (Type)
@@ -833,7 +875,7 @@ static bool WriteBulkData(	const FString& Filename, FPackageStoreBulkDataManifes
 			BulkAr->Close();
 		}
 		
-		const FIoStatus AppendResult = IoStoreWriter->Append(BulkDataChunkId, IoBuffer);
+		const FIoStatus AppendResult = IoStoreWriter->Append(BulkDataChunkId, IoBuffer, *FPaths::GetCleanFilename(Filename));
 		if (!AppendResult.IsOk())
 		{
 			UE_LOG(LogIoStore, Error, TEXT("Failed to append bulkdata for '%s' due to: %s"), *Package.FileName, *AppendResult.ToString());
@@ -1018,13 +1060,8 @@ static void FindExport(TArray<FExportData>& GlobalExports, TMap<FString, int32>&
 	}
 };
 
-FPackage* FindOrAddPackage(const TCHAR* FileName, const TCHAR* CookedDir, TArray<FPackage*>& Packages, TMap<FName, FPackage*>& PackageMap)
+FPackage* FindOrAddPackage(const TCHAR* RelativeFileName, TArray<FPackage*>& Packages, TMap<FName, FPackage*>& PackageMap, const TMap<FString, uint64>& PackageOrderMap)
 {
-	FString RelativeFileName = FileName;
-	RelativeFileName.RemoveFromStart(CookedDir);
-	RelativeFileName.RemoveFromStart(TEXT("/"));
-	RelativeFileName = TEXT("../../../") / RelativeFileName;
-
 	FString PackageName;
 	FString ErrorMessage;
 	if (!FPackageName::TryConvertFilenameToLongPackageName(RelativeFileName, PackageName, &ErrorMessage))
@@ -1043,240 +1080,240 @@ FPackage* FindOrAddPackage(const TCHAR* FileName, const TCHAR* CookedDir, TArray
 		Package->GlobalPackageId = Packages.Num();
 		Packages.Add(Package);
 		PackageMap.Add(PackageFName, Package);
+
+
+		const uint64* FindPackageOrder = PackageOrderMap.Find(PackageName);
+		if (FindPackageOrder)
+		{
+			Package->OpenOrder = *FindPackageOrder;
+		}
 	}
 
 	return Package;
 }
 
-void SerializePackageData(
-	const FContainerTargetFile& TargetFile,
+void FinalizePackageHeader(
+	FPackage* Package,
 	const FNameMapBuilder& NameMapBuilder,
 	const TArray<FObjectExport>& ObjectExports,
 	const TArray<FExportData>& GlobalExports,
 	const TMap<FString, int32>& GlobalImportsByFullName,
-	FExportBundleMetaEntry* ExportBundleMetaEntries,
-	FIoStoreWriter* IoStoreWriter)
+	FExportBundleMetaEntry* ExportBundleMetaEntries)
 {
-	FPackage* Package = TargetFile.Package;
+	// Temporary Archive for serializing ImportMap
+	FBufferWriter ImportMapArchive(nullptr, 0, EBufferWriterFlags::AllowResize | EBufferWriterFlags::TakeOwnership);
+	for (int32 GlobalImportIndex : Package->Imports)
+	{
+		ImportMapArchive << GlobalImportIndex;
+	}
+	Package->ImportMapSize = ImportMapArchive.Tell();
 
-		// Temporary Archive for serializing ImportMap
-		FBufferWriter ImportMapArchive(nullptr, 0, EBufferWriterFlags::AllowResize | EBufferWriterFlags::TakeOwnership);
-		for (int32 GlobalImportIndex : Package->Imports)
-		{
-			ImportMapArchive << GlobalImportIndex;
-		}
-		Package->ImportMapSize = ImportMapArchive.Tell();
+	// Temporary Archive for serializing EDL graph data
+	FBufferWriter GraphArchive(nullptr, 0, EBufferWriterFlags::AllowResize | EBufferWriterFlags::TakeOwnership);
 
-		// Temporary Archive for serializing EDL graph data
-		FBufferWriter GraphArchive(nullptr, 0, EBufferWriterFlags::AllowResize | EBufferWriterFlags::TakeOwnership);
+	int32 InternalArcCount = Package->InternalArcs.Num();
+	GraphArchive << InternalArcCount;
+	for (FArc& InternalArc : Package->InternalArcs)
+	{
+		GraphArchive << InternalArc.FromNodeIndex;
+		GraphArchive << InternalArc.ToNodeIndex;
+	}
 
-		int32 InternalArcCount = Package->InternalArcs.Num();
-		GraphArchive << InternalArcCount;
-		for (FArc& InternalArc : Package->InternalArcs)
-		{
-			GraphArchive << InternalArc.FromNodeIndex;
-			GraphArchive << InternalArc.ToNodeIndex;
-		}
+	int32 ReferencedPackagesCount = Package->ExternalArcs.Num();
+	GraphArchive << ReferencedPackagesCount;
+	for (auto& KV : Package->ExternalArcs)
+	{
+		FPackage* ImportedPackage = KV.Key;
+		TArray<FArc>& Arcs = KV.Value;
+		int32 ExternalArcCount = Arcs.Num();
 
-		int32 ReferencedPackagesCount = Package->ExternalArcs.Num();
-		GraphArchive << ReferencedPackagesCount;
-		for (auto& KV : Package->ExternalArcs)
+		GraphArchive << ImportedPackage->GlobalPackageId;
+		GraphArchive << ExternalArcCount;
+		for (FArc& ExternalArc : Arcs)
 		{
-			FPackage* ImportedPackage = KV.Key;
-			TArray<FArc>& Arcs = KV.Value;
-			int32 ExternalArcCount = Arcs.Num();
-
-			GraphArchive << ImportedPackage->GlobalPackageId;
-			GraphArchive << ExternalArcCount;
-			for (FArc& ExternalArc : Arcs)
-			{
-				GraphArchive << ExternalArc.FromNodeIndex;
-				GraphArchive << ExternalArc.ToNodeIndex;
-			}
+			GraphArchive << ExternalArc.FromNodeIndex;
+			GraphArchive << ExternalArc.ToNodeIndex;
 		}
-		Package->UGraphSize = GraphArchive.Tell();
+	}
+	Package->UGraphSize = GraphArchive.Tell();
 
-		// Temporary Archive for serializing export map data
-		FBufferWriter ExportMapArchive(nullptr, 0, EBufferWriterFlags::AllowResize | EBufferWriterFlags::TakeOwnership);
-		for (int32 I = 0; I < Package->ExportCount; ++I)
-		{
-			const FObjectExport& ObjectExport = ObjectExports[Package->ExportIndexOffset + I];
-			const FExportData& ExportData = GlobalExports[Package->Exports[I]];
-
-			int64 SerialSize = ObjectExport.SerialSize;
-			ExportMapArchive << SerialSize;
-			NameMapBuilder.SerializeName(ExportMapArchive, ObjectExport.ObjectName);
-			FPackageIndex OuterIndex = ExportData.OuterIndex;
-			ExportMapArchive << OuterIndex;
-			FPackageIndex ClassIndex = ExportData.ClassIndex;
-			ExportMapArchive << ClassIndex;
-			FPackageIndex SuperIndex = ExportData.SuperIndex;
-			ExportMapArchive << SuperIndex;
-			FPackageIndex TemplateIndex = ExportData.TemplateIndex;
-			ExportMapArchive << TemplateIndex;
-			int32 GlobalImportIndex = ExportData.GlobalImportIndex;
-			ExportMapArchive << GlobalImportIndex;
-			uint32 ObjectFlags = ObjectExport.ObjectFlags;
-			ExportMapArchive << ObjectFlags;
-			uint8 FilterFlags = uint8(EExportFilterFlags::None);
-			if (ObjectExport.bNotForClient)
-			{
-				FilterFlags = uint8(EExportFilterFlags::NotForClient);
-			}
-			else if (ObjectExport.bNotForServer)
-			{
-				FilterFlags = uint8(EExportFilterFlags::NotForServer);
-			}
-			ExportMapArchive << FilterFlags;
-			uint8 Pad = 0;
-			ExportMapArchive.Serialize(&Pad, 7);
-		}
-		Package->ExportMapSize = ExportMapArchive.Tell();
+	// Temporary Archive for serializing export map data
+	FBufferWriter ExportMapArchive(nullptr, 0, EBufferWriterFlags::AllowResize | EBufferWriterFlags::TakeOwnership);
+	for (int32 I = 0; I < Package->ExportCount; ++I)
+	{
+		const FObjectExport& ObjectExport = ObjectExports[Package->ExportIndexOffset + I];
+		const FExportData& ExportData = GlobalExports[Package->Exports[I]];
 
-		// Temporary archive for serializing export bundle data
-		FBufferWriter ExportBundlesArchive(nullptr, 0, EBufferWriterFlags::AllowResize | EBufferWriterFlags::TakeOwnership);
-		int32 ExportBundleEntryIndex = 0;
-		for (FExportBundle& ExportBundle : Package->ExportBundles)
+		int64 SerialSize = ObjectExport.SerialSize;
+		ExportMapArchive << SerialSize;
+		NameMapBuilder.SerializeName(ExportMapArchive, ObjectExport.ObjectName);
+		FPackageIndex OuterIndex = ExportData.OuterIndex;
+		ExportMapArchive << OuterIndex;
+		FPackageIndex ClassIndex = ExportData.ClassIndex;
+		ExportMapArchive << ClassIndex;
+		FPackageIndex SuperIndex = ExportData.SuperIndex;
+		ExportMapArchive << SuperIndex;
+		FPackageIndex TemplateIndex = ExportData.TemplateIndex;
+		ExportMapArchive << TemplateIndex;
+		int32 GlobalImportIndex = ExportData.GlobalImportIndex;
+		ExportMapArchive << GlobalImportIndex;
+		uint32 ObjectFlags = ObjectExport.ObjectFlags;
+		ExportMapArchive << ObjectFlags;
+		uint8 FilterFlags = uint8(EExportFilterFlags::None);
+		if (ObjectExport.bNotForClient)
 		{
-			ExportBundlesArchive << ExportBundleEntryIndex;
-			int32 EntryCount = ExportBundle.Nodes.Num();
-			ExportBundlesArchive << EntryCount;
-			ExportBundleEntryIndex += ExportBundle.Nodes.Num();
+			FilterFlags = uint8(EExportFilterFlags::NotForClient);
 		}
-		for (FExportBundle& ExportBundle : Package->ExportBundles)
+		else if (ObjectExport.bNotForServer)
 		{
-			for (FExportGraphNode* ExportNode : ExportBundle.Nodes)
-			{
-				uint32 CommandType = uint32(ExportNode->BundleEntry.CommandType);
-				ExportBundlesArchive << ExportNode->BundleEntry.LocalExportIndex;
-				ExportBundlesArchive << CommandType;
-			}
+			FilterFlags = uint8(EExportFilterFlags::NotForServer);
 		}
-		Package->ExportBundlesSize = ExportBundlesArchive.Tell();
-
-		Package->NameMapSize = Package->NameIndices.Num() * Package->NameIndices.GetTypeSize();
+		ExportMapArchive << FilterFlags;
+		uint8 Pad = 0;
+		ExportMapArchive.Serialize(&Pad, 7);
+	}
+	Package->ExportMapSize = ExportMapArchive.Tell();
 
+	// Temporary archive for serializing export bundle data
+	FBufferWriter ExportBundlesArchive(nullptr, 0, EBufferWriterFlags::AllowResize | EBufferWriterFlags::TakeOwnership);
+	int32 ExportBundleEntryIndex = 0;
+	for (FExportBundle& ExportBundle : Package->ExportBundles)
+	{
+		ExportBundlesArchive << ExportBundleEntryIndex;
+		int32 EntryCount = ExportBundle.Nodes.Num();
+		ExportBundlesArchive << EntryCount;
+		ExportBundleEntryIndex += ExportBundle.Nodes.Num();
+	}
+	for (FExportBundle& ExportBundle : Package->ExportBundles)
+	{
+		for (FExportGraphNode* ExportNode : ExportBundle.Nodes)
 		{
-			const uint64 PackageSummarySize =
-				sizeof(FPackageSummary)
-				+ Package->NameMapSize
-				+ Package->ImportMapSize
-				+ Package->ExportMapSize
-				+ Package->ExportBundlesSize
-				+ Package->UGraphSize;
+			uint32 CommandType = uint32(ExportNode->BundleEntry.CommandType);
+			ExportBundlesArchive << ExportNode->BundleEntry.LocalExportIndex;
+			ExportBundlesArchive << CommandType;
+		}
+	}
+	Package->ExportBundlesSize = ExportBundlesArchive.Tell();
 
-			uint8* PackageSummaryBuffer = static_cast<uint8*>(FMemory::Malloc(PackageSummarySize));
-			FPackageSummary* PackageSummary = reinterpret_cast<FPackageSummary*>(PackageSummaryBuffer);
+	Package->NameMapSize = Package->NameIndices.Num() * Package->NameIndices.GetTypeSize();
 
-			PackageSummary->PackageFlags = Package->PackageFlags;
-			PackageSummary->GraphDataSize = Package->UGraphSize;
-			PackageSummary->BulkDataStartOffset = Package->BulkDataStartOffset;
-			const int32* FindGlobalImportIndexForPackage = GlobalImportsByFullName.Find(Package->Name.ToString());
-			PackageSummary->GlobalImportIndex = FindGlobalImportIndexForPackage ? *FindGlobalImportIndexForPackage : -1;
+	uint64 PackageHeaderSize =
+		sizeof(FPackageSummary)
+		+ Package->NameMapSize
+		+ Package->ImportMapSize
+		+ Package->ExportMapSize
+		+ Package->ExportBundlesSize
+		+ Package->UGraphSize;
 
-			FBufferWriter SummaryArchive(PackageSummaryBuffer, PackageSummarySize);
-			SummaryArchive.Seek(sizeof(FPackageSummary));
+	Package->HeaderData.AddZeroed(PackageHeaderSize);
+	uint8* PackageHeaderBuffer = Package->HeaderData.GetData();
+	FPackageSummary* PackageSummary = reinterpret_cast<FPackageSummary*>(PackageHeaderBuffer);
 
-			// NameMap data
-			{
-				PackageSummary->NameMapOffset = SummaryArchive.Tell();
-				SummaryArchive.Serialize(Package->NameIndices.GetData(), Package->NameMapSize);
-			}
+	PackageSummary->PackageFlags = Package->PackageFlags;
+	PackageSummary->GraphDataSize = Package->UGraphSize;
+	PackageSummary->BulkDataStartOffset = Package->BulkDataStartOffset;
+	const int32* FindGlobalImportIndexForPackage = GlobalImportsByFullName.Find(Package->Name.ToString());
+	PackageSummary->GlobalImportIndex = FindGlobalImportIndexForPackage ? *FindGlobalImportIndexForPackage : -1;
 
-			// ImportMap data
-			{
-				check(ImportMapArchive.Tell() == Package->ImportMapSize);
-				PackageSummary->ImportMapOffset = SummaryArchive.Tell();
-				SummaryArchive.Serialize(ImportMapArchive.GetWriterData(), ImportMapArchive.Tell());
-			}
+	FBufferWriter SummaryArchive(PackageHeaderBuffer, PackageHeaderSize);
+	SummaryArchive.Seek(sizeof(FPackageSummary));
 
-			// ExportMap data
-			{
-				check(ExportMapArchive.Tell() == Package->ExportMapSize);
-				PackageSummary->ExportMapOffset = SummaryArchive.Tell();
-				SummaryArchive.Serialize(ExportMapArchive.GetWriterData(), ExportMapArchive.Tell());
-			}
+	// NameMap data
+	{
+		PackageSummary->NameMapOffset = SummaryArchive.Tell();
+		SummaryArchive.Serialize(Package->NameIndices.GetData(), Package->NameMapSize);
+	}
 
-			// ExportBundle data
-			{
-				check(ExportBundlesArchive.Tell() == Package->ExportBundlesSize);
-				PackageSummary->ExportBundlesOffset = SummaryArchive.Tell();
-				SummaryArchive.Serialize(ExportBundlesArchive.GetWriterData(), ExportBundlesArchive.Tell());
-			}
+	// ImportMap data
+	{
+		check(ImportMapArchive.Tell() == Package->ImportMapSize);
+		PackageSummary->ImportMapOffset = SummaryArchive.Tell();
+		SummaryArchive.Serialize(ImportMapArchive.GetWriterData(), ImportMapArchive.Tell());
+	}
 
-			// Graph data
-			{
-				check(GraphArchive.Tell() == Package->UGraphSize);
-				PackageSummary->GraphDataOffset = SummaryArchive.Tell();
-				SummaryArchive.Serialize(GraphArchive.GetWriterData(), GraphArchive.Tell());
-			}
+	// ExportMap data
+	{
+		check(ExportMapArchive.Tell() == Package->ExportMapSize);
+		PackageSummary->ExportMapOffset = SummaryArchive.Tell();
+		SummaryArchive.Serialize(ExportMapArchive.GetWriterData(), ExportMapArchive.Tell());
+	}
 
-			// Export bundle chunks
-			{
-				check(Package->ExportBundles.Num());
+	// ExportBundle data
+	{
+		check(ExportBundlesArchive.Tell() == Package->ExportBundlesSize);
+		PackageSummary->ExportBundlesOffset = SummaryArchive.Tell();
+		SummaryArchive.Serialize(ExportBundlesArchive.GetWriterData(), ExportBundlesArchive.Tell());
+	}
 
-				FString UExpFileName = FPaths::ChangeExtension(Package->FileName, TEXT(".uexp"));
-				TUniquePtr<FArchive> ExpAr(IFileManager::Get().CreateFileReader(*UExpFileName));
-				Package->UExpSize = ExpAr->TotalSize();
-#if !SKIP_WRITE_CONTAINER
-				uint8* ExportsBuffer = static_cast<uint8*>(FMemory::Malloc(ExpAr->TotalSize()));
-				ExpAr->Serialize(ExportsBuffer, ExpAr->TotalSize());
-#endif
-				ExpAr->Close();
+	// Graph data
+	{
+		check(GraphArchive.Tell() == Package->UGraphSize);
+		PackageSummary->GraphDataOffset = SummaryArchive.Tell();
+		SummaryArchive.Serialize(GraphArchive.GetWriterData(), GraphArchive.Tell());
+	}
 
-				uint64 BundleBufferSize = PackageSummarySize;
-				for (int32 ExportBundleIndex = 0; ExportBundleIndex < Package->ExportBundles.Num(); ++ExportBundleIndex)
-				{
-					FExportBundle& ExportBundle = Package->ExportBundles[ExportBundleIndex];
-					FExportBundleMetaEntry* ExportBundleMetaEntry = ExportBundleMetaEntries + Package->FirstExportBundleMetaEntry + ExportBundleIndex;
-					ExportBundleMetaEntry->LoadOrder = ExportBundle.LoadOrder;
-					for (FExportGraphNode* Node : ExportBundle.Nodes)
-					{
-						if (Node->BundleEntry.CommandType == FExportBundleEntry::ExportCommandType_Serialize)
-						{
-							const FObjectExport& ObjectExport = ObjectExports[Package->ExportIndexOffset + Node->BundleEntry.LocalExportIndex];
-							BundleBufferSize += ObjectExport.SerialSize;
-						}
-					}
-					if (ExportBundleIndex == 0)
-					{
-						ExportBundleMetaEntry->PayloadSize = BundleBufferSize;
-					}
-				}
+	// Export bundle payload
+	{
+		check(Package->ExportBundles.Num());
 
-#if !SKIP_WRITE_CONTAINER
-				uint8* BundleBuffer = static_cast<uint8*>(FMemory::Malloc(BundleBufferSize));
-				FMemory::Memcpy(BundleBuffer, PackageSummaryBuffer, PackageSummarySize);
-#endif
-				FMemory::Free(PackageSummaryBuffer);
-				uint64 BundleBufferOffset = PackageSummarySize;
-				for (FExportBundle& ExportBundle : Package->ExportBundles)
+		for (int32 ExportBundleIndex = 0; ExportBundleIndex < Package->ExportBundles.Num(); ++ExportBundleIndex)
+		{
+			FExportBundle& ExportBundle = Package->ExportBundles[ExportBundleIndex];
+			FExportBundleMetaEntry* ExportBundleMetaEntry = ExportBundleMetaEntries + Package->FirstExportBundleMetaEntry + ExportBundleIndex;
+			ExportBundleMetaEntry->LoadOrder = ExportBundle.LoadOrder;
+			for (FExportGraphNode* Node : ExportBundle.Nodes)
+			{
+				if (Node->BundleEntry.CommandType == FExportBundleEntry::ExportCommandType_Serialize)
 				{
-					for (FExportGraphNode* Node : ExportBundle.Nodes)
-					{
-						if (Node->BundleEntry.CommandType == FExportBundleEntry::ExportCommandType_Serialize)
-						{
-							const FObjectExport& ObjectExport = ObjectExports[Package->ExportIndexOffset + Node->BundleEntry.LocalExportIndex];
-							const int64 Offset = ObjectExport.SerialOffset - Package->UAssetSize;
-#if !SKIP_WRITE_CONTAINER
-							FMemory::Memcpy(BundleBuffer + BundleBufferOffset, ExportsBuffer + Offset, ObjectExport.SerialSize);
-#endif
-							BundleBufferOffset += ObjectExport.SerialSize;
-						}
-					}
+					const FObjectExport& ObjectExport = ObjectExports[Package->ExportIndexOffset + Node->BundleEntry.LocalExportIndex];
+					Package->ExportsPayloadSize += ObjectExport.SerialSize;
 				}
+			}
+		}
+		FExportBundleMetaEntry* FirstExportBundleMetaEntry = ExportBundleMetaEntries + Package->FirstExportBundleMetaEntry;
+		FirstExportBundleMetaEntry->PayloadSize = Package->HeaderData.Num() + Package->ExportsPayloadSize;
+	}
+}
 
+void SerializePackageData(
+	FPackage* Package,
+	const TArray<FObjectExport>& ObjectExports,
+	FIoStoreWriter* IoStoreWriter)
+{
+	FString UExpFileName = FPaths::ChangeExtension(Package->FileName, TEXT(".uexp"));
+	TUniquePtr<FArchive> ExpAr(IFileManager::Get().CreateFileReader(*UExpFileName));
+	Package->UExpSize = ExpAr->TotalSize();
 #if !SKIP_WRITE_CONTAINER
-				FIoBuffer IoBuffer(FIoBuffer::Wrap, BundleBuffer, BundleBufferSize);
-				IoStoreWriter->Append(CreateChunkId(Package->GlobalPackageId, 0, EIoChunkType::ExportBundleData, *Package->FileName), IoBuffer);
-				FMemory::Free(BundleBuffer);
-				FMemory::Free(ExportsBuffer);
-#endif
+	uint8* ExportsBuffer = static_cast<uint8*>(FMemory::Malloc(ExpAr->TotalSize()));
+	ExpAr->Serialize(ExportsBuffer, ExpAr->TotalSize());
+	ExpAr->Close();
+	
+	uint64 BundleBufferSize = Package->HeaderData.Num() + Package->ExportsPayloadSize;
+	uint8* BundleBuffer = static_cast<uint8*>(FMemory::Malloc(BundleBufferSize));
+	FMemory::Memcpy(BundleBuffer, Package->HeaderData.GetData(), Package->HeaderData.Num());
+	uint64 BundleBufferOffset = Package->HeaderData.Num();
+	for (FExportBundle& ExportBundle : Package->ExportBundles)
+	{
+		for (FExportGraphNode* Node : ExportBundle.Nodes)
+		{
+			if (Node->BundleEntry.CommandType == FExportBundleEntry::ExportCommandType_Serialize)
+			{
+				const FObjectExport& ObjectExport = ObjectExports[Package->ExportIndexOffset + Node->BundleEntry.LocalExportIndex];
+				const int64 Offset = ObjectExport.SerialOffset - Package->UAssetSize;
+				FMemory::Memcpy(BundleBuffer + BundleBufferOffset, ExportsBuffer + Offset, ObjectExport.SerialSize);
+				BundleBufferOffset += ObjectExport.SerialSize;
 			}
+		}
 	}
+
+	FIoBuffer IoBuffer(FIoBuffer::Wrap, BundleBuffer, BundleBufferSize);
+	IoStoreWriter->Append(CreateChunkId(Package->GlobalPackageId, 0, EIoChunkType::ExportBundleData, *Package->FileName), IoBuffer, *FPaths::GetCleanFilename(Package->FileName));
+	FMemory::Free(BundleBuffer);
+	FMemory::Free(ExportsBuffer);
+#endif
 }
 
-int32 CreateTarget(const TCHAR* OutputDir, const TCHAR* CookedDir, const TArray<FContainerSourceSpec> Containers)
+int32 CreateTarget(const TCHAR* OutputDir, const TCHAR* CookedDir, const TArray<FContainerSourceSpec>& Containers, const TMap<FString, uint64> PackageOrderMap)
 {
 	TGuardValue<int32> GuardAllowUnversionedContentInEditor(GAllowUnversionedContentInEditor, 1);
 
@@ -1324,6 +1361,10 @@ int32 CreateTarget(const TCHAR* OutputDir, const TCHAR* CookedDir, const TArray<
 #if !SKIP_WRITE_CONTAINER
 		FIoStatus IoStatus = GlobalIoStoreWriter->Initialize();
 		check(IoStatus.IsOk());
+#if OUTPUT_CONTAINER_CSV
+		IoStatus = GlobalIoStoreWriter->EnableCsvOutput();
+		check(IoStatus.IsOk());
+#endif
 #endif
 	}
 
@@ -1345,15 +1386,33 @@ int32 CreateTarget(const TCHAR* OutputDir, const TCHAR* CookedDir, const TArray<
 #if !SKIP_WRITE_CONTAINER
 			FIoStatus IoStatus = ContainerTarget.IoStoreWriter->Initialize();
 			check(IoStatus.IsOk());
+#if OUTPUT_CONTAINER_CSV
+			IoStatus = ContainerTarget.IoStoreWriter->EnableCsvOutput();
+			check(IoStatus.IsOk());
+#endif
 #endif
 		}
 		for (const FContainerSourceFile& SourceFile : Container.SourceFiles)
 		{
-			FPackage* Package = FindOrAddPackage(*SourceFile.Path, CookedDir, Packages, PackageMap);
+			FFileStatData FileStatData = IFileManager::Get().GetStatData(*SourceFile.Path);
+			if (!FileStatData.bIsValid)
+			{
+				UE_LOG(LogIoStore, Warning, TEXT("File not found: '%s'"), *SourceFile.Path);
+			}
+
+			FString RelativeFileName = SourceFile.Path;
+			RelativeFileName.RemoveFromStart(CookedDir);
+			FPaths::NormalizeFilename(RelativeFileName);
+			RelativeFileName.RemoveFromStart(TEXT("/"));
+			RelativeFileName = TEXT("../../../") / RelativeFileName;
+			
+			FPackage* Package = FindOrAddPackage(*RelativeFileName, Packages, PackageMap, PackageOrderMap);
 			if (Package)
 			{
 				FContainerTargetFile& TargetFile = ContainerTarget.TargetFiles.AddDefaulted_GetRef();
+				TargetFile.Size = uint64(FileStatData.FileSize);
 				TargetFile.SourcePath = SourceFile.Path;
+				TargetFile.TargetPath = MoveTemp(RelativeFileName);
 				TargetFile.Package = Package;
 
 				FString Extension = FPaths::GetExtension(SourceFile.Path);
@@ -1638,6 +1697,7 @@ int32 CreateTarget(const TCHAR* OutputDir, const TCHAR* CookedDir, const TArray<
 			if (ImportData.Package)
 			{
 				Package->ImportedPackages.Add(ImportData.Package);
+				ImportData.Package->ImportedByPackages.Add(Package);
 			}
 		}
 	}
@@ -1753,6 +1813,9 @@ int32 CreateTarget(const TCHAR* OutputDir, const TCHAR* CookedDir, const TArray<
 	UE_LOG(LogIoStore, Display, TEXT("Building bundles..."));
 	BuildBundles(ExportGraph, Packages);
 
+	UE_LOG(LogIoStore, Display, TEXT("Building clusters..."));
+	BuildClusters(Packages);
+
 	TUniquePtr<FLargeMemoryWriter> StoreTocArchive = MakeUnique<FLargeMemoryWriter>(0, true);
 	TUniquePtr<FLargeMemoryWriter> ImportedPackagesArchive = MakeUnique<FLargeMemoryWriter>(0, true);
 	TUniquePtr<FLargeMemoryWriter> GlobalImportNamesArchive = MakeUnique<FLargeMemoryWriter>(0, true);
@@ -1832,6 +1895,18 @@ int32 CreateTarget(const TCHAR* OutputDir, const TCHAR* CookedDir, const TArray<
 		}
 	}
 
+	UE_LOG(LogIoStore, Display, TEXT("Finalizing package headers..."));
+	for (FPackage* Package : Packages)
+	{
+		FinalizePackageHeader(
+			Package,
+			NameMapBuilder,
+			ObjectExports,
+			GlobalExports,
+			GlobalImportsByFullName,
+			ExportBundleMetaEntries.GetData());
+	}
+
 	UE_LOG(LogIoStore, Display, TEXT("Serializing..."));
 
 	int32 TotalFileCount = 0;
@@ -1842,6 +1917,22 @@ int32 CreateTarget(const TCHAR* OutputDir, const TCHAR* CookedDir, const TArray<
 	int32 CurrentFileIndex = 0;
 	for (FContainerTargetSpec& ContainerTarget : ContainerTargets)
 	{
+		Algo::Sort(ContainerTarget.TargetFiles, [](const FContainerTargetFile& A, const FContainerTargetFile& B)
+		{
+			if (A.bIsBulkData != B.bIsBulkData)
+			{
+				return B.bIsBulkData;
+			}
+
+			if (A.Package->ClusterIndex != B.Package->ClusterIndex)
+			{
+				return A.Package->ClusterIndex < B.Package->ClusterIndex;
+			}
+			
+			uint64 ABundleOrder = A.Package->ExportBundles.Num() > 0 ? A.Package->ExportBundles[0].LoadOrder : MAX_uint64;
+			uint64 BBundleOrder = B.Package->ExportBundles.Num() > 0 ? B.Package->ExportBundles[0].LoadOrder : MAX_uint64;
+			return ABundleOrder < BBundleOrder;
+		});
 		for (FContainerTargetFile& TargetFile : ContainerTarget.TargetFiles)
 		{
 			UE_CLOG(CurrentFileIndex % 1000 == 0, LogIoStore, Display, TEXT("Serializing %d/%d: '%s'"), CurrentFileIndex, TotalFileCount, *TargetFile.SourcePath);
@@ -1866,14 +1957,7 @@ int32 CreateTarget(const TCHAR* OutputDir, const TCHAR* CookedDir, const TArray<
 			}
 			else
 			{
-				SerializePackageData(
-					TargetFile,
-					NameMapBuilder,
-					ObjectExports,
-					GlobalExports,
-					GlobalImportsByFullName,
-					ExportBundleMetaEntries.GetData(),
-					ContainerTarget.IoStoreWriter);
+				SerializePackageData(TargetFile.Package, ObjectExports, ContainerTarget.IoStoreWriter);
 			}
 		}
 	}
@@ -1899,7 +1983,7 @@ int32 CreateTarget(const TCHAR* OutputDir, const TCHAR* CookedDir, const TArray<
 		GlobalMetaArchive.Serialize(ExportBundleMetaEntries.GetData(), ExportBundleMetaByteCount);
 
 #if !SKIP_WRITE_CONTAINER
-		const FIoStatus Status = GlobalIoStoreWriter->Append(CreateIoChunkId(0, 0, EIoChunkType::LoaderGlobalMeta), FIoBuffer(FIoBuffer::Wrap, GlobalMetaArchive.GetData(), GlobalMetaArchive.TotalSize()));
+		const FIoStatus Status = GlobalIoStoreWriter->Append(CreateIoChunkId(0, 0, EIoChunkType::LoaderGlobalMeta), FIoBuffer(FIoBuffer::Wrap, GlobalMetaArchive.GetData(), GlobalMetaArchive.TotalSize()), TEXT("LoaderGlobalMeta"));
 		if (!Status.IsOk())
 		{
 			UE_LOG(LogIoStore, Error, TEXT("Failed to save global meta data to container file"));
@@ -1910,7 +1994,7 @@ int32 CreateTarget(const TCHAR* OutputDir, const TCHAR* CookedDir, const TArray<
 	{
 		UE_LOG(LogIoStore, Display, TEXT("Saving initial load meta data to container file"));
 #if !SKIP_WRITE_CONTAINER
-		const FIoStatus Status = GlobalIoStoreWriter->Append(CreateIoChunkId(0, 0, EIoChunkType::LoaderInitialLoadMeta), FIoBuffer(FIoBuffer::Wrap, InitialLoadArchive->GetData(), InitialLoadArchive->TotalSize()));
+		const FIoStatus Status = GlobalIoStoreWriter->Append(CreateIoChunkId(0, 0, EIoChunkType::LoaderInitialLoadMeta), FIoBuffer(FIoBuffer::Wrap, InitialLoadArchive->GetData(), InitialLoadArchive->TotalSize()), TEXT("LoaderInitialLoadMeta"));
 		if (!Status.IsOk())
 		{
 			UE_LOG(LogIoStore, Error, TEXT("Failed to save initial load meta data to container file"));
@@ -1932,9 +2016,9 @@ int32 CreateTarget(const TCHAR* OutputDir, const TCHAR* CookedDir, const TArray<
 
 #if !SKIP_WRITE_CONTAINER
 		FIoStatus NameStatus = GlobalIoStoreWriter->Append(CreateIoChunkId(0, 0, EIoChunkType::LoaderGlobalNames), 
-													 FIoBuffer(FIoBuffer::Wrap, Names.GetData(), Names.Num()));
+													 FIoBuffer(FIoBuffer::Wrap, Names.GetData(), Names.Num()), TEXT("LoaderGlobalNames"));
 		FIoStatus HashStatus = GlobalIoStoreWriter->Append(CreateIoChunkId(0, 0, EIoChunkType::LoaderGlobalNameHashes),
-													 FIoBuffer(FIoBuffer::Wrap, Hashes.GetData(), Hashes.Num()));
+													 FIoBuffer(FIoBuffer::Wrap, Hashes.GetData(), Hashes.Num()), TEXT("LoaderGlobalNameHashes"));
 #endif
 		
 		if (!NameStatus.IsOk() || !HashStatus.IsOk())
@@ -2043,6 +2127,93 @@ int32 CreateTarget(const TCHAR* OutputDir, const TCHAR* CookedDir, const TArray<
 	return 0;
 }
 
+static bool ParsePakResponseFile(const TCHAR* FilePath, TArray<FContainerSourceFile>& OutFiles)
+{
+	TArray<FString> ResponseFileContents;
+	if (!FFileHelper::LoadFileToStringArray(ResponseFileContents, FilePath))
+	{
+		UE_LOG(LogIoStore, Error, TEXT("Failed to read response file '%s'."), *FilePath);
+		return false;
+	}
+
+	for (const FString& ResponseLine : ResponseFileContents)
+	{
+		TArray<FString> SourceAndDest;
+		TArray<FString> Switches;
+
+		FString NextToken;
+		const TCHAR* ResponseLinePtr = *ResponseLine;
+		while (FParse::Token(ResponseLinePtr, NextToken, false))
+		{
+			if ((**NextToken == TCHAR('-')))
+			{
+				new(Switches) FString(NextToken.Mid(1));
+			}
+			else
+			{
+				new(SourceAndDest) FString(NextToken);
+			}
+		}
+
+		if (SourceAndDest.Num() == 0)
+		{
+			continue;
+		}
+
+		if (SourceAndDest.Num() != 2)
+		{
+			UE_LOG(LogIoStore, Error, TEXT("Invalid line in response file '%s'."), *ResponseLine);
+			return false;
+		}
+
+		FContainerSourceFile& FileEntry = OutFiles.AddDefaulted_GetRef();
+		FileEntry.Path = SourceAndDest[0];
+	}
+	return true;
+}
+
+static bool ParsePakOrderFile(const TCHAR* FilePath, TMap<FString, uint64>& OutMap)
+{
+	TArray<FString> OrderFileContents;
+	if (!FFileHelper::LoadFileToStringArray(OrderFileContents, FilePath))
+	{
+		UE_LOG(LogIoStore, Error, TEXT("Failed to read order file '%s'."), *FilePath);
+		return false;
+	}
+
+	uint64 LineNumber = 1;
+	for (const FString& OrderLine : OrderFileContents)
+	{
+		const TCHAR* OrderLinePtr = *OrderLine;
+		FString Path;
+		if (!FParse::Token(OrderLinePtr, Path, false))
+		{
+			UE_LOG(LogIoStore, Error, TEXT("Invalid line in order file '%s'."), *OrderLine);
+			return false;
+		}
+		FPaths::NormalizeFilename(Path);
+		uint64 Order = LineNumber;
+		FString OrderStr;
+		if (FParse::Token(OrderLinePtr, OrderStr, false))
+		{
+			if (!OrderStr.IsNumeric())
+			{
+				UE_LOG(LogIoStore, Error, TEXT("Invalid line in order file '%s'."), *OrderLine);
+				return false;
+			}
+			Order = FCString::Atoi64(*OrderStr);
+		}
+
+		if (!OutMap.Contains(Path))
+		{
+			OutMap.Emplace(MoveTemp(Path), Order);
+		}
+
+		++LineNumber;
+	}
+	return true;
+}
+
 int32 CreateIoStoreContainerFiles(const TCHAR* CmdLine)
 {
 	UE_LOG(LogIoStore, Display, TEXT("==================== IoStore Utils ===================="));
@@ -2091,49 +2262,23 @@ int32 CreateIoStoreContainerFiles(const TCHAR* CmdLine)
 				return -1;
 			}
 
-			TArray<FString> ResponseFileContents;
-			if (!FFileHelper::LoadFileToStringArray(ResponseFileContents, *ResponseFilePath))
+			if (!ParsePakResponseFile(*ResponseFilePath, ContainerSpec.SourceFiles))
 			{
-				UE_LOG(LogIoStore, Error, TEXT("Failed to read response file '%s'."), *ResponseFilePath);
 				return -1;
 			}
+		}
 
-			for (const FString& ResponseLine : ResponseFileContents)
+		TMap<FString, uint64> PackageOrderMap;
+		FString PackageOrderFilePath;
+		if (FParse::Value(FCommandLine::Get(), TEXT("PackageOrder="), PackageOrderFilePath))
+		{
+			if (!ParsePakOrderFile(*PackageOrderFilePath, PackageOrderMap))
 			{
-				TArray<FString> SourceAndDest;
-				TArray<FString> Switches;
-
-				FString NextToken;
-				const TCHAR* ResponseLinePtr = *ResponseLine;
-				while (FParse::Token(ResponseLinePtr, NextToken, false))
-				{
-					if ((**NextToken == TCHAR('-')))
-					{
-						new(Switches) FString(NextToken.Mid(1));
-					}
-					else
-					{
-						new(SourceAndDest) FString(NextToken);
-					}
-				}
-
-				if (SourceAndDest.Num() == 0)
-				{
-					continue;
-				}
-
-				if (SourceAndDest.Num() != 2)
-				{
-					UE_LOG(LogIoStore, Error, TEXT("Invalid line in response file '%s'."), *ResponseLine);
-					return -1;
-				}
-
-				FContainerSourceFile& FileEntry = ContainerSpec.SourceFiles.AddDefaulted_GetRef();
-				FileEntry.Path = SourceAndDest[0];
+				return -1;
 			}
 		}
 
-		int32 ReturnValue = CreateTarget(*OutputDirectory, *CookedDirectory, Containers);
+		int32 ReturnValue = CreateTarget(*OutputDirectory, *CookedDirectory, Containers, PackageOrderMap);
 		if (ReturnValue != 0)
 		{
 			return ReturnValue;
@@ -2183,7 +2328,7 @@ int32 CreateIoStoreContainerFiles(const TCHAR* CmdLine)
 
 			UE_LOG(LogIoStore, Display, TEXT("Creating target: '%s' using output path: '%s'"), *TargetPlatform->PlatformName(), *TargetCookedProjectDirectory);
 
-			int32 ReturnValue = CreateTarget(*TargetCookedProjectDirectory, *TargetCookedDirectory, Containers);
+			int32 ReturnValue = CreateTarget(*TargetCookedProjectDirectory, *TargetCookedDirectory, Containers, TMap<FString, uint64>());
 			if (ReturnValue != 0)
 			{
 				return ReturnValue;
diff --git a/Engine/Source/Runtime/Core/Private/IO/IoStore.cpp b/Engine/Source/Runtime/Core/Private/IO/IoStore.cpp
--- a/Engine/Source/Runtime/Core/Private/IO/IoStore.cpp
+++ b/Engine/Source/Runtime/Core/Private/IO/IoStore.cpp
@@ -62,7 +62,21 @@ public:
 		return FIoStatus::Ok;
 	}
 
-	UE_NODISCARD FIoStatus Append(FIoChunkId ChunkId, FIoBuffer Chunk)
+	FIoStatus EnableCsvOutput()
+	{
+		FString CsvFilePath = Environment.GetPath() + TEXT(".csv");
+		CsvArchive.Reset(IFileManager::Get().CreateFileWriter(*CsvFilePath));
+		if (!CsvArchive)
+		{
+			return FIoStatusBuilder(EIoErrorCode::FileOpenFailed) << TEXT("Failed to open IoStore CSV file '") << *CsvFilePath << TEXT("'");
+		}
+		ANSICHAR Header[] = "Name,Offset,Size\n";
+		CsvArchive->Serialize(Header, sizeof(Header) - 1);
+
+		return FIoStatus::Ok;
+	}
+
+	UE_NODISCARD FIoStatus Append(FIoChunkId ChunkId, FIoBuffer Chunk, const TCHAR* Name)
 	{
 		if (!ContainerFileHandle)
 		{
@@ -102,6 +116,13 @@ public:
 		{
 			Toc.Add(ChunkId, TocEntry);
 
+			if (CsvArchive)
+			{
+				ANSICHAR Line[MAX_SPRINTF];
+				FCStringAnsi::Sprintf(Line, "%s,%lld,%lld\n", TCHAR_TO_ANSI(Name), TocEntry.GetOffset(), TocEntry.GetLength());
+				CsvArchive->Serialize(Line, FCStringAnsi::Strlen(Line));
+			}
+
 			return FIoStatus::Ok;
 		}
 		else
@@ -182,6 +203,7 @@ private:
 	TMap<FIoChunkId, FIoStoreTocEntry>	Toc;
 	TUniquePtr<IFileHandle>				ContainerFileHandle;
 	TUniquePtr<IFileHandle>				TocFileHandle;
+	TUniquePtr<FArchive>				CsvArchive;
 	bool								IsMetadataDirty = true;
 };
 
@@ -200,9 +222,14 @@ FIoStatus FIoStoreWriter::Initialize()
 	return Impl->Initialize();
 }
 
-FIoStatus FIoStoreWriter::Append(FIoChunkId ChunkId, FIoBuffer Chunk)
+FIoStatus FIoStoreWriter::EnableCsvOutput()
+{
+	return Impl->EnableCsvOutput();
+}
+
+FIoStatus FIoStoreWriter::Append(FIoChunkId ChunkId, FIoBuffer Chunk, const TCHAR* Name)
 {
-	return Impl->Append(ChunkId, Chunk);
+	return Impl->Append(ChunkId, Chunk, Name);
 }
 
 FIoStatus FIoStoreWriter::MapPartialRange(FIoChunkId OriginalChunkId, uint64 Offset, uint64 Length, FIoChunkId ChunkIdPartialRange)
diff --git a/Engine/Source/Runtime/Core/Public/IO/IoDispatcher.h b/Engine/Source/Runtime/Core/Public/IO/IoDispatcher.h
--- a/Engine/Source/Runtime/Core/Public/IO/IoDispatcher.h
+++ b/Engine/Source/Runtime/Core/Public/IO/IoDispatcher.h
@@ -789,7 +789,8 @@ public:
 	FIoStoreWriter& operator=(const FIoStoreWriter&) = delete;
 
 	CORE_API FIoStatus	Initialize();
-	CORE_API FIoStatus	Append(FIoChunkId ChunkId, FIoBuffer Chunk);
+	CORE_API FIoStatus	EnableCsvOutput();
+	CORE_API FIoStatus	Append(FIoChunkId ChunkId, FIoBuffer Chunk, const TCHAR* Name);
 
 	/**
 	 * Creates an addressable range in an already mapped Chunk.
