IoStore: Replace individual file GetStatData calls with one directory scan of the cooked directory.

diff --git a/Engine/Source/Developer/IoStoreUtilities/Private/IoStoreUtilities.cpp b/Engine/Source/Developer/IoStoreUtilities/Private/IoStoreUtilities.cpp
--- a/Engine/Source/Developer/IoStoreUtilities/Private/IoStoreUtilities.cpp
+++ b/Engine/Source/Developer/IoStoreUtilities/Private/IoStoreUtilities.cpp
@@ -43,7 +43,7 @@ DEFINE_LOG_CATEGORY_STATIC(LogIoStore, Log, All);
 
 struct FContainerSourceFile 
 {
-	FString Path;
+	FString NormalizedPath;
 };
 
 struct FContainerSourceSpec
@@ -52,15 +52,30 @@ struct FContainerSourceSpec
 	TArray<FContainerSourceFile> SourceFiles;
 };
 
+struct FCookedFileStatData
+{
+	static constexpr int32 NumPackageExtensions = 2;
+	static constexpr TCHAR* Extensions[] = { TEXT("umap"), TEXT("uasset"), TEXT("ubulk"), TEXT("uptnl") };
+	enum FFileExt { UMap, UAsset, UBulk, UPtnl };
+	enum FFileType { PackageHeader, BulkData };
+
+	int64 FileSize = 0;
+	FFileType FileType = PackageHeader;
+	FFileExt FileExt = UMap;
+};
+
+using FCookedFileStatMap = TMap<FString, FCookedFileStatData>;
+
 struct FPackage;
 
 struct FContainerTargetFile
 {
 	FPackage* Package = nullptr;
-	FString SourcePath;
+	FString NormalizedSourcePath;
 	FString TargetPath;
 	uint64 Size = 0;
 	bool bIsBulkData = false;
+	bool bIsOptionalBulkData = false;
 };
 
 struct FContainerTargetSpec
@@ -1322,7 +1337,12 @@ void SerializePackageData(
 #endif
 }
 
-int32 CreateTarget(const TCHAR* OutputDir, const TCHAR* CookedDir, const TArray<FContainerSourceSpec>& Containers, const TMap<FString, uint64> PackageOrderMap)
+int32 CreateTarget(
+	const TCHAR* OutputDir,
+	const TCHAR* CookedDir,
+	const TArray<FContainerSourceSpec>& Containers,
+	const FCookedFileStatMap& CookedFileStatMap,
+	const TMap<FString, uint64> PackageOrderMap)
 {
 	TGuardValue<int32> GuardAllowUnversionedContentInEditor(GAllowUnversionedContentInEditor, 1);
 
@@ -1401,38 +1421,37 @@ int32 CreateTarget(const TCHAR* OutputDir, const TCHAR* CookedDir, const TArray<
 #endif
 #endif
 		}
+
+		int32 CookedDirLen = FCString::Strlen(CookedDir) + 1;
 		for (const FContainerSourceFile& SourceFile : Container.SourceFiles)
 		{
-			FFileStatData FileStatData = IFileManager::Get().GetStatData(*SourceFile.Path);
-			if (!FileStatData.bIsValid)
+			const FCookedFileStatData* CookedFileStatData = CookedFileStatMap.Find(SourceFile.NormalizedPath);
+			if (!CookedFileStatData)
 			{
-				UE_LOG(LogIoStore, Warning, TEXT("File not found: '%s'"), *SourceFile.Path);
+				UE_LOG(LogIoStore, Warning, TEXT("File not found: '%s'"), *SourceFile.NormalizedPath);
 			}
 
-			FString RelativeFileName = SourceFile.Path;
-			RelativeFileName.RemoveFromStart(CookedDir);
-			FPaths::NormalizeFilename(RelativeFileName);
-			RelativeFileName.RemoveFromStart(TEXT("/"));
-			RelativeFileName = TEXT("../../../") / RelativeFileName;
+			FString RelativeFileName = TEXT("../../../");
+			RelativeFileName.AppendChars(*SourceFile.NormalizedPath + CookedDirLen, SourceFile.NormalizedPath.Len() - CookedDirLen);
 			
 			FPackage* Package = FindOrAddPackage(*RelativeFileName, Packages, PackageMap, PackageOrderMap);
 			if (Package)
 			{
 				FContainerTargetFile& TargetFile = ContainerTarget.TargetFiles.AddDefaulted_GetRef();
-				TargetFile.Size = uint64(FileStatData.FileSize);
-				TargetFile.SourcePath = SourceFile.Path;
+				TargetFile.Size = uint64(CookedFileStatData->FileSize);
+				TargetFile.NormalizedSourcePath = SourceFile.NormalizedPath;
 				TargetFile.TargetPath = MoveTemp(RelativeFileName);
 				TargetFile.Package = Package;
 
-				FString Extension = FPaths::GetExtension(SourceFile.Path);
-				if (Extension == TEXT("uasset") || Extension == TEXT("umap"))
+				if (CookedFileStatData->FileType == FCookedFileStatData::PackageHeader)
 				{
 					TargetFile.bIsBulkData = false;
-					Package->FileName = SourceFile.Path;
+					Package->FileName = SourceFile.NormalizedPath;
 				}
 				else
 				{
 					TargetFile.bIsBulkData = true;
+					TargetFile.bIsOptionalBulkData = CookedFileStatData->FileExt == FCookedFileStatData::UPtnl;
 				}
 			}
 		}
@@ -1944,24 +1963,18 @@ int32 CreateTarget(const TCHAR* OutputDir, const TCHAR* CookedDir, const TArray<
 		});
 		for (FContainerTargetFile& TargetFile : ContainerTarget.TargetFiles)
 		{
-			UE_CLOG(CurrentFileIndex % 1000 == 0, LogIoStore, Display, TEXT("Serializing %d/%d: '%s'"), CurrentFileIndex, TotalFileCount, *TargetFile.SourcePath);
+			UE_CLOG(CurrentFileIndex % 1000 == 0, LogIoStore, Display, TEXT("Serializing %d/%d: '%s'"), CurrentFileIndex, TotalFileCount, *TargetFile.NormalizedSourcePath);
 			++CurrentFileIndex;
 			if (TargetFile.bIsBulkData)
 			{
 				if (bWithBulkDataManifest)
 				{
-					FString NormalizedSourcePath = TargetFile.SourcePath;
-					FPaths::NormalizeFilename(NormalizedSourcePath);
+					FPackageStoreBulkDataManifest::EBulkdataType BulkDataType =
+						TargetFile.bIsOptionalBulkData ?
+						FPackageStoreBulkDataManifest::EBulkdataType::Optional :
+						FPackageStoreBulkDataManifest::EBulkdataType::Normal;
 
-					FString Extension = FPaths::GetExtension(NormalizedSourcePath);
-					if (Extension == TEXT("uptnl"))
-					{
-						WriteBulkData(NormalizedSourcePath, FPackageStoreBulkDataManifest::EBulkdataType::Optional, *TargetFile.Package, BulkDataManifest, ContainerTarget.IoStoreWriter);
-					}
-					else
-					{
-						WriteBulkData(NormalizedSourcePath, FPackageStoreBulkDataManifest::EBulkdataType::Normal, *TargetFile.Package, BulkDataManifest, ContainerTarget.IoStoreWriter);
-					}
+					WriteBulkData(TargetFile.NormalizedSourcePath, BulkDataType, *TargetFile.Package, BulkDataManifest, ContainerTarget.IoStoreWriter);
 				}
 			}
 			else
@@ -2175,8 +2188,10 @@ static bool ParsePakResponseFile(const TCHAR* FilePath, TArray<FContainerSourceF
 			return false;
 		}
 
+		FPaths::NormalizeFilename(SourceAndDest[0]);
+
 		FContainerSourceFile& FileEntry = OutFiles.AddDefaulted_GetRef();
-		FileEntry.Path = SourceAndDest[0];
+		FileEntry.NormalizedPath = MoveTemp(SourceAndDest[0]);
 	}
 	return true;
 }
@@ -2223,6 +2238,67 @@ static bool ParsePakOrderFile(const TCHAR* FilePath, TMap<FString, uint64>& OutM
 	return true;
 }
 
+class FCookedFileVisitor : public IPlatformFile::FDirectoryStatVisitor
+{
+	FCookedFileStatMap& CookedFileStatMap;
+	FContainerSourceSpec* ContainerSpec = nullptr;
+
+public:
+	FCookedFileVisitor(FCookedFileStatMap& InCookedFileSizes, FContainerSourceSpec* InContainerSpec)
+		: CookedFileStatMap(InCookedFileSizes)
+		, ContainerSpec(InContainerSpec)
+	{}
+
+	FCookedFileVisitor(FCookedFileStatMap& InFileSizes)
+		: CookedFileStatMap(InFileSizes)
+	{}
+
+	virtual bool Visit(const TCHAR* FilenameOrDirectory, const FFileStatData& StatData)
+	{
+		if (StatData.bIsDirectory)
+		{
+			return true;
+		}
+
+		const TCHAR* Extension = FCString::Strrchr(FilenameOrDirectory, '.');
+		if (!Extension || *(++Extension) == TEXT('\0'))
+		{
+			return true;
+		}
+
+		int32 ExtIndex = 0;
+		for (ExtIndex = 0; ExtIndex < UE_ARRAY_COUNT(FCookedFileStatData::Extensions); ++ExtIndex)
+		{
+			if (0 == FCString::Stricmp(Extension, FCookedFileStatData::Extensions[ExtIndex]))
+				break;
+		}
+
+		if (ExtIndex >= UE_ARRAY_COUNT(FCookedFileStatData::Extensions))
+		{
+			return true;
+		}
+
+		FString Path = FilenameOrDirectory;
+		FPaths::NormalizeFilename(Path);
+
+		if (ContainerSpec)
+		{
+			FContainerSourceFile& FileEntry = ContainerSpec->SourceFiles.AddDefaulted_GetRef();
+			FileEntry.NormalizedPath = Path;
+		}
+
+		FCookedFileStatData& CookedFileStatData = CookedFileStatMap.Add(MoveTemp(Path));
+		CookedFileStatData.FileSize = StatData.FileSize;
+		CookedFileStatData.FileExt = FCookedFileStatData::FFileExt(ExtIndex);
+		CookedFileStatData.FileType = 
+			ExtIndex < FCookedFileStatData::NumPackageExtensions ?
+			FCookedFileStatData::PackageHeader : 
+			FCookedFileStatData::BulkData;
+
+		return true;
+	}
+};
+
 int32 CreateIoStoreContainerFiles(const TCHAR* CmdLine)
 {
 	UE_LOG(LogIoStore, Display, TEXT("==================== IoStore Utils ===================="));
@@ -2287,7 +2363,13 @@ int32 CreateIoStoreContainerFiles(const TCHAR* CmdLine)
 			}
 		}
 
-		int32 ReturnValue = CreateTarget(*OutputDirectory, *CookedDirectory, Containers, PackageOrderMap);
+		UE_LOG(LogIoStore, Display, TEXT("Searching for cooked files..."));
+		FCookedFileStatMap CookedFileStatMap;
+		FCookedFileVisitor CookedFileVistor(CookedFileStatMap);
+		IFileManager::Get().IterateDirectoryStatRecursively(*CookedDirectory, CookedFileVistor);
+		UE_LOG(LogIoStore, Display, TEXT("Found '%d' files"), CookedFileStatMap.Num());
+
+		int32 ReturnValue = CreateTarget(*OutputDirectory, *CookedDirectory, Containers, CookedFileStatMap, PackageOrderMap);
 		if (ReturnValue != 0)
 		{
 			return ReturnValue;
@@ -2302,42 +2384,19 @@ int32 CreateIoStoreContainerFiles(const TCHAR* CmdLine)
 			const FString TargetCookedDirectory = FPaths::ProjectSavedDir() / TEXT("Cooked") / TargetPlatform->PlatformName();
 			const FString TargetCookedProjectDirectory = TargetCookedDirectory / FApp::GetProjectName();
 
-			UE_LOG(LogIoStore, Display, TEXT("Searching for cooked files..."));
-			class FCookedFileVisitor
-				: public IPlatformFile::FDirectoryVisitor
-			{
-				FContainerSourceSpec& ContainerSpec;
-
-			public:
-				FCookedFileVisitor(FContainerSourceSpec& InContainerSpec)
-					: ContainerSpec(InContainerSpec)
-				{
-				}
-
-				virtual bool Visit(const TCHAR* FilenameOrDirectory, bool bIsDirectory)
-				{
-					if (bIsDirectory == false)
-					{
-						FString Filename(FilenameOrDirectory);
-						if (Filename.EndsWith(TEXT(".uasset")) || Filename.EndsWith(TEXT(".umap")) || Filename.EndsWith(TEXT(".ubulk")) || Filename.EndsWith(TEXT(".uptnl")))
-						{
-							FContainerSourceFile& FileEntry = ContainerSpec.SourceFiles.AddDefaulted_GetRef();
-							FileEntry.Path = MoveTemp(FilenameOrDirectory);
-						}
-					}
-					return true;
-				}
-			};
-			
+			FCookedFileStatMap CookedFileStatMap;
 			TArray<FContainerSourceSpec> Containers;
 			FContainerSourceSpec& ContainerSpec = Containers.AddDefaulted_GetRef();
-			FCookedFileVisitor CookedFileVistor(ContainerSpec);
-			FPlatformFileManager::Get().GetPlatformFile().IterateDirectoryRecursively(*TargetCookedDirectory, CookedFileVistor);
+			FCookedFileVisitor CookedFileVistor(CookedFileStatMap, &ContainerSpec);
+
+			UE_LOG(LogIoStore, Display, TEXT("Searching for cooked files..."));
+			IFileManager::Get().IterateDirectoryStatRecursively(*TargetCookedDirectory, CookedFileVistor);
+
 			UE_LOG(LogIoStore, Display, TEXT("Found '%d' files"), ContainerSpec.SourceFiles.Num());
 
 			UE_LOG(LogIoStore, Display, TEXT("Creating target: '%s' using output path: '%s'"), *TargetPlatform->PlatformName(), *TargetCookedProjectDirectory);
 
-			int32 ReturnValue = CreateTarget(*TargetCookedProjectDirectory, *TargetCookedDirectory, Containers, TMap<FString, uint64>());
+			int32 ReturnValue = CreateTarget(*TargetCookedProjectDirectory, *TargetCookedDirectory, Containers, CookedFileStatMap, TMap<FString, uint64>());
 			if (ReturnValue != 0)
 			{
 				return ReturnValue;
