IoStoreUtilities: Refactored custom serialization to use standard serialization and added parallel parsing of uasset files

diff --git a/Engine/Source/Developer/IoStoreUtilities/Private/IoStoreUtilities.cpp b/Engine/Source/Developer/IoStoreUtilities/Private/IoStoreUtilities.cpp
--- a/Engine/Source/Developer/IoStoreUtilities/Private/IoStoreUtilities.cpp
+++ b/Engine/Source/Developer/IoStoreUtilities/Private/IoStoreUtilities.cpp
@@ -28,6 +28,8 @@
 #include "UObject/Package.h"
 #include "Algo/Find.h"
 #include "Misc/FileHelper.h"
+#include "Misc/ScopeLock.h"
+#include "Async/ParallelFor.h"
 
 //PRAGMA_DISABLE_OPTIMIZATION
 
@@ -84,6 +86,36 @@ struct FContainerTargetSpec
 	TArray<FContainerTargetFile> TargetFiles;
 };
 
+class FNameProxyArchive
+	: public FArchiveProxy
+{
+public:
+	using FArchiveProxy::FArchiveProxy;
+
+	FNameProxyArchive(FArchive& InAr, const TArray<FNameEntryId>& InNameMap)
+		: FArchiveProxy(InAr)
+		, NameMap(InNameMap) { }
+
+	virtual FArchive& operator<<(FName& Name) override
+	{
+		int32 NameIndex, Number;
+		InnerArchive << NameIndex << Number;
+
+		if (!NameMap.IsValidIndex(NameIndex))
+		{
+			UE_LOG(LogIoStore, Fatal, TEXT("Bad name index %i/%i"), NameIndex, NameMap.Num());
+		}
+
+		const FNameEntryId MappedName = NameMap[NameIndex];
+		Name = FName::CreateFromDisplayId(MappedName, Number);
+
+		return *this;
+	}
+
+private:
+	const TArray<FNameEntryId>& NameMap;
+};
+
 class FNameMapBuilder
 {
 public:
@@ -96,6 +128,7 @@ public:
 			Index = NameIndices.Num();
 			NameMap.Add(Id);
 		}
+#if OUTPUT_NAMEMAP_CSV
 		// debug counts
 		{
 			const int32 Number = Name.GetNumber();
@@ -114,6 +147,7 @@ public:
 				}
 			}
 		}
+#endif
 	}
 
 	int32 MapName(const FName& Name) const
@@ -131,7 +165,6 @@ public:
 		A << NameIndex << NameNumber;
 	}
 
-
 	const TArray<FNameEntryId>& GetNameMap() const
 	{
 		return NameMap;
@@ -170,7 +203,9 @@ public:
 private:
 	TMap<FNameEntryId, int32> NameIndices;
 	TArray<FNameEntryId> NameMap;
+#if OUTPUT_NAMEMAP_CSV
 	TMap<FNameEntryId, TTuple<int32,int32,int32>> DebugNameCounts; // <Number0Count,OtherNumberCount,MaxNumber>
+#endif
 };
 
 #if OUTPUT_CHUNKID_DIRECTORY
@@ -374,10 +409,10 @@ struct FPackage
 	uint32 PackageFlags = 0;
 	int32 NameCount = 0;
 	int32 ImportCount = 0;
-	int32 ImportOffset = 0;
 	int32 ExportCount = 0;
 	int32 FirstGlobalImport = -1;
 	int32 GlobalImportCount = -1;
+	int32 ImportIndexOffset = -1;
 	int32 ExportIndexOffset = -1;
 	int32 PreloadIndexOffset = -1;
 	int32 FirstExportBundleMetaEntry = -1;
@@ -1137,11 +1172,7 @@ void FinalizePackageHeader(
 
 	int32 InternalArcCount = Package->InternalArcs.Num();
 	GraphArchive << InternalArcCount;
-	for (FArc& InternalArc : Package->InternalArcs)
-	{
-		GraphArchive << InternalArc.FromNodeIndex;
-		GraphArchive << InternalArc.ToNodeIndex;
-	}
+	GraphArchive.Serialize(Package->InternalArcs.GetData(), Package->InternalArcs.Num() * sizeof(FArc));
 
 	int32 ReferencedPackagesCount = Package->ExternalArcs.Num();
 	GraphArchive << ReferencedPackagesCount;
@@ -1153,11 +1184,7 @@ void FinalizePackageHeader(
 
 		GraphArchive << ImportedPackage->GlobalPackageId;
 		GraphArchive << ExternalArcCount;
-		for (FArc& ExternalArc : Arcs)
-		{
-			GraphArchive << ExternalArc.FromNodeIndex;
-			GraphArchive << ExternalArc.ToNodeIndex;
-		}
+		GraphArchive.Serialize(Arcs.GetData(), ExternalArcCount * sizeof(FArc));
 	}
 	Package->UGraphSize = GraphArchive.Tell();
 
@@ -1168,33 +1195,27 @@ void FinalizePackageHeader(
 		const FObjectExport& ObjectExport = ObjectExports[Package->ExportIndexOffset + I];
 		const FExportData& ExportData = GlobalExports[Package->Exports[I]];
 
-		int64 SerialSize = ObjectExport.SerialSize;
-		ExportMapArchive << SerialSize;
-		NameMapBuilder.SerializeName(ExportMapArchive, ObjectExport.ObjectName);
-		FPackageIndex OuterIndex = ExportData.OuterIndex;
-		ExportMapArchive << OuterIndex;
-		FPackageIndex ClassIndex = ExportData.ClassIndex;
-		ExportMapArchive << ClassIndex;
-		FPackageIndex SuperIndex = ExportData.SuperIndex;
-		ExportMapArchive << SuperIndex;
-		FPackageIndex TemplateIndex = ExportData.TemplateIndex;
-		ExportMapArchive << TemplateIndex;
-		int32 GlobalImportIndex = ExportData.GlobalImportIndex;
-		ExportMapArchive << GlobalImportIndex;
-		uint32 ObjectFlags = ObjectExport.ObjectFlags;
-		ExportMapArchive << ObjectFlags;
-		uint8 FilterFlags = uint8(EExportFilterFlags::None);
+		FExportMapEntry ExportMapEntry;
+		ExportMapEntry.SerialSize = ObjectExport.SerialSize;
+		ExportMapEntry.ObjectName[0] = NameMapBuilder.MapName(ObjectExport.ObjectName);
+		ExportMapEntry.ObjectName[1] = ObjectExport.ObjectName.GetNumber();
+		ExportMapEntry.OuterIndex = ExportData.OuterIndex;
+		ExportMapEntry.ClassIndex = ExportData.ClassIndex;
+		ExportMapEntry.SuperIndex = ExportData.SuperIndex;
+		ExportMapEntry.TemplateIndex = ExportData.TemplateIndex;
+		ExportMapEntry.GlobalImportIndex = ExportData.GlobalImportIndex;
+		ExportMapEntry.ObjectFlags = ObjectExport.ObjectFlags;
+		ExportMapEntry.FilterFlags = EExportFilterFlags::None;
 		if (ObjectExport.bNotForClient)
 		{
-			FilterFlags = uint8(EExportFilterFlags::NotForClient);
+			ExportMapEntry.FilterFlags = EExportFilterFlags::NotForClient;
 		}
 		else if (ObjectExport.bNotForServer)
 		{
-			FilterFlags = uint8(EExportFilterFlags::NotForServer);
+			ExportMapEntry.FilterFlags = EExportFilterFlags::NotForServer;
 		}
-		ExportMapArchive << FilterFlags;
-		uint8 Pad = 0;
-		ExportMapArchive.Serialize(&Pad, 7);
+
+		ExportMapArchive << ExportMapEntry;
 	}
 	Package->ExportMapSize = ExportMapArchive.Tell();
 
@@ -1203,18 +1224,17 @@ void FinalizePackageHeader(
 	int32 ExportBundleEntryIndex = 0;
 	for (FExportBundle& ExportBundle : Package->ExportBundles)
 	{
-		ExportBundlesArchive << ExportBundleEntryIndex;
-		int32 EntryCount = ExportBundle.Nodes.Num();
-		ExportBundlesArchive << EntryCount;
-		ExportBundleEntryIndex += ExportBundle.Nodes.Num();
+		const int32 EntryCount = ExportBundle.Nodes.Num();
+		FExportBundleHeader ExportBundleHeader { ExportBundleEntryIndex, EntryCount };
+		ExportBundlesArchive << ExportBundleHeader ;
+
+		ExportBundleEntryIndex += EntryCount; 
 	}
 	for (FExportBundle& ExportBundle : Package->ExportBundles)
 	{
 		for (FExportGraphNode* ExportNode : ExportBundle.Nodes)
 		{
-			uint32 CommandType = uint32(ExportNode->BundleEntry.CommandType);
-			ExportBundlesArchive << ExportNode->BundleEntry.LocalExportIndex;
-			ExportBundlesArchive << CommandType;
+			ExportBundlesArchive << ExportNode->BundleEntry;
 		}
 	}
 	Package->ExportBundlesSize = ExportBundlesArchive.Tell();
@@ -1359,7 +1379,7 @@ int32 CreateTarget(
 
 	FNameMapBuilder NameMapBuilder;
 
-	uint64 NameSize = 0;
+	FCriticalSection GlobalsCriticalSection;
 	TArray<FObjectImport> ObjectImports;
 	TArray<FObjectExport> ObjectExports;
 	TArray<FImportData> GlobalImports;
@@ -1368,11 +1388,6 @@ int32 CreateTarget(
 	TMap<FString, int32> GlobalExportsByFullName;
 	TArray<FString> TempFullNames;
 	TArray<FPackageIndex> PreloadDependencies;
-	uint64 UPackageImports = 0;
-	TArray<int32> ImportPreloadCounts;
-	TArray<int32> ExportPreloadCounts;
-	uint64 ImportPreloadCount = 0;
-	uint64 ExportPreloadCount = 0;
 	FExportGraph ExportGraph;
 	TArray<FExportBundleMetaEntry> ExportBundleMetaEntries;
 
@@ -1457,166 +1472,184 @@ int32 CreateTarget(
 		}
 	}
 
-	ImportPreloadCounts.AddDefaulted(Packages.Num());
-	ExportPreloadCounts.AddDefaulted(Packages.Num());
-
-	for (FPackage* Package : Packages)
 	{
-		UE_CLOG(Package->GlobalPackageId % 1000 == 0, LogIoStore, Display, TEXT("Parsing %d/%d: '%s'"), Package->GlobalPackageId, Packages.Num(), *Package->FileName);
+		UE_LOG(LogIoStore, Display, TEXT("Parsing packages..."));
+
+		TAtomic<int32> ReadCount = 0;
+		TAtomic<int32> ParseCount = 0;
+		const int32 TotalPackageCount = Packages.Num();
 
-		FPackageFileSummary Summary;
-		TUniquePtr<FArchive> Ar(IFileManager::Get().CreateFileReader(*Package->FileName));
-		check(Ar);
-		*Ar << Summary;
+		TArray<FPackageFileSummary> PackageFileSummaries;
+		TArray<TUniquePtr<uint8[]>> PackageAssetBuffers;
 
-		Package->UAssetSize = Ar->TotalSize();
-		Package->SummarySize = Ar->Tell();
-		Package->NameCount = Summary.NameCount;
-		Package->ImportCount = Summary.ImportCount;
-		Package->ExportCount = Summary.ExportCount;
-		Package->PackageFlags = Summary.PackageFlags;
-		Package->BulkDataStartOffset = Summary.BulkDataStartOffset;
+		PackageFileSummaries.SetNum(TotalPackageCount);
+		PackageAssetBuffers.SetNum(TotalPackageCount);
 
-		if (Summary.NameCount > 0)
+		ParallelFor(TotalPackageCount,[
+			&ReadCount,
+			&PackageAssetBuffers,
+			&PackageFileSummaries,
+			&Packages,
+			&ObjectImports,
+			&ObjectExports,
+			&PreloadDependencies,
+			&GlobalsCriticalSection](int32 Index)
 		{
-			Ar->Seek(Summary.NameOffset);
-			uint64 LastOffset = Summary.NameOffset;
+			TUniquePtr<uint8[]>& PackageBuffer = PackageAssetBuffers[Index];
+			FPackageFileSummary& Summary = PackageFileSummaries[Index];
+			FPackage& Package = *Packages[Index];
 
-			Package->NameMap.Reserve(Summary.NameCount);
-			Package->NameIndices.Reserve(Summary.NameCount);
-			FNameEntrySerialized NameEntry(ENAME_LinkerConstructor);
+			UE_CLOG(++ReadCount % 1000 == 0, LogIoStore, Display, TEXT("Reading %d/%d: '%s'"), ReadCount.Load(), Packages.Num(), *Package.FileName);
 
-			for (int32 I = 0; I < Summary.NameCount; ++I)
 			{
-				*Ar << NameEntry;
-				FName Name(NameEntry);
-				NameMapBuilder.MarkNameAsReferenced(Name);
-				Package->NameMap.Emplace(Name.GetDisplayIndex());
-				Package->NameIndices.Add(NameMapBuilder.MapName(Name));
+				TUniquePtr<FArchive> FileAr(IFileManager::Get().CreateFileReader(*Package.FileName));
+				check(FileAr);
+				Package.UAssetSize = FileAr->TotalSize();
+
+				PackageBuffer = MakeUnique<uint8[]>(Package.UAssetSize);
+				FileAr->Serialize(PackageBuffer.Get(), Package.UAssetSize);
+				FileAr->Close();
 			}
 
-			NameSize += Ar->Tell() - Summary.NameOffset;
-		}
+			TArrayView<const uint8> MemView(PackageBuffer.Get(), Package.UAssetSize);
+			FMemoryReaderView Ar(MemView);
+			Ar << Summary;
 
-		auto DeserializeName = [&](FArchive& A, FName& N)
-		{
-			int32 DisplayIndex, NameNumber;
-			A << DisplayIndex << NameNumber;
-			FNameEntryId DisplayEntry = Package->NameMap[DisplayIndex];
-			N = FName::CreateFromDisplayId(DisplayEntry, NameNumber);
-		};
+			Package.SummarySize = Ar.Tell();
+			Package.NameCount = Summary.NameCount;
+			Package.ImportCount = Summary.ImportCount;
+			Package.ExportCount = Summary.ExportCount;
+			Package.PackageFlags = Summary.PackageFlags;
+			Package.BulkDataStartOffset = Summary.BulkDataStartOffset;
 
-		if (Summary.ImportCount > 0)
-		{
-			Ar->Seek(Summary.ImportOffset);
+			if (Summary.ImportCount > 0)
+			{
+				FScopeLock _(&GlobalsCriticalSection);
+				Package.ImportIndexOffset = ObjectImports.Num();
+				ObjectImports.AddUninitialized(Summary.ImportCount);
+			}
+			
+			if (Summary.PreloadDependencyCount > 0)
+			{
+				FScopeLock _(&GlobalsCriticalSection);
+				Package.PreloadIndexOffset = PreloadDependencies.Num();
+				PreloadDependencies.AddUninitialized(Summary.PreloadDependencyCount);
+			}
 
-			int32 NumPackages = 0;
-			int32 BaseIndex = ObjectImports.Num();
-			ObjectImports.AddUninitialized(Summary.ImportCount);
-			TArray<FString> ImportNames;
-			ImportNames.Reserve(Summary.ImportCount);
-			for (int32 I = 0; I < Summary.ImportCount; ++I)
+			if (Summary.ExportCount > 0)
 			{
-				FObjectImport& ObjectImport = ObjectImports[BaseIndex + I];
-				DeserializeName(*Ar, ObjectImport.ClassPackage);
-				DeserializeName(*Ar, ObjectImport.ClassName);
-				*Ar << ObjectImport.OuterIndex;
-				DeserializeName(*Ar, ObjectImport.ObjectName);
+				FScopeLock _(&GlobalsCriticalSection);
+				Package.ExportIndexOffset = ObjectExports.Num();
+				ObjectExports.AddUninitialized(Summary.ExportCount);
+			}
 
-				if (ObjectImport.OuterIndex.IsNull())
-				{
-					++NumPackages;
-				}
+		}, EParallelForFlags::Unbalanced);
+
+		ParallelFor(TotalPackageCount,[
+			&ParseCount,
+			&PackageAssetBuffers,
+			&PackageFileSummaries,
+			&Packages,
+			&NameMapBuilder,
+			&GlobalImports,
+			&GlobalImportsByFullName,
+			&GlobalExports,
+			&GlobalExportsByFullName,
+			&ObjectImports,
+			&ObjectExports,
+			&PreloadDependencies,
+			&ExportGraph,
+			&GlobalsCriticalSection](int32 Index)
+		{
+			TUniquePtr<uint8[]>& PackageBuffer = PackageAssetBuffers[Index];
+			const FPackageFileSummary& Summary = PackageFileSummaries[Index];
+			FPackage& Package = *Packages[Index];
+			TArrayView<const uint8> MemView(PackageBuffer.Get(), Package.UAssetSize);
+			FMemoryReaderView Ar(MemView);
+
+			UE_CLOG(++ParseCount % 1000 == 0, LogIoStore, Display, TEXT("Parsing %d/%d: '%s'"), ParseCount.Load(), Packages.Num(), *Package.FileName);
+
+			if (Summary.NameCount > 0)
+			{
+				Ar.Seek(Summary.NameOffset);
 
-				ImportNames.Emplace(ObjectImport.ObjectName.ToString());
-			}
+				Package.NameMap.Reserve(Summary.NameCount);
+				Package.NameIndices.Reserve(Summary.NameCount);
+				FNameEntrySerialized NameEntry(ENAME_LinkerConstructor);
 
-			UPackageImports += NumPackages;
+				for (int32 I = 0; I < Summary.NameCount; ++I)
+				{
+					Ar << NameEntry;
+					FName Name(NameEntry);
 
-			Package->ImportedFullNames.SetNum(Summary.ImportCount);
-			for (int32 I = 0; I < Summary.ImportCount; ++I)
-			{
-				FindImport(GlobalImports, GlobalImportsByFullName, Package->ImportedFullNames, ObjectImports.GetData() + BaseIndex, I);
+					int32 MappedNameIndex = 0;
+					{
+						FScopeLock _(&GlobalsCriticalSection);
+						NameMapBuilder.MarkNameAsReferenced(Name);
+						MappedNameIndex = NameMapBuilder.MapName(Name);
+					}
+					Package.NameMap.Emplace(Name.GetDisplayIndex());
+					Package.NameIndices.Add(MappedNameIndex);
+				}
 			}
-		}
 
-		Package->PreloadIndexOffset = PreloadDependencies.Num();
-		int32 PreloadDependenciesBaseIndex = -1;
-		if (Summary.PreloadDependencyCount > 0)
-		{
-			Ar->Seek(Summary.PreloadDependencyOffset);
-			PreloadDependenciesBaseIndex = PreloadDependencies.Num();
-			PreloadDependencies.AddUninitialized(Summary.PreloadDependencyCount);
-			for (int32 I = 0; I < Summary.PreloadDependencyCount; ++I)
+			if (Summary.ImportCount > 0)
 			{
-				FPackageIndex& Index = PreloadDependencies[PreloadDependenciesBaseIndex + I];
-				*Ar << Index;
-				if (Index.IsImport())
+				FNameProxyArchive ProxyAr(Ar, Package.NameMap);
+				ProxyAr.Seek(Summary.ImportOffset);
+
+				for (int32 I = 0; I < Summary.ImportCount; ++I)
 				{
-					++ImportPreloadCounts[Package->GlobalPackageId];
-					++ImportPreloadCount;
+					FObjectImport& ObjectImport = ObjectImports[Package.ImportIndexOffset + I];
+					ProxyAr << ObjectImport;
 				}
-				else
+
+				Package.ImportedFullNames.SetNum(Summary.ImportCount);
+				for (int32 I = 0; I < Summary.ImportCount; ++I)
 				{
-					++ExportPreloadCounts[Package->GlobalPackageId];
-					++ExportPreloadCount;
+					FScopeLock _(&GlobalsCriticalSection);
+					FindImport(GlobalImports, GlobalImportsByFullName, Package.ImportedFullNames, ObjectImports.GetData() + Package.ImportIndexOffset, I);
 				}
 			}
-		}
 
-		Package->ExportIndexOffset = ObjectExports.Num();
-		if (Summary.ExportCount > 0)
-		{
-			Ar->Seek(Summary.ExportOffset);
-
-			int32 BaseIndex = ObjectExports.Num();
-			ObjectExports.AddUninitialized(Summary.ExportCount);
-			for (int32 I = 0; I < Summary.ExportCount; ++I)
+			if (Summary.PreloadDependencyCount > 0)
 			{
-				FObjectExport& ObjectExport = ObjectExports[BaseIndex + I];
-				*Ar << ObjectExport.ClassIndex;
-				*Ar << ObjectExport.SuperIndex;
-				*Ar << ObjectExport.TemplateIndex;
-				*Ar << ObjectExport.OuterIndex;
-				DeserializeName(*Ar, ObjectExport.ObjectName);
-				uint32 ObjectFlags;
-				*Ar << ObjectFlags;
-				ObjectExport.ObjectFlags = (EObjectFlags)ObjectFlags;
-				*Ar << ObjectExport.SerialSize;
-				*Ar << ObjectExport.SerialOffset;
-				*Ar << ObjectExport.bForcedExport;
-				*Ar << ObjectExport.bNotForClient;
-				*Ar << ObjectExport.bNotForServer;
-				*Ar << ObjectExport.PackageGuid;
-				*Ar << ObjectExport.PackageFlags;
-				*Ar << ObjectExport.bNotAlwaysLoadedForEditorGame;
-				*Ar << ObjectExport.bIsAsset;
-				*Ar << ObjectExport.FirstExportDependency;
-				*Ar << ObjectExport.SerializationBeforeSerializationDependencies;
-				*Ar << ObjectExport.CreateBeforeSerializationDependencies;
-				*Ar << ObjectExport.SerializationBeforeCreateDependencies;
-				*Ar << ObjectExport.CreateBeforeCreateDependencies;
+				Ar.Seek(Summary.PreloadDependencyOffset);
+				Ar.Serialize(PreloadDependencies.GetData() + Package.PreloadIndexOffset, Summary.PreloadDependencyCount * sizeof(FPackageIndex));
 			}
 
-			TempFullNames.Reset();
-			TempFullNames.SetNum(Summary.ExportCount, false);
-			for (int32 I = 0; I < Summary.ExportCount; ++I)
+			if (Summary.ExportCount > 0)
 			{
-				FindExport(GlobalExports, GlobalExportsByFullName, TempFullNames, ObjectExports.GetData() + BaseIndex, I, Package->Name);
-
-				FExportData& ExportData = GlobalExports[*GlobalExportsByFullName.Find(TempFullNames[I])];
-				Package->Exports.Add(ExportData.GlobalIndex);
-				ExportData.CreateNode = ExportGraph.AddNode(Package, { uint32(I), FExportBundleEntry::ExportCommandType_Create });
-				ExportData.SerializeNode = ExportGraph.AddNode(Package, { uint32(I), FExportBundleEntry::ExportCommandType_Serialize });
-				Package->CreateExportNodes.Add(ExportData.CreateNode);
-				Package->SerializeExportNodes.Add(ExportData.SerializeNode);
-				ExportGraph.AddInternalDependency(ExportData.CreateNode, ExportData.SerializeNode);
+				FNameProxyArchive ProxyAr(Ar, Package.NameMap);
+				ProxyAr.Seek(Summary.ExportOffset);
+
+				for (int32 I = 0; I < Summary.ExportCount; ++I)
+				{
+					FObjectExport& ObjectExport = ObjectExports[Package.ExportIndexOffset + I];
+					ProxyAr << ObjectExport;
+				}
+
+				TArray<FString> TmpFullNames;
+				TmpFullNames.SetNum(Summary.ExportCount, false);
+				for (int32 I = 0; I < Summary.ExportCount; ++I)
+				{
+					FScopeLock _(&GlobalsCriticalSection);
+
+					FindExport(GlobalExports, GlobalExportsByFullName, TmpFullNames, ObjectExports.GetData() + Package.ExportIndexOffset, I, Package.Name);
+
+					FExportData& ExportData = GlobalExports[*GlobalExportsByFullName.Find(TmpFullNames[I])];
+					Package.Exports.Add(ExportData.GlobalIndex);
+					ExportData.CreateNode = ExportGraph.AddNode(&Package, { uint32(I), FExportBundleEntry::ExportCommandType_Create });
+					ExportData.SerializeNode = ExportGraph.AddNode(&Package, { uint32(I), FExportBundleEntry::ExportCommandType_Serialize });
+					Package.CreateExportNodes.Add(ExportData.CreateNode);
+					Package.SerializeExportNodes.Add(ExportData.SerializeNode);
+
+					ExportGraph.AddInternalDependency(ExportData.CreateNode, ExportData.SerializeNode);
+				}
 			}
-		}
-	
-		Ar->Close();
+		}, EParallelForFlags::Unbalanced);
 	}
-
+	
 	int32 NumScriptImports = 0;
 	{
 		// Sort imports by script objects first
diff --git a/Engine/Source/Runtime/CoreUObject/Private/Serialization/AsyncLoading2.cpp b/Engine/Source/Runtime/CoreUObject/Private/Serialization/AsyncLoading2.cpp
--- a/Engine/Source/Runtime/CoreUObject/Private/Serialization/AsyncLoading2.cpp
+++ b/Engine/Source/Runtime/CoreUObject/Private/Serialization/AsyncLoading2.cpp
@@ -135,25 +135,53 @@ public:
 	}
 };
 
-struct FExportMapEntry
+FArchive& operator<<(FArchive& Ar, FExportBundleEntry& ExportBundleEntry)
 {
-	uint64 SerialSize;
-	int32 ObjectName[2];
-	FPackageIndex OuterIndex;
-	FPackageIndex ClassIndex;
-	FPackageIndex SuperIndex;
-	FPackageIndex TemplateIndex;
-	int32 GlobalImportIndex;
-	EObjectFlags ObjectFlags;
-	EExportFilterFlags FilterFlags;
-	uint8 Pad[7];
-};
+	Ar << ExportBundleEntry.LocalExportIndex;
+	Ar << ExportBundleEntry.CommandType;
+
+	return Ar;
+}
 
-struct FExportBundleHeader
+FArchive& operator<<(FArchive& Ar, FExportBundleHeader& ExportBundleHeader)
 {
-	uint32 FirstEntryIndex;
-	uint32 EntryCount;
-};
+	Ar << ExportBundleHeader.FirstEntryIndex;
+	Ar << ExportBundleHeader.EntryCount;
+
+	return Ar;
+}
+
+FArchive& operator<<(FArchive& Ar, FExportMapEntry& ExportMapEntry)
+{
+	Ar << ExportMapEntry.SerialSize;
+	Ar << ExportMapEntry.ObjectName[0] << ExportMapEntry.ObjectName[1];
+	Ar << ExportMapEntry.OuterIndex;
+	Ar << ExportMapEntry.ClassIndex;
+	Ar << ExportMapEntry.SuperIndex;
+	Ar << ExportMapEntry.TemplateIndex;
+	Ar << ExportMapEntry.GlobalImportIndex;
+
+	uint32 ObjectFlags = uint32(ExportMapEntry.ObjectFlags);
+	Ar << ObjectFlags;
+	
+	if (Ar.IsLoading())
+	{
+		ExportMapEntry.ObjectFlags = EObjectFlags(ObjectFlags);
+	}
+
+	uint8 FilterFlags = uint8(ExportMapEntry.FilterFlags);
+	Ar << FilterFlags;
+
+	if (Ar.IsLoading())
+	{
+		ExportMapEntry.FilterFlags = EExportFilterFlags(FilterFlags);
+	}
+
+	uint64 Pad = 0;
+	Ar.Serialize(&Pad, 7);
+
+	return Ar;
+}
 
 struct FExportObject
 {
diff --git a/Engine/Source/Runtime/CoreUObject/Public/Serialization/AsyncLoading2.h b/Engine/Source/Runtime/CoreUObject/Public/Serialization/AsyncLoading2.h
--- a/Engine/Source/Runtime/CoreUObject/Public/Serialization/AsyncLoading2.h
+++ b/Engine/Source/Runtime/CoreUObject/Public/Serialization/AsyncLoading2.h
@@ -7,7 +7,9 @@
 #pragma once
 
 #include "CoreMinimal.h"
+#include "UObject/ObjectResource.h"
 
+class FArchive;
 class IAsyncPackageLoader;
 class FIoDispatcher;
 class IEDLBootNotificationManager;
@@ -65,6 +67,8 @@ struct FExportBundleEntry
 	};
 	uint32 LocalExportIndex;
 	uint32 CommandType;
+
+	COREUOBJECT_API friend FArchive& operator<<(FArchive& Ar, FExportBundleEntry& ExportBundleEntry);
 };
 
 /**
@@ -76,6 +80,36 @@ struct FExportBundleMetaEntry
 	uint32 PayloadSize = ~uint32(0);
 };
 
+/**
+ * Export bundle header
+ */
+struct FExportBundleHeader
+{
+	uint32 FirstEntryIndex;
+	uint32 EntryCount;
+
+	COREUOBJECT_API friend FArchive& operator<<(FArchive& Ar, FExportBundleHeader& ExportBundleHeader);
+};
+
+/**
+ * Export map entry.
+ */
+struct FExportMapEntry
+{
+	uint64 SerialSize;
+	int32 ObjectName[2];
+	FPackageIndex OuterIndex;
+	FPackageIndex ClassIndex;
+	FPackageIndex SuperIndex;
+	FPackageIndex TemplateIndex;
+	int32 GlobalImportIndex;
+	EObjectFlags ObjectFlags;
+	EExportFilterFlags FilterFlags;
+	uint8 Pad[7];
+
+	COREUOBJECT_API friend FArchive& operator<<(FArchive& Ar, FExportMapEntry& ExportMapEntry);
+};
+
 /**
  * Creates a new instance of the AsyncPackageLoader #2.
  *
