IoStore: Linux compile fixes.

diff --git a/Engine/Source/Developer/IoStoreUtilities/Private/IoStoreUtilities.cpp b/Engine/Source/Developer/IoStoreUtilities/Private/IoStoreUtilities.cpp
--- a/Engine/Source/Developer/IoStoreUtilities/Private/IoStoreUtilities.cpp
+++ b/Engine/Source/Developer/IoStoreUtilities/Private/IoStoreUtilities.cpp
@@ -21,6 +21,7 @@
 #include "Serialization/JsonSerializer.h"
 #include "Serialization/BufferWriter.h"
 #include "Serialization/LargeMemoryWriter.h"
+#include "Serialization/MemoryReader.h"
 #include "Serialization/AsyncLoading2.h"
 #include "UObject/NameBatchSerialization.h"
 #include "UObject/PackageFileSummary.h"
@@ -57,7 +58,7 @@ struct FContainerSourceSpec
 struct FCookedFileStatData
 {
 	static constexpr int32 NumPackageExtensions = 2;
-	static constexpr TCHAR* Extensions[] = { TEXT("umap"), TEXT("uasset"), TEXT("ubulk"), TEXT("uptnl") };
+	static constexpr TCHAR const* Extensions[] = { TEXT("umap"), TEXT("uasset"), TEXT("ubulk"), TEXT("uptnl") };
 	enum FFileExt { UMap, UAsset, UBulk, UPtnl };
 	enum FFileType { PackageHeader, BulkData };
 
@@ -1221,10 +1222,10 @@ void FinalizePackageHeader(
 
 	// Temporary archive for serializing export bundle data
 	FBufferWriter ExportBundlesArchive(nullptr, 0, EBufferWriterFlags::AllowResize | EBufferWriterFlags::TakeOwnership);
-	int32 ExportBundleEntryIndex = 0;
+	uint32 ExportBundleEntryIndex = 0;
 	for (FExportBundle& ExportBundle : Package->ExportBundles)
 	{
-		const int32 EntryCount = ExportBundle.Nodes.Num();
+		const uint32 EntryCount = ExportBundle.Nodes.Num();
 		FExportBundleHeader ExportBundleHeader { ExportBundleEntryIndex, EntryCount };
 		ExportBundlesArchive << ExportBundleHeader ;
 
@@ -1475,8 +1476,8 @@ int32 CreateTarget(
 	{
 		UE_LOG(LogIoStore, Display, TEXT("Parsing packages..."));
 
-		TAtomic<int32> ReadCount = 0;
-		TAtomic<int32> ParseCount = 0;
+		TAtomic<int32> ReadCount {0};
+		TAtomic<int32> ParseCount {0};
 		const int32 TotalPackageCount = Packages.Num();
 
 		TArray<FPackageFileSummary> PackageFileSummaries;
