IoStoreUtilities: Parallel serialize container files and cpu trace scopes

diff --git a/Engine/Source/Developer/IoStoreUtilities/Private/IoStoreUtilities.cpp b/Engine/Source/Developer/IoStoreUtilities/Private/IoStoreUtilities.cpp
--- a/Engine/Source/Developer/IoStoreUtilities/Private/IoStoreUtilities.cpp
+++ b/Engine/Source/Developer/IoStoreUtilities/Private/IoStoreUtilities.cpp
@@ -38,6 +38,9 @@ IMPLEMENT_MODULE(FDefaultModuleImpl, IoStoreUtilities);
 
 DEFINE_LOG_CATEGORY_STATIC(LogIoStore, Log, All);
 
+#define IOSTORE_CPU_SCOPE(NAME) TRACE_CPUPROFILER_EVENT_SCOPE(IoStore##NAME);
+#define IOSTORE_CPU_SCOPE_DATA(NAME, DATA) TRACE_CPUPROFILER_EVENT_SCOPE(NAME);
+
 #define OUTPUT_CHUNKID_DIRECTORY 0
 #define OUTPUT_NAMEMAP_CSV 0
 #define OUTPUT_IMPORTMAP_CSV 0
@@ -87,6 +90,14 @@ struct FContainerTargetSpec
 	TArray<FContainerTargetFile> TargetFiles;
 };
 
+struct FPackageAssetData
+{
+	FCriticalSection CriticalSection;
+	TArray<FObjectImport> ObjectImports;
+	TArray<FObjectExport> ObjectExports;
+	TArray<FPackageIndex> PreloadDependencies;
+};
+
 class FNameProxyArchive
 	: public FArchiveProxy
 {
@@ -120,8 +131,29 @@ private:
 class FNameMapBuilder
 {
 public:
+	void MarkNamesAsReferenced(const TArray<FName>& Names, TArray<FNameEntryId>& OutNameMap, TArray<int32>& OutNameIndices)
+	{
+		FScopeLock _(&CriticalSection);
+
+		for (const FName& Name : Names)
+		{
+			const FNameEntryId Id = Name.GetComparisonIndex();
+			int32& Index = NameIndices.FindOrAdd(Id);
+			if (Index == 0)
+			{
+				Index = NameIndices.Num();
+				NameMap.Add(Id);
+			}
+
+			OutNameMap.Emplace(Name.GetDisplayIndex());
+			OutNameIndices.Add(Index - 1);
+		}
+	}
+
 	void MarkNameAsReferenced(const FName& Name)
 	{
+		FScopeLock _(&CriticalSection);
+
 		const FNameEntryId Id = Name.GetComparisonIndex();
 		int32& Index = NameIndices.FindOrAdd(Id);
 		if (Index == 0)
@@ -153,6 +185,8 @@ public:
 
 	int32 MapName(const FName& Name) const
 	{
+		FScopeLock _(&CriticalSection);
+
 		const FNameEntryId Id = Name.GetComparisonIndex();
 		const int32* Index = NameIndices.Find(Id);
 		check(Index);
@@ -161,6 +195,8 @@ public:
 
 	void SerializeName(FArchive& A, const FName& N) const
 	{
+		FScopeLock _(&CriticalSection);
+
 		int32 NameIndex = MapName(N);
 		int32 NameNumber = N.GetNumber();
 		A << NameIndex << NameNumber;
@@ -202,6 +238,7 @@ public:
 #endif
 
 private:
+	mutable FCriticalSection CriticalSection;
 	TMap<FNameEntryId, int32> NameIndices;
 	TArray<FNameEntryId> NameMap;
 #if OUTPUT_NAMEMAP_CSV
@@ -821,6 +858,8 @@ static void AddPostLoadDependencies(
 
 static void BuildBundles(FExportGraph& ExportGraph, const TArray<FPackage*>& Packages)
 {
+	IOSTORE_CPU_SCOPE(BuildBundles)
+
 	TArray<FExportGraphNode*> ExportLoadOrder = ExportGraph.ComputeLoadOrder(Packages);
 	FPackage* LastPackage = nullptr;
 	uint32 BundleLoadOrder = 0;
@@ -886,6 +925,8 @@ static void BuildClustersRecursive(FPackage* Package, uint64 ClusterIndex)
 
 static void BuildClusters(const TArray<FPackage*>& Packages)
 {
+	IOSTORE_CPU_SCOPE(BuildClusters);
+
 	uint64 NextClusterIndex = 0;
 	for (FPackage* Package : Packages)
 	{
@@ -1016,7 +1057,27 @@ struct FExportData
 	FExportGraphNode* SerializeNode = nullptr;
 };
 
-static void FindImport(TArray<FImportData>& GlobalImports, TMap<FString, int32>& GlobalImportsByFullName, TArray<FString>& TempFullNames, FObjectImport* ImportMap, int32 LocalImportIndex)
+template <typename T>
+struct FGlobalObjects
+{
+	TArray<T> Objects;
+	TMap<FString, int32> ObjectsByFullName;
+};
+
+using FGlobalImports = FGlobalObjects<FImportData>;
+using FGlobalExports = FGlobalObjects<FExportData>;
+
+struct FGlobalPackageData
+{
+	FGlobalImports Imports;
+	FGlobalExports Exports;
+};
+
+static void FindImport(
+	FGlobalImports& GlobalImports,
+	TArray<FString>& TempFullNames,
+	FObjectImport* ImportMap,
+	const int32 LocalImportIndex)
 {
 	FObjectImport* Import = &ImportMap[LocalImportIndex];
 	FString& FullName = TempFullNames[LocalImportIndex];
@@ -1025,13 +1086,13 @@ static void FindImport(TArray<FImportData>& GlobalImports, TMap<FString, int32>&
 		if (Import->OuterIndex.IsNull())
 		{
 			Import->ObjectName.AppendString(FullName);
-			int32* FindGlobalImport = GlobalImportsByFullName.Find(FullName);
+			const int32* FindGlobalImport = GlobalImports.ObjectsByFullName.Find(FullName);
 			if (!FindGlobalImport)
 			{
 				// first time, assign global index for this root package
-				int32 GlobalImportIndex = GlobalImports.Num();
-				GlobalImportsByFullName.Add(FullName, GlobalImportIndex);
-				FImportData& GlobalImport = GlobalImports.AddDefaulted_GetRef();
+				const int32 GlobalImportIndex = GlobalImports.Objects.Num();
+				GlobalImports.ObjectsByFullName.Add(FullName, GlobalImportIndex);
+				FImportData& GlobalImport = GlobalImports.Objects.AddDefaulted_GetRef();
 				GlobalImport.GlobalIndex = GlobalImportIndex;
 				GlobalImport.OutermostIndex = GlobalImportIndex;
 				GlobalImport.OuterIndex = -1;
@@ -1043,30 +1104,29 @@ static void FindImport(TArray<FImportData>& GlobalImports, TMap<FString, int32>&
 			}
 			else
 			{
-				++GlobalImports[*FindGlobalImport].RefCount;
+				++GlobalImports.Objects[*FindGlobalImport].RefCount;
 			}
 		}
 		else
 		{
-			int32 LocalOuterIndex = Import->OuterIndex.ToImport();
-			FindImport(GlobalImports, GlobalImportsByFullName, TempFullNames, ImportMap, LocalOuterIndex);
+			const int32 LocalOuterIndex = Import->OuterIndex.ToImport();
+			FindImport(GlobalImports, TempFullNames, ImportMap, LocalOuterIndex);
 			FString& OuterName = TempFullNames[LocalOuterIndex];
 			ensure(OuterName.Len() > 0);
 
 			FullName.Append(OuterName);
 			FullName.AppendChar(TEXT('/'));
 			Import->ObjectName.AppendString(FullName);
-
-			int32* FindGlobalImport = GlobalImportsByFullName.Find(FullName);
+			const int32* FindGlobalImport = GlobalImports.ObjectsByFullName.Find(FullName);
 			if (!FindGlobalImport)
 			{
 				// first time, assign global index for this intermediate import
-				int32 GlobalImportIndex = GlobalImports.Num();
-				GlobalImportsByFullName.Add(FullName, GlobalImportIndex);
-				FImportData& GlobalImport = GlobalImports.AddDefaulted_GetRef();
-				int32* FindOuterGlobalImport = GlobalImportsByFullName.Find(OuterName);
+				const int32 GlobalImportIndex = GlobalImports.Objects.Num();
+				GlobalImports.ObjectsByFullName.Add(FullName, GlobalImportIndex);
+				FImportData& GlobalImport = GlobalImports.Objects.AddDefaulted_GetRef();
+				const int32* FindOuterGlobalImport = GlobalImports.ObjectsByFullName.Find(OuterName);
 				check(FindOuterGlobalImport);
-				const FImportData& OuterGlobalImport = GlobalImports[*FindOuterGlobalImport];
+				const FImportData& OuterGlobalImport = GlobalImports.Objects[*FindOuterGlobalImport];
 				GlobalImport.GlobalIndex = GlobalImportIndex;
 				GlobalImport.OutermostIndex = OuterGlobalImport.OutermostIndex;
 				GlobalImport.OuterIndex = OuterGlobalImport.GlobalIndex;
@@ -1077,13 +1137,18 @@ static void FindImport(TArray<FImportData>& GlobalImports, TMap<FString, int32>&
 			}
 			else
 			{
-				++GlobalImports[*FindGlobalImport].RefCount;
+				++GlobalImports.Objects[*FindGlobalImport].RefCount;
 			}
 		}
 	}
-};
+}
 
-static void FindExport(TArray<FExportData>& GlobalExports, TMap<FString, int32>& GlobalExportsByFullName, TArray<FString>& TempFullNames, const FObjectExport* ExportMap, int32 LocalExportIndex, const FName& PackageName)
+static void FindExport(
+	FGlobalExports& GlobalExports,
+	TArray<FString>& TempFullNames,
+	const FObjectExport* ExportMap,
+	const int32 LocalExportIndex,
+	const FName& PackageName)
 {
 	const FObjectExport* Export = ExportMap + LocalExportIndex;
 	FString& FullName = TempFullNames[LocalExportIndex];
@@ -1100,7 +1165,7 @@ static void FindExport(TArray<FExportData>& GlobalExports, TMap<FString, int32>&
 		{
 			check(Export->OuterIndex.IsExport());
 
-			FindExport(GlobalExports, GlobalExportsByFullName, TempFullNames, ExportMap, Export->OuterIndex.ToExport(), PackageName);
+			FindExport(GlobalExports, TempFullNames, ExportMap, Export->OuterIndex.ToExport(), PackageName);
 			FString& OuterName = TempFullNames[Export->OuterIndex.ToExport()];
 			check(OuterName.Len() > 0);
 
@@ -1108,17 +1173,17 @@ static void FindExport(TArray<FExportData>& GlobalExports, TMap<FString, int32>&
 			FullName.AppendChar(TEXT('/'));
 			Export->ObjectName.AppendString(FullName);
 		}
-		check(!GlobalExportsByFullName.Contains(FullName));
-		int32 GlobalExportIndex = GlobalExports.Num();
-		GlobalExportsByFullName.Add(FullName, GlobalExportIndex);
-		FExportData& ExportData = GlobalExports.AddDefaulted_GetRef();
+		check(!GlobalExports.ObjectsByFullName.Contains(FullName));
+		const int32 GlobalExportIndex = GlobalExports.Objects.Num();
+		GlobalExports.ObjectsByFullName.Add(FullName, GlobalExportIndex);
+		FExportData& ExportData = GlobalExports.Objects.AddDefaulted_GetRef();
 		ExportData.GlobalIndex = GlobalExportIndex;
 		ExportData.SourcePackageName = PackageName;
 		ExportData.ObjectName = Export->ObjectName;
 		ExportData.SourceIndex = LocalExportIndex;
 		ExportData.FullName = FullName;
 	}
-};
+}
 
 FPackage* FindOrAddPackage(const TCHAR* RelativeFileName, TArray<FPackage*>& Packages, TMap<FName, FPackage*>& PackageMap, const TMap<FString, uint64>& PackageOrderMap)
 {
@@ -1358,148 +1423,42 @@ void SerializePackageData(
 #endif
 }
 
-int32 CreateTarget(
-	const TCHAR* OutputDir,
-	const TCHAR* CookedDir,
-	const TArray<FContainerSourceSpec>& Containers,
-	const FCookedFileStatMap& CookedFileStatMap,
-	const TMap<FString, uint64> PackageOrderMap)
+static void ParsePackageAssets(
+	FNameMapBuilder& NameMapBuilder,
+	TArray<FPackage*>& Packages,
+	FPackageAssetData& PackageAssetData,
+	FGlobalPackageData& GlobalPackageData,
+	FExportGraph& ExportGraph)
 {
-	TGuardValue<int32> GuardAllowUnversionedContentInEditor(GAllowUnversionedContentInEditor, 1);
+	IOSTORE_CPU_SCOPE(ParsePackageAssets);
 
-	FPackageStoreBulkDataManifest BulkDataManifest(FString(CookedDir) / FApp::GetProjectName());
-	const bool bWithBulkDataManifest = BulkDataManifest.Load();
-	if (bWithBulkDataManifest)
-	{
-		UE_LOG(LogIoStore, Display, TEXT("Loaded Bulk Data manifest '%s'"), *BulkDataManifest.GetFilename());
-	}
-
-#if OUTPUT_CHUNKID_DIRECTORY
-	ChunkIdCsv.CreateOutputFile(CookedDir);
-#endif
+	TAtomic<int32> ReadCount {0};
+	TAtomic<int32> ParseCount {0};
+	const int32 TotalPackageCount = Packages.Num();
 
-	FNameMapBuilder NameMapBuilder;
-
-	FCriticalSection GlobalsCriticalSection;
-	TArray<FObjectImport> ObjectImports;
-	TArray<FObjectExport> ObjectExports;
-	TArray<FImportData> GlobalImports;
-	TArray<FExportData> GlobalExports;
-	TMap<FString, int32> GlobalImportsByFullName;
-	TMap<FString, int32> GlobalExportsByFullName;
-	TArray<FString> TempFullNames;
-	TArray<FPackageIndex> PreloadDependencies;
-	FExportGraph ExportGraph;
-	TArray<FExportBundleMetaEntry> ExportBundleMetaEntries;
+	TArray<FPackageFileSummary> PackageFileSummaries;
+	TArray<TUniquePtr<uint8[]>> PackageAssetBuffers;
 
-	TArray<FPackage*> Packages;
-	TMap<FName, FPackage*> PackageMap;
+	PackageFileSummaries.SetNum(TotalPackageCount);
+	PackageAssetBuffers.SetNum(TotalPackageCount);
 
-	TArray<FIoStoreWriter*> IoStoreWriters;
-	FIoStoreWriter* GlobalIoStoreWriter;
-	FIoStoreEnvironment GlobalIoStoreEnv;
+	UE_LOG(LogIoStore, Display, TEXT("Reading package assets..."));
 	{
-		FString GlobalOutputPath = FString(OutputDir) / TEXT("global");
-		GlobalIoStoreEnv.InitializeFileEnvironment(*GlobalOutputPath);
-		GlobalIoStoreWriter = new FIoStoreWriter(GlobalIoStoreEnv);
-		IoStoreWriters.Add(GlobalIoStoreWriter);
-#if !SKIP_WRITE_CONTAINER
-		FIoStatus IoStatus = GlobalIoStoreWriter->Initialize();
-		check(IoStatus.IsOk());
-#if OUTPUT_CONTAINER_CSV
-		IoStatus = GlobalIoStoreWriter->EnableCsvOutput();
-		check(IoStatus.IsOk());
-#endif
-#endif
-	}
-
-	TArray<FContainerTargetSpec> ContainerTargets;
-
-	for (const FContainerSourceSpec& Container : Containers)
-	{
-		FContainerTargetSpec& ContainerTarget = ContainerTargets.AddDefaulted_GetRef();
-		if (Container.OutputPath.IsEmpty())
-		{
-			ContainerTarget.IoStoreWriter = GlobalIoStoreWriter;
-		}
-		else
-		{
-			FIoStoreEnvironment ContainerIoStoreEnv;
-			ContainerIoStoreEnv.InitializeFileEnvironment(Container.OutputPath);
-			ContainerTarget.IoStoreWriter = new FIoStoreWriter(ContainerIoStoreEnv);
-			IoStoreWriters.Add(ContainerTarget.IoStoreWriter);
-#if !SKIP_WRITE_CONTAINER
-			FIoStatus IoStatus = ContainerTarget.IoStoreWriter->Initialize();
-			check(IoStatus.IsOk());
-#if OUTPUT_CONTAINER_CSV
-			IoStatus = ContainerTarget.IoStoreWriter->EnableCsvOutput();
-			check(IoStatus.IsOk());
-#endif
-#endif
-		}
-
-		int32 CookedDirLen = FCString::Strlen(CookedDir) + 1;
-		for (const FContainerSourceFile& SourceFile : Container.SourceFiles)
-		{
-			const FCookedFileStatData* CookedFileStatData = CookedFileStatMap.Find(SourceFile.NormalizedPath);
-			if (!CookedFileStatData)
-			{
-				UE_LOG(LogIoStore, Warning, TEXT("File not found: '%s'"), *SourceFile.NormalizedPath);
-			}
-
-			FString RelativeFileName = TEXT("../../../");
-			RelativeFileName.AppendChars(*SourceFile.NormalizedPath + CookedDirLen, SourceFile.NormalizedPath.Len() - CookedDirLen);
-			
-			FPackage* Package = FindOrAddPackage(*RelativeFileName, Packages, PackageMap, PackageOrderMap);
-			if (Package)
-			{
-				FContainerTargetFile& TargetFile = ContainerTarget.TargetFiles.AddDefaulted_GetRef();
-				TargetFile.Size = uint64(CookedFileStatData->FileSize);
-				TargetFile.NormalizedSourcePath = SourceFile.NormalizedPath;
-				TargetFile.TargetPath = MoveTemp(RelativeFileName);
-				TargetFile.Package = Package;
-
-				if (CookedFileStatData->FileType == FCookedFileStatData::PackageHeader)
-				{
-					TargetFile.bIsBulkData = false;
-					Package->FileName = SourceFile.NormalizedPath;
-				}
-				else
-				{
-					TargetFile.bIsBulkData = true;
-					TargetFile.bIsOptionalBulkData = CookedFileStatData->FileExt == FCookedFileStatData::UPtnl;
-				}
-			}
-		}
-	}
-
-	{
-		UE_LOG(LogIoStore, Display, TEXT("Parsing packages..."));
-
-		TAtomic<int32> ReadCount {0};
-		TAtomic<int32> ParseCount {0};
-		const int32 TotalPackageCount = Packages.Num();
-
-		TArray<FPackageFileSummary> PackageFileSummaries;
-		TArray<TUniquePtr<uint8[]>> PackageAssetBuffers;
-
-		PackageFileSummaries.SetNum(TotalPackageCount);
-		PackageAssetBuffers.SetNum(TotalPackageCount);
+		IOSTORE_CPU_SCOPE(ReadAssets);
 
 		ParallelFor(TotalPackageCount,[
 			&ReadCount,
 			&PackageAssetBuffers,
 			&PackageFileSummaries,
 			&Packages,
-			&ObjectImports,
-			&ObjectExports,
-			&PreloadDependencies,
-			&GlobalsCriticalSection](int32 Index)
+			&PackageAssetData](int32 Index)
 		{
 			TUniquePtr<uint8[]>& PackageBuffer = PackageAssetBuffers[Index];
 			FPackageFileSummary& Summary = PackageFileSummaries[Index];
 			FPackage& Package = *Packages[Index];
 
+			IOSTORE_CPU_SCOPE_DATA(ReadPackage, TCHAR_TO_ANSI(*Package.FileName));
+
 			UE_CLOG(++ReadCount % 1000 == 0, LogIoStore, Display, TEXT("Reading %d/%d: '%s'"), ReadCount.Load(), Packages.Num(), *Package.FileName);
 
 			{
@@ -1525,42 +1484,41 @@ int32 CreateTarget(
 
 			if (Summary.ImportCount > 0)
 			{
-				FScopeLock _(&GlobalsCriticalSection);
-				Package.ImportIndexOffset = ObjectImports.Num();
-				ObjectImports.AddUninitialized(Summary.ImportCount);
+				FScopeLock _(&PackageAssetData.CriticalSection);
+				Package.ImportIndexOffset = PackageAssetData.ObjectImports.Num();
+				PackageAssetData.ObjectImports.AddUninitialized(Summary.ImportCount);
 			}
 			
 			if (Summary.PreloadDependencyCount > 0)
 			{
-				FScopeLock _(&GlobalsCriticalSection);
-				Package.PreloadIndexOffset = PreloadDependencies.Num();
-				PreloadDependencies.AddUninitialized(Summary.PreloadDependencyCount);
+				FScopeLock _(&PackageAssetData.CriticalSection);
+				Package.PreloadIndexOffset = PackageAssetData.PreloadDependencies.Num();
+				PackageAssetData.PreloadDependencies.AddUninitialized(Summary.PreloadDependencyCount);
 			}
 
 			if (Summary.ExportCount > 0)
 			{
-				FScopeLock _(&GlobalsCriticalSection);
-				Package.ExportIndexOffset = ObjectExports.Num();
-				ObjectExports.AddUninitialized(Summary.ExportCount);
+				FScopeLock _(&PackageAssetData.CriticalSection);
+				Package.ExportIndexOffset = PackageAssetData.ObjectExports.Num();
+				PackageAssetData.ObjectExports.AddUninitialized(Summary.ExportCount);
 			}
 
 		}, EParallelForFlags::Unbalanced);
+	}
+
+	UE_LOG(LogIoStore, Display, TEXT("Serializing package assets..."));
+	{
+		IOSTORE_CPU_SCOPE(SerializeAssets);
 
 		ParallelFor(TotalPackageCount,[
-			&ParseCount,
-			&PackageAssetBuffers,
-			&PackageFileSummaries,
-			&Packages,
-			&NameMapBuilder,
-			&GlobalImports,
-			&GlobalImportsByFullName,
-			&GlobalExports,
-			&GlobalExportsByFullName,
-			&ObjectImports,
-			&ObjectExports,
-			&PreloadDependencies,
-			&ExportGraph,
-			&GlobalsCriticalSection](int32 Index)
+				&ParseCount,
+				&PackageAssetBuffers,
+				&PackageFileSummaries,
+				&Packages,
+				&NameMapBuilder,
+				&PackageAssetData,
+				&GlobalPackageData,
+				&ExportGraph](int32 Index)
 		{
 			TUniquePtr<uint8[]>& PackageBuffer = PackageAssetBuffers[Index];
 			const FPackageFileSummary& Summary = PackageFileSummaries[Index];
@@ -1568,12 +1526,16 @@ int32 CreateTarget(
 			TArrayView<const uint8> MemView(PackageBuffer.Get(), Package.UAssetSize);
 			FMemoryReaderView Ar(MemView);
 
+			IOSTORE_CPU_SCOPE_DATA(ParsePackage, TCHAR_TO_ANSI(*Package.FileName));
+
 			UE_CLOG(++ParseCount % 1000 == 0, LogIoStore, Display, TEXT("Parsing %d/%d: '%s'"), ParseCount.Load(), Packages.Num(), *Package.FileName);
 
 			if (Summary.NameCount > 0)
 			{
 				Ar.Seek(Summary.NameOffset);
 
+				TArray<FName> Names;
+				Names.Reserve(Summary.NameCount);
 				Package.NameMap.Reserve(Summary.NameCount);
 				Package.NameIndices.Reserve(Summary.NameCount);
 				FNameEntrySerialized NameEntry(ENAME_LinkerConstructor);
@@ -1581,17 +1543,10 @@ int32 CreateTarget(
 				for (int32 I = 0; I < Summary.NameCount; ++I)
 				{
 					Ar << NameEntry;
-					FName Name(NameEntry);
-
-					int32 MappedNameIndex = 0;
-					{
-						FScopeLock _(&GlobalsCriticalSection);
-						NameMapBuilder.MarkNameAsReferenced(Name);
-						MappedNameIndex = NameMapBuilder.MapName(Name);
-					}
-					Package.NameMap.Emplace(Name.GetDisplayIndex());
-					Package.NameIndices.Add(MappedNameIndex);
+					Names.Emplace(NameEntry);
 				}
+
+				NameMapBuilder.MarkNamesAsReferenced(Names, Package.NameMap, Package.NameIndices); 
 			}
 
 			if (Summary.ImportCount > 0)
@@ -1601,22 +1556,15 @@ int32 CreateTarget(
 
 				for (int32 I = 0; I < Summary.ImportCount; ++I)
 				{
-					FObjectImport& ObjectImport = ObjectImports[Package.ImportIndexOffset + I];
+					FObjectImport& ObjectImport = PackageAssetData.ObjectImports[Package.ImportIndexOffset + I];
 					ProxyAr << ObjectImport;
 				}
-
-				Package.ImportedFullNames.SetNum(Summary.ImportCount);
-				for (int32 I = 0; I < Summary.ImportCount; ++I)
-				{
-					FScopeLock _(&GlobalsCriticalSection);
-					FindImport(GlobalImports, GlobalImportsByFullName, Package.ImportedFullNames, ObjectImports.GetData() + Package.ImportIndexOffset, I);
-				}
 			}
 
 			if (Summary.PreloadDependencyCount > 0)
 			{
 				Ar.Seek(Summary.PreloadDependencyOffset);
-				Ar.Serialize(PreloadDependencies.GetData() + Package.PreloadIndexOffset, Summary.PreloadDependencyCount * sizeof(FPackageIndex));
+				Ar.Serialize(PackageAssetData.PreloadDependencies.GetData() + Package.PreloadIndexOffset, Summary.PreloadDependencyCount * sizeof(FPackageIndex));
 			}
 
 			if (Summary.ExportCount > 0)
@@ -1626,33 +1574,185 @@ int32 CreateTarget(
 
 				for (int32 I = 0; I < Summary.ExportCount; ++I)
 				{
-					FObjectExport& ObjectExport = ObjectExports[Package.ExportIndexOffset + I];
+					FObjectExport& ObjectExport = PackageAssetData.ObjectExports[Package.ExportIndexOffset + I];
 					ProxyAr << ObjectExport;
 				}
+			}
+		}, EParallelForFlags::Unbalanced);
+	}
 
-				TArray<FString> TmpFullNames;
-				TmpFullNames.SetNum(Summary.ExportCount, false);
-				for (int32 I = 0; I < Summary.ExportCount; ++I)
+	UE_LOG(LogIoStore, Display, TEXT("Creating global imports and exports..."));
+	{
+		IOSTORE_CPU_SCOPE(CreateGlobalImportsAndExports);
+
+		TArray<FString> TmpFullNames;
+		for (FPackage* Package : Packages)
+		{
+			if (Package->ImportCount > 0)
+			{
+				Package->ImportedFullNames.SetNum(Package->ImportCount);
+				for (int32 ImportIndex = 0; ImportIndex < Package->ImportCount; ++ImportIndex)
 				{
-					FScopeLock _(&GlobalsCriticalSection);
+					FindImport(
+						GlobalPackageData.Imports,
+						Package->ImportedFullNames,
+						PackageAssetData.ObjectImports.GetData() + Package->ImportIndexOffset,
+						ImportIndex);
+				}
+			}
 
-					FindExport(GlobalExports, GlobalExportsByFullName, TmpFullNames, ObjectExports.GetData() + Package.ExportIndexOffset, I, Package.Name);
+			if (Package->ExportCount > 0)
+			{
+				TmpFullNames.Reset();
+				TmpFullNames.SetNum(Package->ExportCount);
 
-					FExportData& ExportData = GlobalExports[*GlobalExportsByFullName.Find(TmpFullNames[I])];
-					Package.Exports.Add(ExportData.GlobalIndex);
-					ExportData.CreateNode = ExportGraph.AddNode(&Package, { uint32(I), FExportBundleEntry::ExportCommandType_Create });
-					ExportData.SerializeNode = ExportGraph.AddNode(&Package, { uint32(I), FExportBundleEntry::ExportCommandType_Serialize });
-					Package.CreateExportNodes.Add(ExportData.CreateNode);
-					Package.SerializeExportNodes.Add(ExportData.SerializeNode);
+				for (int32 ExportIndex = 0; ExportIndex < Package->ExportCount; ExportIndex++)
+				{
+					FindExport(
+						GlobalPackageData.Exports,
+						TmpFullNames,
+						PackageAssetData.ObjectExports.GetData() + Package->ExportIndexOffset,
+						ExportIndex,
+						Package->Name);
+
+					const int32* GlobalIndex = GlobalPackageData.Exports.ObjectsByFullName.Find(TmpFullNames[ExportIndex]);
+					FExportData& ExportData = GlobalPackageData.Exports.Objects[*GlobalIndex];
+					Package->Exports.Add(ExportData.GlobalIndex);
+					ExportData.CreateNode = ExportGraph.AddNode(Package, { uint32(ExportIndex), FExportBundleEntry::ExportCommandType_Create });
+					ExportData.SerializeNode = ExportGraph.AddNode(Package, { uint32(ExportIndex), FExportBundleEntry::ExportCommandType_Serialize });
+					Package->CreateExportNodes.Add(ExportData.CreateNode);
+					Package->SerializeExportNodes.Add(ExportData.SerializeNode);
 
 					ExportGraph.AddInternalDependency(ExportData.CreateNode, ExportData.SerializeNode);
 				}
 			}
-		}, EParallelForFlags::Unbalanced);
+		}
 	}
-	
+}
+
+int32 CreateTarget(
+	const TCHAR* OutputDir,
+	const TCHAR* CookedDir,
+	const TArray<FContainerSourceSpec>& Containers,
+	const FCookedFileStatMap& CookedFileStatMap,
+	const TMap<FString, uint64> PackageOrderMap)
+{
+	TGuardValue<int32> GuardAllowUnversionedContentInEditor(GAllowUnversionedContentInEditor, 1);
+
+	FPackageStoreBulkDataManifest BulkDataManifest(FString(CookedDir) / FApp::GetProjectName());
+	const bool bWithBulkDataManifest = BulkDataManifest.Load();
+	if (bWithBulkDataManifest)
+	{
+		UE_LOG(LogIoStore, Display, TEXT("Loaded Bulk Data manifest '%s'"), *BulkDataManifest.GetFilename());
+	}
+
+#if OUTPUT_CHUNKID_DIRECTORY
+	ChunkIdCsv.CreateOutputFile(CookedDir);
+#endif
+
+	FNameMapBuilder NameMapBuilder;
+	FPackageAssetData PackageAssetData;
+	FGlobalPackageData GlobalPackageData;
+	FExportGraph ExportGraph;
+	TArray<FExportBundleMetaEntry> ExportBundleMetaEntries;
+
+	TArray<FPackage*> Packages;
+	TMap<FName, FPackage*> PackageMap;
+
+	TArray<FIoStoreWriter*> IoStoreWriters;
+	FIoStoreWriter* GlobalIoStoreWriter;
+	FIoStoreEnvironment GlobalIoStoreEnv;
+	{
+		FString GlobalOutputPath = FString(OutputDir) / TEXT("global");
+		GlobalIoStoreEnv.InitializeFileEnvironment(*GlobalOutputPath);
+		GlobalIoStoreWriter = new FIoStoreWriter(GlobalIoStoreEnv);
+		IoStoreWriters.Add(GlobalIoStoreWriter);
+#if !SKIP_WRITE_CONTAINER
+		FIoStatus IoStatus = GlobalIoStoreWriter->Initialize();
+		check(IoStatus.IsOk());
+#if OUTPUT_CONTAINER_CSV
+		IoStatus = GlobalIoStoreWriter->EnableCsvOutput();
+		check(IoStatus.IsOk());
+#endif
+#endif
+	}
+
+	TArray<FContainerTargetSpec> ContainerTargets;
+
+	UE_LOG(LogIoStore, Display, TEXT("Creating container targets..."));
+	{
+		IOSTORE_CPU_SCOPE(CreateContainerTargets);
+
+		for (const FContainerSourceSpec& Container : Containers)
+		{
+			FContainerTargetSpec& ContainerTarget = ContainerTargets.AddDefaulted_GetRef();
+			if (Container.OutputPath.IsEmpty())
+			{
+				ContainerTarget.IoStoreWriter = GlobalIoStoreWriter;
+			}
+			else
+			{
+				FIoStoreEnvironment ContainerIoStoreEnv;
+				ContainerIoStoreEnv.InitializeFileEnvironment(Container.OutputPath);
+				ContainerTarget.IoStoreWriter = new FIoStoreWriter(ContainerIoStoreEnv);
+				IoStoreWriters.Add(ContainerTarget.IoStoreWriter);
+#if !SKIP_WRITE_CONTAINER
+				FIoStatus IoStatus = ContainerTarget.IoStoreWriter->Initialize();
+				check(IoStatus.IsOk());
+#if OUTPUT_CONTAINER_CSV
+				IoStatus = ContainerTarget.IoStoreWriter->EnableCsvOutput();
+				check(IoStatus.IsOk());
+#endif
+#endif
+			}
+
+			int32 CookedDirLen = FCString::Strlen(CookedDir) + 1;
+			for (const FContainerSourceFile& SourceFile : Container.SourceFiles)
+			{
+				const FCookedFileStatData* CookedFileStatData = CookedFileStatMap.Find(SourceFile.NormalizedPath);
+				if (!CookedFileStatData)
+				{
+					UE_LOG(LogIoStore, Warning, TEXT("File not found: '%s'"), *SourceFile.NormalizedPath);
+				}
+
+				FString RelativeFileName = TEXT("../../../");
+				RelativeFileName.AppendChars(*SourceFile.NormalizedPath + CookedDirLen, SourceFile.NormalizedPath.Len() - CookedDirLen);
+				
+				FPackage* Package = FindOrAddPackage(*RelativeFileName, Packages, PackageMap, PackageOrderMap);
+				if (Package)
+				{
+					FContainerTargetFile& TargetFile = ContainerTarget.TargetFiles.AddDefaulted_GetRef();
+					TargetFile.Size = uint64(CookedFileStatData->FileSize);
+					TargetFile.NormalizedSourcePath = SourceFile.NormalizedPath;
+					TargetFile.TargetPath = MoveTemp(RelativeFileName);
+					TargetFile.Package = Package;
+
+					if (CookedFileStatData->FileType == FCookedFileStatData::PackageHeader)
+					{
+						TargetFile.bIsBulkData = false;
+						Package->FileName = SourceFile.NormalizedPath;
+					}
+					else
+					{
+						TargetFile.bIsBulkData = true;
+						TargetFile.bIsOptionalBulkData = CookedFileStatData->FileExt == FCookedFileStatData::UPtnl;
+					}
+				}
+			}
+		}
+	}
+
+	UE_LOG(LogIoStore, Display, TEXT("Parsing packages..."));
+	ParsePackageAssets(NameMapBuilder, Packages, PackageAssetData, GlobalPackageData, ExportGraph);
+
+	TArray<FImportData>& GlobalImports = GlobalPackageData.Imports.Objects;
+	TMap<FString, int32>& GlobalImportsByFullName = GlobalPackageData.Imports.ObjectsByFullName;
+	TMap<FString, int32>& GlobalExportsByFullName = GlobalPackageData.Exports.ObjectsByFullName;
+
 	int32 NumScriptImports = 0;
 	{
+		IOSTORE_CPU_SCOPE(GlobalImports);
+
 		// Sort imports by script objects first
 		GlobalImports.Sort();
 
@@ -1711,10 +1811,14 @@ int32 CreateTarget(
 		}
 	}
 
-	for (FExportData& GlobalExport : GlobalExports)
 	{
-		int32* FindGlobalImport = GlobalImportsByFullName.Find(GlobalExport.FullName);
-		GlobalExport.GlobalImportIndex = FindGlobalImport ? *FindGlobalImport : -1;
+		IOSTORE_CPU_SCOPE(GlobalExports);
+
+		for (FExportData& GlobalExport : GlobalPackageData.Exports.Objects)
+		{
+			int32* FindGlobalImport = GlobalPackageData.Imports.ObjectsByFullName.Find(GlobalExport.FullName);
+			GlobalExport.GlobalImportIndex = FindGlobalImport ? *FindGlobalImport : -1;
+		}
 	}
 
 #if OUTPUT_NAMEMAP_CSV
@@ -1747,126 +1851,142 @@ int32 CreateTarget(
 
 	// Lookup global indices and package pointers for all imports before adding preload and postload arcs
 	UE_LOG(LogIoStore, Display, TEXT("Looking up import packages..."));
-	for (FPackage* Package : Packages)
 	{
-		Package->ImportedPackages.Reserve(Package->ImportCount);
-		for (int32 I = 0; I < Package->ImportCount; ++I)
-		{
-			int32 GlobalImportIndex = *GlobalImportsByFullName.Find(Package->ImportedFullNames[I]);
+		IOSTORE_CPU_SCOPE(ImportedPackages);
 
-			FImportData& ImportData = GlobalImports[GlobalImportIndex];
-			Package->Imports.Add(ImportData.GlobalIndex);
-			if (ImportData.Package)
+		for (FPackage* Package : Packages)
+		{
+			Package->ImportedPackages.Reserve(Package->ImportCount);
+			for (int32 I = 0; I < Package->ImportCount; ++I)
 			{
-				Package->ImportedPackages.Add(ImportData.Package);
-				ImportData.Package->ImportedByPackages.Add(Package);
+				int32 GlobalImportIndex = *GlobalImportsByFullName.Find(Package->ImportedFullNames[I]);
+
+				FImportData& ImportData = GlobalImports[GlobalImportIndex];
+				Package->Imports.Add(ImportData.GlobalIndex);
+				if (ImportData.Package)
+				{
+					Package->ImportedPackages.Add(ImportData.Package);
+					ImportData.Package->ImportedByPackages.Add(Package);
+				}
 			}
 		}
 	}
 
 	UE_LOG(LogIoStore, Display, TEXT("Converting export map import indices..."));
-	for (FPackage* Package : Packages)
 	{
-		for (int32 I = 0; I < Package->ExportCount; ++I)
-		{
-			FObjectExport& ObjectExport = ObjectExports[Package->ExportIndexOffset + I];
-			FExportData& ExportData = GlobalExports[Package->Exports[I]];
+		IOSTORE_CPU_SCOPE(ExportData);
 
-			check(!ObjectExport.OuterIndex.IsImport());
-			ExportData.OuterIndex = ObjectExport.OuterIndex;
-			ExportData.ClassIndex =
-				ObjectExport.ClassIndex.IsImport() ?
-				FPackageIndex::FromImport(Package->Imports[ObjectExport.ClassIndex.ToImport()]) :
-				ObjectExport.ClassIndex;
-			ExportData.SuperIndex =
-				ObjectExport.SuperIndex.IsImport() ?
-				FPackageIndex::FromImport(Package->Imports[ObjectExport.SuperIndex.ToImport()]) :
-				ObjectExport.SuperIndex;
-			ExportData.TemplateIndex =
-				ObjectExport.TemplateIndex.IsImport() ?
-				FPackageIndex::FromImport(Package->Imports[ObjectExport.TemplateIndex.ToImport()]) :
-				ObjectExport.TemplateIndex;
+		for (FPackage* Package : Packages)
+		{
+			for (int32 I = 0; I < Package->ExportCount; ++I)
+			{
+				FObjectExport& ObjectExport = PackageAssetData.ObjectExports[Package->ExportIndexOffset + I];
+				FExportData& ExportData = GlobalPackageData.Exports.Objects[Package->Exports[I]];
+
+				check(!ObjectExport.OuterIndex.IsImport());
+				ExportData.OuterIndex = ObjectExport.OuterIndex;
+				ExportData.ClassIndex =
+					ObjectExport.ClassIndex.IsImport() ?
+					FPackageIndex::FromImport(Package->Imports[ObjectExport.ClassIndex.ToImport()]) :
+					ObjectExport.ClassIndex;
+				ExportData.SuperIndex =
+					ObjectExport.SuperIndex.IsImport() ?
+					FPackageIndex::FromImport(Package->Imports[ObjectExport.SuperIndex.ToImport()]) :
+					ObjectExport.SuperIndex;
+				ExportData.TemplateIndex =
+					ObjectExport.TemplateIndex.IsImport() ?
+					FPackageIndex::FromImport(Package->Imports[ObjectExport.TemplateIndex.ToImport()]) :
+					ObjectExport.TemplateIndex;
+			}
 		}
 	}
 
 	UE_LOG(LogIoStore, Display, TEXT("Adding optimized postload dependencies..."));
-	for (FPackage* Package : Packages)
 	{
-		Visited.Reset();
-		Visited.Add(Package);
-		AddPostLoadDependencies(*Package, Visited, CircularChains);
+		IOSTORE_CPU_SCOPE(PostLoadDependencies);
+
+		for (FPackage* Package : Packages)
+		{
+			Visited.Reset();
+			Visited.Add(Package);
+			AddPostLoadDependencies(*Package, Visited, CircularChains);
+		}
 	}
 		
 	UE_LOG(LogIoStore, Display, TEXT("Adding preload dependencies..."));
-	for (FPackage* Package : Packages)
 	{
-		// Convert PreloadDependencies to arcs
-		for (int32 I = 0; I < Package->ExportCount; ++I)
+		IOSTORE_CPU_SCOPE(PreLoadDependencies);
+
+		for (FPackage* Package : Packages)
 		{
-			FObjectExport& ObjectExport = ObjectExports[Package->ExportIndexOffset + I];
-			FExportData& ExportData = GlobalExports[Package->Exports[I]];
-			int32 PreloadDependenciesBaseIndex = Package->PreloadIndexOffset;
+			// Convert PreloadDependencies to arcs
+			for (int32 I = 0; I < Package->ExportCount; ++I)
+			{
+				FObjectExport& ObjectExport = PackageAssetData.ObjectExports[Package->ExportIndexOffset + I];
+				FExportData& ExportData = GlobalPackageData.Exports.Objects[Package->Exports[I]];
+				int32 PreloadDependenciesBaseIndex = Package->PreloadIndexOffset;
 
-			FPackageIndex ExportPackageIndex = FPackageIndex::FromExport(I);
+				FPackageIndex ExportPackageIndex = FPackageIndex::FromExport(I);
 
-			auto AddPreloadArc = [&](FPackageIndex Dep, EPreloadDependencyType PhaseFrom, EPreloadDependencyType PhaseTo)
-			{
-				if (Dep.IsExport())
-				{
-					AddInternalExportArc(ExportGraph, *Package, Dep.ToExport(), PhaseFrom, I, PhaseTo);
-				}
-				else
+				auto AddPreloadArc = [&](FPackageIndex Dep, EPreloadDependencyType PhaseFrom, EPreloadDependencyType PhaseTo)
 				{
-					FImportData& Import = GlobalImports[Package->Imports[Dep.ToImport()]];
-					check(!Import.bIsPackage);
-					if (Import.bIsScript)
+					if (Dep.IsExport())
 					{
-						// Add script arc with null package and global import index as node index
-						AddScriptArc(*Package, Import.GlobalIndex, I, PhaseTo);
+						AddInternalExportArc(ExportGraph, *Package, Dep.ToExport(), PhaseFrom, I, PhaseTo);
 					}
 					else
 					{
-						check(Import.GlobalExportIndex != - 1);
-						FExportData& Export = GlobalExports[Import.GlobalExportIndex];
-						FPackage* SourcePackage = PackageMap.FindRef(Export.SourcePackageName);
-						check(SourcePackage);
-						FPackageIndex SourceDep = FPackageIndex::FromExport(Export.SourceIndex);
-						AddExternalExportArc(ExportGraph, *SourcePackage, Export.SourceIndex, PhaseFrom, *Package, I, PhaseTo);
-						Package->ImportedPreloadPackages.Add(SourcePackage);
+						FImportData& Import = GlobalImports[Package->Imports[Dep.ToImport()]];
+						check(!Import.bIsPackage);
+						if (Import.bIsScript)
+						{
+							// Add script arc with null package and global import index as node index
+							AddScriptArc(*Package, Import.GlobalIndex, I, PhaseTo);
+						}
+						else
+						{
+							check(Import.GlobalExportIndex != - 1);
+							FExportData& Export = GlobalPackageData.Exports.Objects[Import.GlobalExportIndex];
+							FPackage* SourcePackage = PackageMap.FindRef(Export.SourcePackageName);
+							check(SourcePackage);
+							FPackageIndex SourceDep = FPackageIndex::FromExport(Export.SourceIndex);
+							AddExternalExportArc(ExportGraph, *SourcePackage, Export.SourceIndex, PhaseFrom, *Package, I, PhaseTo);
+							Package->ImportedPreloadPackages.Add(SourcePackage);
+						}
 					}
-				}
-			};
+				};
 
-			if (PreloadDependenciesBaseIndex >= 0 && ObjectExport.FirstExportDependency >= 0)
-			{
-				int32 RunningIndex = PreloadDependenciesBaseIndex + ObjectExport.FirstExportDependency;
-				for (int32 Index = ObjectExport.SerializationBeforeSerializationDependencies; Index > 0; Index--)
+				if (PreloadDependenciesBaseIndex >= 0 && ObjectExport.FirstExportDependency >= 0)
 				{
-					FPackageIndex Dep = PreloadDependencies[RunningIndex++];
-					check(!Dep.IsNull());
-					AddPreloadArc(Dep, PreloadDependencyType_Serialize, PreloadDependencyType_Serialize);
-				}
+					int32 RunningIndex = PreloadDependenciesBaseIndex + ObjectExport.FirstExportDependency;
+					for (int32 Index = ObjectExport.SerializationBeforeSerializationDependencies; Index > 0; Index--)
+					{
+						FPackageIndex Dep = PackageAssetData.PreloadDependencies[RunningIndex++];
+						check(!Dep.IsNull());
+						AddPreloadArc(Dep, PreloadDependencyType_Serialize, PreloadDependencyType_Serialize);
+					}
 
-				for (int32 Index = ObjectExport.CreateBeforeSerializationDependencies; Index > 0; Index--)
-				{
-					FPackageIndex Dep = PreloadDependencies[RunningIndex++];
-					check(!Dep.IsNull());
-					AddPreloadArc(Dep, PreloadDependencyType_Create, PreloadDependencyType_Serialize);
-				}
+					for (int32 Index = ObjectExport.CreateBeforeSerializationDependencies; Index > 0; Index--)
+					{
+						FPackageIndex Dep = PackageAssetData.PreloadDependencies[RunningIndex++];
+						check(!Dep.IsNull());
+						AddPreloadArc(Dep, PreloadDependencyType_Create, PreloadDependencyType_Serialize);
+					}
 
-				for (int32 Index = ObjectExport.SerializationBeforeCreateDependencies; Index > 0; Index--)
-				{
-					FPackageIndex Dep = PreloadDependencies[RunningIndex++];
-					check(!Dep.IsNull());
-					AddPreloadArc(Dep, PreloadDependencyType_Serialize, PreloadDependencyType_Create);
-				}
+					for (int32 Index = ObjectExport.SerializationBeforeCreateDependencies; Index > 0; Index--)
+					{
+						FPackageIndex Dep = PackageAssetData.PreloadDependencies[RunningIndex++];
+						check(!Dep.IsNull());
+						AddPreloadArc(Dep, PreloadDependencyType_Serialize, PreloadDependencyType_Create);
+					}
 
-				for (int32 Index = ObjectExport.CreateBeforeCreateDependencies; Index > 0; Index--)
-				{
-					FPackageIndex Dep = PreloadDependencies[RunningIndex++];
-					check(!Dep.IsNull());
-					// can't create this export until these things are created
-					AddPreloadArc(Dep, PreloadDependencyType_Create, PreloadDependencyType_Create);
+					for (int32 Index = ObjectExport.CreateBeforeCreateDependencies; Index > 0; Index--)
+					{
+						FPackageIndex Dep = PackageAssetData.PreloadDependencies[RunningIndex++];
+						check(!Dep.IsNull());
+						// can't create this export until these things are created
+						AddPreloadArc(Dep, PreloadDependencyType_Create, PreloadDependencyType_Create);
+					}
 				}
 			}
 		}
@@ -1886,6 +2006,8 @@ int32 CreateTarget(
 	UE_LOG(LogIoStore, Display, TEXT("Serializing global import names..."));
 	if (GlobalImportNamesArchive)
 	{
+		IOSTORE_CPU_SCOPE(SerializeGlobalImportNames);
+
 		for (const FImportData& ImportData : GlobalImports)
 		{
 			NameMapBuilder.SerializeName(*GlobalImportNamesArchive, ImportData.ObjectName);
@@ -1896,6 +2018,8 @@ int32 CreateTarget(
 	// Separate file for script arcs that are only required during initial loading
 	if (InitialLoadArchive)
 	{
+		IOSTORE_CPU_SCOPE(SerializeInitialLoad);
+
 		int32 PackageCount = PackageMap.Num();
 		*InitialLoadArchive << PackageCount;
 		*InitialLoadArchive << NumScriptImports;
@@ -1933,96 +2057,122 @@ int32 CreateTarget(
 		InitialLoadArchive->Serialize(ScriptArcsArchive.GetWriterData(), ScriptArcsArchive.Tell());
 	}
 
-	for (FPackage* Package : Packages)
 	{
-		Package->FirstExportBundleMetaEntry = ExportBundleMetaEntries.Num();
-		ExportBundleMetaEntries.AddDefaulted(Package->ExportBundles.Num());
-
-		NameMapBuilder.MarkNameAsReferenced(Package->Name);
-		NameMapBuilder.SerializeName(*StoreTocArchive, Package->Name);
-		*StoreTocArchive << Package->ExportCount;
-		int32 ExportBundleCount = Package->ExportBundles.Num();
-		*StoreTocArchive << ExportBundleCount;
-		*StoreTocArchive << Package->FirstExportBundleMetaEntry;
-		*StoreTocArchive << Package->FirstGlobalImport;
-		*StoreTocArchive << Package->GlobalImportCount;
-		int32 ImportedPackagesCount = Package->ImportedPackages.Num();
-		*StoreTocArchive << ImportedPackagesCount;
-		int32 ImportedPackagesOffset = ImportedPackagesArchive->Tell();
-		*StoreTocArchive << ImportedPackagesOffset;
+		IOSTORE_CPU_SCOPE(SerializeStoreToc);
 
-		for (FPackage* ImportedPackage : Package->ImportedPackages)
+		for (FPackage* Package : Packages)
 		{
-			*ImportedPackagesArchive << ImportedPackage->GlobalPackageId;
+			Package->FirstExportBundleMetaEntry = ExportBundleMetaEntries.Num();
+			ExportBundleMetaEntries.AddDefaulted(Package->ExportBundles.Num());
+
+			NameMapBuilder.MarkNameAsReferenced(Package->Name);
+			NameMapBuilder.SerializeName(*StoreTocArchive, Package->Name);
+			*StoreTocArchive << Package->ExportCount;
+			int32 ExportBundleCount = Package->ExportBundles.Num();
+			*StoreTocArchive << ExportBundleCount;
+			*StoreTocArchive << Package->FirstExportBundleMetaEntry;
+			*StoreTocArchive << Package->FirstGlobalImport;
+			*StoreTocArchive << Package->GlobalImportCount;
+			int32 ImportedPackagesCount = Package->ImportedPackages.Num();
+			*StoreTocArchive << ImportedPackagesCount;
+			int32 ImportedPackagesOffset = ImportedPackagesArchive->Tell();
+			*StoreTocArchive << ImportedPackagesOffset;
+
+			for (FPackage* ImportedPackage : Package->ImportedPackages)
+			{
+				*ImportedPackagesArchive << ImportedPackage->GlobalPackageId;
+			}
 		}
 	}
 
 	UE_LOG(LogIoStore, Display, TEXT("Finalizing package headers..."));
-	for (FPackage* Package : Packages)
 	{
-		FinalizePackageHeader(
-			Package,
-			NameMapBuilder,
-			ObjectExports,
-			GlobalExports,
-			GlobalImportsByFullName,
-			ExportBundleMetaEntries.GetData());
-	}
-
-	UE_LOG(LogIoStore, Display, TEXT("Serializing..."));
+		IOSTORE_CPU_SCOPE(FinalizePackageHeaders);
 
-	int32 TotalFileCount = 0;
-	for (FContainerTargetSpec& ContainerTarget : ContainerTargets)
-	{
-		TotalFileCount += ContainerTarget.TargetFiles.Num();
+		for (FPackage* Package : Packages)
+		{
+			FinalizePackageHeader(
+				Package,
+				NameMapBuilder,
+				PackageAssetData.ObjectExports,
+				GlobalPackageData.Exports.Objects,
+				GlobalPackageData.Imports.ObjectsByFullName,
+				ExportBundleMetaEntries.GetData());
+		}
 	}
-	int32 CurrentFileIndex = 0;
-	for (FContainerTargetSpec& ContainerTarget : ContainerTargets)
+
+	UE_LOG(LogIoStore, Display, TEXT("Serializing container(s)..."));
 	{
-		Algo::Sort(ContainerTarget.TargetFiles, [](const FContainerTargetFile& A, const FContainerTargetFile& B)
+		IOSTORE_CPU_SCOPE(SerializeContainers);
+
+		int32 TotalFileCount = 0;
+		for (FContainerTargetSpec& ContainerTarget : ContainerTargets)
 		{
-			if (A.bIsBulkData != B.bIsBulkData)
-			{
-				return B.bIsBulkData;
-			}
+			TotalFileCount += ContainerTarget.TargetFiles.Num();
+		}
+		int32 CurrentFileIndex = 0;
 
-			if (A.Package->ClusterIndex != B.Package->ClusterIndex)
-			{
-				return A.Package->ClusterIndex < B.Package->ClusterIndex;
-			}
-			
-			uint64 ABundleOrder = A.Package->ExportBundles.Num() > 0 ? A.Package->ExportBundles[0].LoadOrder : MAX_uint64;
-			uint64 BBundleOrder = B.Package->ExportBundles.Num() > 0 ? B.Package->ExportBundles[0].LoadOrder : MAX_uint64;
-			return ABundleOrder < BBundleOrder;
-		});
-		for (FContainerTargetFile& TargetFile : ContainerTarget.TargetFiles)
+		ParallelFor(ContainerTargets.Num(),[
+			&CurrentFileIndex,
+			&TotalFileCount,
+			&ContainerTargets,
+			&PackageAssetData,
+			bWithBulkDataManifest,
+			&BulkDataManifest](int Index)
 		{
-			UE_CLOG(CurrentFileIndex % 1000 == 0, LogIoStore, Display, TEXT("Serializing %d/%d: '%s'"), CurrentFileIndex, TotalFileCount, *TargetFile.NormalizedSourcePath);
-			++CurrentFileIndex;
-			if (TargetFile.bIsBulkData)
+			FContainerTargetSpec& ContainerTarget = ContainerTargets[Index];
+
+			Algo::Sort(ContainerTarget.TargetFiles, [](const FContainerTargetFile& A, const FContainerTargetFile& B)
 			{
-				if (bWithBulkDataManifest)
+				if (A.bIsBulkData != B.bIsBulkData)
 				{
-					FPackageStoreBulkDataManifest::EBulkdataType BulkDataType =
-						TargetFile.bIsOptionalBulkData ?
-						FPackageStoreBulkDataManifest::EBulkdataType::Optional :
-						FPackageStoreBulkDataManifest::EBulkdataType::Normal;
+					return B.bIsBulkData;
+				}
 
-					WriteBulkData(TargetFile.NormalizedSourcePath, BulkDataType, *TargetFile.Package, BulkDataManifest, ContainerTarget.IoStoreWriter);
+				if (A.Package->ClusterIndex != B.Package->ClusterIndex)
+				{
+					return A.Package->ClusterIndex < B.Package->ClusterIndex;
 				}
-			}
-			else
+				
+				uint64 ABundleOrder = A.Package->ExportBundles.Num() > 0 ? A.Package->ExportBundles[0].LoadOrder : MAX_uint64;
+				uint64 BBundleOrder = B.Package->ExportBundles.Num() > 0 ? B.Package->ExportBundles[0].LoadOrder : MAX_uint64;
+				return ABundleOrder < BBundleOrder;
+			});
+			for (FContainerTargetFile& TargetFile : ContainerTarget.TargetFiles)
 			{
-				SerializePackageData(TargetFile.Package, ObjectExports, ContainerTarget.IoStoreWriter);
+				UE_CLOG(CurrentFileIndex % 1000 == 0, LogIoStore, Display, TEXT("Serializing %d/%d: '%s'"), CurrentFileIndex, TotalFileCount, *TargetFile.NormalizedSourcePath);
+				++CurrentFileIndex;
+				if (TargetFile.bIsBulkData)
+				{
+					if (bWithBulkDataManifest)
+					{
+						IOSTORE_CPU_SCOPE_DATA(SerializeFile, TCHAR_TO_ANSI(*TargetFile.NormalizedSourcePath));
+
+						FPackageStoreBulkDataManifest::EBulkdataType BulkDataType =
+							TargetFile.bIsOptionalBulkData ?
+							FPackageStoreBulkDataManifest::EBulkdataType::Optional :
+							FPackageStoreBulkDataManifest::EBulkdataType::Normal;
+
+						WriteBulkData(TargetFile.NormalizedSourcePath, BulkDataType, *TargetFile.Package, BulkDataManifest, ContainerTarget.IoStoreWriter);
+					}
+				}
+				else
+				{
+					IOSTORE_CPU_SCOPE_DATA(SerializeFile, TCHAR_TO_ANSI(*TargetFile.Package->FileName));
+
+					SerializePackageData(TargetFile.Package, PackageAssetData.ObjectExports, ContainerTarget.IoStoreWriter);
+				}
 			}
-		}
+		});
 	}
-	
+
 	int32 StoreTocByteCount = StoreTocArchive->TotalSize();
 	int32 ImportedPackagesByteCount = ImportedPackagesArchive->TotalSize();
 	int32 GlobalImportNamesByteCount = GlobalImportNamesArchive->TotalSize();
 	int32 ExportBundleMetaByteCount = ExportBundleMetaEntries.Num() * sizeof(FExportBundleMetaEntry);
 	{
+		IOSTORE_CPU_SCOPE(SerializeGlobalMetaData);
+
 		UE_LOG(LogIoStore, Display, TEXT("Saving global meta data to container file"));
 		FLargeMemoryWriter GlobalMetaArchive(0, true);
 
@@ -2061,6 +2211,8 @@ int32 CreateTarget(
 	uint64 GlobalNamesMB = 0;
 	uint64 GlobalNameHashesMB = 0;
 	{
+		IOSTORE_CPU_SCOPE(SerializeGlobalNameMap);
+
 		UE_LOG(LogIoStore, Display, TEXT("Saving global name map to container file"));
 
 		TArray<uint8> Names;
@@ -2093,6 +2245,8 @@ int32 CreateTarget(
 	}
 	IoStoreWriters.Empty();
 
+	IOSTORE_CPU_SCOPE(CalculateStats);
+
 	UE_LOG(LogIoStore, Display, TEXT("Calculating stats..."));
 	uint64 UExpSize = 0;
 	uint64 UAssetSize = 0;
@@ -2162,7 +2316,7 @@ int32 CreateTarget(
 	UE_LOG(LogIoStore, Display, TEXT("-------------------- IoStore Summary --------------------"));
 	UE_LOG(LogIoStore, Display, TEXT("Packages: %8d total, %d circular dependencies, %d no preload dependencies, %d no import dependencies"),
 		PackageMap.Num(), CircularPackagesCount, PackagesWithoutPreloadDependenciesCount, PackagesWithoutImportDependenciesCount);
-	UE_LOG(LogIoStore, Display, TEXT("Bundles:  %8d total, %d entries, %d export objects"), BundleCount, BundleEntryCount, GlobalExports.Num());
+	UE_LOG(LogIoStore, Display, TEXT("Bundles:  %8d total, %d entries, %d export objects"), BundleCount, BundleEntryCount, GlobalPackageData.Exports.Objects.Num());
 
 	UE_LOG(LogIoStore, Display, TEXT("IoStore: %8.2f MB GlobalNames, %d unique names"), (double)GlobalNamesMB, NameMapBuilder.GetNameMap().Num());
 	UE_LOG(LogIoStore, Display, TEXT("IoStore: %8.2f MB GlobalNameHashes"), (double)GlobalNameHashesMB);
@@ -2335,6 +2489,8 @@ public:
 
 int32 CreateIoStoreContainerFiles(const TCHAR* CmdLine)
 {
+	IOSTORE_CPU_SCOPE(CreateIoStoreContainerFiles);
+
 	UE_LOG(LogIoStore, Display, TEXT("==================== IoStore Utils ===================="));
 
 	FString CommandListFile;
