Use CookerOpenOrder when building IoDispatcher containers

diff --git a/Engine/Source/Developer/IoStoreUtilities/Private/IoStoreUtilities.cpp b/Engine/Source/Developer/IoStoreUtilities/Private/IoStoreUtilities.cpp
--- a/Engine/Source/Developer/IoStoreUtilities/Private/IoStoreUtilities.cpp
+++ b/Engine/Source/Developer/IoStoreUtilities/Private/IoStoreUtilities.cpp
@@ -443,6 +443,7 @@ struct FPackage
 {
 	FName Name;
 	FString FileName;
+	FString RelativeFileName;
 	int32 GlobalPackageId = 0;
 	uint32 PackageFlags = 0;
 	int32 NameCount = 0;
@@ -494,8 +495,7 @@ struct FPackage
 	TArray<uint8> HeaderData;
 	uint64 ExportsPayloadSize = 0;
 
-	uint64 ClusterIndex = MAX_uint64;
-	uint64 OpenOrder = MAX_uint64;
+	uint64 DiskLayoutOrder = MAX_uint64;
 };
 
 struct FCircularImportChain
@@ -904,37 +904,120 @@ static void BuildBundles(FExportGraph& ExportGraph, const TArray<FPackage*>& Pac
 	}
 }
 
-static void BuildClustersRecursive(FPackage* Package, uint64 ClusterIndex)
+static void CreateDiskLayout(const TArray<FPackage*>& Packages, const TMap<FString, uint64> PackageOrderMap, const TMap<FString, uint64>& CookerFileOpenOrderMap)
 {
-	Package->ClusterIndex = ClusterIndex;
-	for (FPackage* OtherPackage : Package->ImportedPackages)
+	IOSTORE_CPU_SCOPE(CreateDiskLayout);
+
+	struct FCluster
+	{
+		TArray<FPackage*> Packages;
+		uint64 ClusterSize = 0;
+	};
+
+	TArray<FCluster*> Clusters;
+	TSet<FPackage*> AssignedPackages;
+	TArray<FPackage*> ProcessStack;
+
+	struct FPackageAndOrder
 	{
-		if (OtherPackage->ClusterIndex == MAX_uint64)
+		FPackage* Package;
+		uint64 PackageLoadOrder;
+		uint64 CookerOpenOrder;
+
+		bool operator<(const FPackageAndOrder& Other) const
 		{
-			BuildClustersRecursive(OtherPackage, ClusterIndex);
+			if (PackageLoadOrder != Other.PackageLoadOrder)
+			{
+				return PackageLoadOrder < Other.PackageLoadOrder;
+			}
+			if (CookerOpenOrder != Other.CookerOpenOrder)
+			{
+				return CookerOpenOrder < Other.CookerOpenOrder;
+			}
+			// Fallback to reverse bundle order
+			return Package->ExportBundles[0].LoadOrder > Other.Package->ExportBundles[0].LoadOrder;
 		}
+	};
+
+	TArray<FPackageAndOrder> SortedPackages;
+	SortedPackages.Reserve(Packages.Num());
+	for (FPackage* Package : Packages)
+	{
+		FPackageAndOrder& Entry = SortedPackages.AddDefaulted_GetRef();
+		Entry.Package = Package;
+		const uint64* FindPackageLoadOrder = PackageOrderMap.Find(Package->Name.ToString());
+		Entry.PackageLoadOrder = FindPackageLoadOrder ? *FindPackageLoadOrder : MAX_uint64;
+		const uint64* FindCookerOpenOrder = CookerFileOpenOrderMap.Find(Package->RelativeFileName);
+		/*if (!FindCookerOpenOrder)
+		{
+			UE_LOG(LogIoStore, Warning, TEXT("Missing cooker order for package: %s"), *Package->RelativeFileName);
+		}*/
+		Entry.CookerOpenOrder = FindCookerOpenOrder ? *FindCookerOpenOrder : MAX_uint64;
 	}
-	for (FPackage* OtherPackage : Package->ImportedByPackages)
+	bool bHasPackageOrder = true;
+	bool bHasCookerOrder = true;
+	int32 LastAssignedCount = 0;
+	Algo::Sort(SortedPackages);
+	for (FPackageAndOrder& Entry : SortedPackages)
 	{
-		if (OtherPackage->ClusterIndex == MAX_uint64)
+		if (bHasPackageOrder && Entry.PackageLoadOrder == MAX_uint64)
+		{
+			UE_LOG(LogIoStore, Display, TEXT("Ordered %d/%d packages using package load order"), AssignedPackages.Num(), Packages.Num());
+			LastAssignedCount = AssignedPackages.Num();
+			bHasPackageOrder = false;
+		}
+		if (!bHasPackageOrder && bHasCookerOrder && Entry.CookerOpenOrder == MAX_uint64)
 		{
-			BuildClustersRecursive(OtherPackage, ClusterIndex);
+			UE_LOG(LogIoStore, Display, TEXT("Ordered %d/%d packages using cooker order"), AssignedPackages.Num() - LastAssignedCount, Packages.Num() - LastAssignedCount);
+			LastAssignedCount = AssignedPackages.Num();
+			bHasCookerOrder = false;
+		}
+		if (!AssignedPackages.Contains(Entry.Package))
+		{
+			FCluster* Cluster = new FCluster();
+			Clusters.Add(Cluster);
+			ProcessStack.Push(Entry.Package);
+
+			while (ProcessStack.Num())
+			{
+				FPackage* PackageToProcess = ProcessStack.Pop(false);
+				if (!AssignedPackages.Contains(PackageToProcess))
+				{
+					AssignedPackages.Add(PackageToProcess);
+					Cluster->Packages.Add(PackageToProcess);
+					for (FPackage* ImportedPackage : PackageToProcess->ImportedPackages)
+					{
+						ProcessStack.Push(ImportedPackage);
+					}
+				}
+			}
 		}
 	}
-}
+	UE_LOG(LogIoStore, Display, TEXT("Ordered %d packages using fallback bundle order"), AssignedPackages.Num() - LastAssignedCount);
 
-static void BuildClusters(const TArray<FPackage*>& Packages)
-{
-	IOSTORE_CPU_SCOPE(BuildClusters);
+	check(AssignedPackages.Num() == Packages.Num());
+	
+	for (FCluster* Cluster : Clusters)
+	{
+		Algo::Sort(Cluster->Packages, [](const FPackage* A, const FPackage* B)
+		{
+			return A->ExportBundles[0].LoadOrder < B->ExportBundles[0].LoadOrder;
+		});
 
-	uint64 NextClusterIndex = 0;
-	for (FPackage* Package : Packages)
+		for (FPackage* Package : Cluster->Packages)
+		{
+			Cluster->ClusterSize += Package->HeaderData.Num() + Package->ExportsPayloadSize;
+		}
+	}
+
+	uint64 LayoutIndex = 0;
+	for (FCluster* Cluster : Clusters)
 	{
-		if (Package->ClusterIndex == MAX_uint64)
+		for (FPackage* Package : Cluster->Packages)
 		{
-			BuildClustersRecursive(Package, NextClusterIndex);
-			++NextClusterIndex;
+			Package->DiskLayoutOrder = LayoutIndex++;
 		}
+		delete Cluster;
 	}
 }
 
@@ -1185,7 +1268,7 @@ static void FindExport(
 	}
 }
 
-FPackage* FindOrAddPackage(const TCHAR* RelativeFileName, TArray<FPackage*>& Packages, TMap<FName, FPackage*>& PackageMap, const TMap<FString, uint64>& PackageOrderMap)
+FPackage* FindOrAddPackage(const TCHAR* RelativeFileName, TArray<FPackage*>& Packages, TMap<FName, FPackage*>& PackageMap)
 {
 	FString PackageName;
 	FString ErrorMessage;
@@ -1205,13 +1288,6 @@ FPackage* FindOrAddPackage(const TCHAR* RelativeFileName, TArray<FPackage*>& Pac
 		Package->GlobalPackageId = Packages.Num();
 		Packages.Add(Package);
 		PackageMap.Add(PackageFName, Package);
-
-
-		const uint64* FindPackageOrder = PackageOrderMap.Find(PackageName);
-		if (FindPackageOrder)
-		{
-			Package->OpenOrder = *FindPackageOrder;
-		}
 	}
 
 	return Package;
@@ -1635,7 +1711,8 @@ int32 CreateTarget(
 	const TCHAR* CookedDir,
 	const TArray<FContainerSourceSpec>& Containers,
 	const FCookedFileStatMap& CookedFileStatMap,
-	const TMap<FString, uint64> PackageOrderMap)
+	const TMap<FString, uint64>& PackageOrderMap,
+	const TMap<FString, uint64>& CookerFileOpenOrderMap)
 {
 	TGuardValue<int32> GuardAllowUnversionedContentInEditor(GAllowUnversionedContentInEditor, 1);
 
@@ -1718,7 +1795,7 @@ int32 CreateTarget(
 				FString RelativeFileName = TEXT("../../../");
 				RelativeFileName.AppendChars(*SourceFile.NormalizedPath + CookedDirLen, SourceFile.NormalizedPath.Len() - CookedDirLen);
 				
-				FPackage* Package = FindOrAddPackage(*RelativeFileName, Packages, PackageMap, PackageOrderMap);
+				FPackage* Package = FindOrAddPackage(*RelativeFileName, Packages, PackageMap);
 				if (Package)
 				{
 					FContainerTargetFile& TargetFile = ContainerTarget.TargetFiles.AddDefaulted_GetRef();
@@ -1731,6 +1808,7 @@ int32 CreateTarget(
 					{
 						TargetFile.bIsBulkData = false;
 						Package->FileName = SourceFile.NormalizedPath;
+						Package->RelativeFileName = TargetFile.TargetPath;
 					}
 					else
 					{
@@ -1995,9 +2073,6 @@ int32 CreateTarget(
 	UE_LOG(LogIoStore, Display, TEXT("Building bundles..."));
 	BuildBundles(ExportGraph, Packages);
 
-	UE_LOG(LogIoStore, Display, TEXT("Building clusters..."));
-	BuildClusters(Packages);
-
 	TUniquePtr<FLargeMemoryWriter> StoreTocArchive = MakeUnique<FLargeMemoryWriter>(0, true);
 	TUniquePtr<FLargeMemoryWriter> ImportedPackagesArchive = MakeUnique<FLargeMemoryWriter>(0, true);
 	TUniquePtr<FLargeMemoryWriter> GlobalImportNamesArchive = MakeUnique<FLargeMemoryWriter>(0, true);
@@ -2101,6 +2176,9 @@ int32 CreateTarget(
 		}
 	}
 
+	UE_LOG(LogIoStore, Display, TEXT("Creating disk layout..."));
+	CreateDiskLayout(Packages, PackageOrderMap, CookerFileOpenOrderMap);
+
 	UE_LOG(LogIoStore, Display, TEXT("Serializing container(s)..."));
 	{
 		IOSTORE_CPU_SCOPE(SerializeContainers);
@@ -2129,14 +2207,7 @@ int32 CreateTarget(
 					return B.bIsBulkData;
 				}
 
-				if (A.Package->ClusterIndex != B.Package->ClusterIndex)
-				{
-					return A.Package->ClusterIndex < B.Package->ClusterIndex;
-				}
-				
-				uint64 ABundleOrder = A.Package->ExportBundles.Num() > 0 ? A.Package->ExportBundles[0].LoadOrder : MAX_uint64;
-				uint64 BBundleOrder = B.Package->ExportBundles.Num() > 0 ? B.Package->ExportBundles[0].LoadOrder : MAX_uint64;
-				return ABundleOrder < BBundleOrder;
+				return A.Package->DiskLayoutOrder < B.Package->DiskLayoutOrder;
 			});
 			for (FContainerTargetFile& TargetFile : ContainerTarget.TargetFiles)
 			{
@@ -2493,6 +2564,26 @@ int32 CreateIoStoreContainerFiles(const TCHAR* CmdLine)
 
 	UE_LOG(LogIoStore, Display, TEXT("==================== IoStore Utils ===================="));
 
+	TMap<FString, uint64> PackageOrderMap;
+	FString PackageOrderFilePath;
+	if (FParse::Value(FCommandLine::Get(), TEXT("PackageOrder="), PackageOrderFilePath))
+	{
+		if (!ParsePakOrderFile(*PackageOrderFilePath, PackageOrderMap))
+		{
+			return -1;
+		}
+	}
+
+	TMap<FString, uint64> CookerFileOpenOrderMap;
+	FString CookerFileOpenOrderFilePath;
+	if (FParse::Value(FCommandLine::Get(), TEXT("CookerOrder="), CookerFileOpenOrderFilePath))
+	{
+		if (!ParsePakOrderFile(*CookerFileOpenOrderFilePath, CookerFileOpenOrderMap))
+		{
+			return -1;
+		}
+	}
+
 	FString CommandListFile;
 	if (FParse::Value(FCommandLine::Get(), TEXT("Commands="), CommandListFile))
 	{
@@ -2543,23 +2634,13 @@ int32 CreateIoStoreContainerFiles(const TCHAR* CmdLine)
 			}
 		}
 
-		TMap<FString, uint64> PackageOrderMap;
-		FString PackageOrderFilePath;
-		if (FParse::Value(FCommandLine::Get(), TEXT("PackageOrder="), PackageOrderFilePath))
-		{
-			if (!ParsePakOrderFile(*PackageOrderFilePath, PackageOrderMap))
-			{
-				return -1;
-			}
-		}
-
 		UE_LOG(LogIoStore, Display, TEXT("Searching for cooked files..."));
 		FCookedFileStatMap CookedFileStatMap;
 		FCookedFileVisitor CookedFileVistor(CookedFileStatMap);
 		IFileManager::Get().IterateDirectoryStatRecursively(*CookedDirectory, CookedFileVistor);
 		UE_LOG(LogIoStore, Display, TEXT("Found '%d' files"), CookedFileStatMap.Num());
 
-		int32 ReturnValue = CreateTarget(*OutputDirectory, *CookedDirectory, Containers, CookedFileStatMap, PackageOrderMap);
+		int32 ReturnValue = CreateTarget(*OutputDirectory, *CookedDirectory, Containers, CookedFileStatMap, PackageOrderMap, CookerFileOpenOrderMap);
 		if (ReturnValue != 0)
 		{
 			return ReturnValue;
@@ -2586,7 +2667,7 @@ int32 CreateIoStoreContainerFiles(const TCHAR* CmdLine)
 
 			UE_LOG(LogIoStore, Display, TEXT("Creating target: '%s' using output path: '%s'"), *TargetPlatform->PlatformName(), *TargetCookedProjectDirectory);
 
-			int32 ReturnValue = CreateTarget(*TargetCookedProjectDirectory, *TargetCookedDirectory, Containers, CookedFileStatMap, TMap<FString, uint64>());
+			int32 ReturnValue = CreateTarget(*TargetCookedProjectDirectory, *TargetCookedDirectory, Containers, CookedFileStatMap, PackageOrderMap, CookerFileOpenOrderMap);
 			if (ReturnValue != 0)
 			{
 				return ReturnValue;
diff --git a/Engine/Source/Programs/AutomationTool/Scripts/CopyBuildToStagingDirectory.Automation.cs b/Engine/Source/Programs/AutomationTool/Scripts/CopyBuildToStagingDirectory.Automation.cs
--- a/Engine/Source/Programs/AutomationTool/Scripts/CopyBuildToStagingDirectory.Automation.cs
+++ b/Engine/Source/Programs/AutomationTool/Scripts/CopyBuildToStagingDirectory.Automation.cs
@@ -2236,7 +2236,27 @@ public partial class Project : CommandUtils
 					Writer.WriteLine(Command);
 				}
 			}
-			RunIoStore(Params, SC, IoStoreCommandsFileName);
+
+			FileReference PackageOpenOrderFileLocation = null;
+			FileReference CookerOpenOrderFileLocation = null;
+			// search CookPlaform (e.g. IOSClient and then regular platform (e.g. IOS).
+			string[] OrderLocations = new string[] { SC.FinalCookPlatform, SC.StageTargetPlatform.GetTargetPlatformDescriptor().Type.ToString() };
+
+			foreach (string OrderLocation in OrderLocations.Reverse())
+			{
+				DirectoryReference IoStoreOrderFileLocationBase = DirectoryReference.Combine(SC.ProjectRoot, "Build", OrderLocation, "FileOpenOrder");
+				FileReference FileLocation = FileReference.Combine(IoStoreOrderFileLocationBase, "PackageOpenOrder.log");
+				if (FileExists_NoExceptions(FileLocation.FullName))
+				{
+					PackageOpenOrderFileLocation = FileLocation;
+				}
+				FileLocation = FileReference.Combine(IoStoreOrderFileLocationBase, "CookerOpenOrder.log");
+				if (FileExists_NoExceptions(FileLocation.FullName))
+				{
+					CookerOpenOrderFileLocation = FileLocation;
+				}
+			}
+			RunIoStore(Params, SC, IoStoreCommandsFileName, PackageOpenOrderFileLocation, CookerOpenOrderFileLocation);
 		}
 
 		// Do any additional processing on the command output
@@ -2380,10 +2400,18 @@ public partial class Project : CommandUtils
 		}
 	}
 
-	private static void RunIoStore(ProjectParams Params, DeploymentContext SC, string CommandsFileName)
+	private static void RunIoStore(ProjectParams Params, DeploymentContext SC, string CommandsFileName, FileReference PackageOrderFileLocation, FileReference CookerOpenOrderFileLocation)
 	{
 		FileReference OutputLocation = FileReference.Combine(SC.RuntimeRootDir, SC.RelativeProjectRootForStage.Name);
 		string CommandletParams = String.Format("-OutputDirectory={0} -CookedDirectory={1} -Commands={2}", OutputLocation.FullName, SC.PlatformCookDir, CommandsFileName);
+		if (PackageOrderFileLocation != null)
+		{
+			CommandletParams += String.Format(" -PackageOrder={0}", PackageOrderFileLocation.FullName);
+		}
+		if (CookerOpenOrderFileLocation != null)
+		{
+			CommandletParams += String.Format(" -CookerOrder={0}", CookerOpenOrderFileLocation.FullName);
+		}
 		LogInformation("Running IoStore commandlet with arguments: {0}", CommandletParams);
 		RunCommandlet(SC.RawProjectPath, Params.UE4Exe, "IoStore", CommandletParams);
 	}
