IoDispatcher:
Gracefully handle missing packages when building containers

diff --git a/Engine/Source/Developer/IoStoreUtilities/Private/IoStoreUtilities.cpp b/Engine/Source/Developer/IoStoreUtilities/Private/IoStoreUtilities.cpp
--- a/Engine/Source/Developer/IoStoreUtilities/Private/IoStoreUtilities.cpp
+++ b/Engine/Source/Developer/IoStoreUtilities/Private/IoStoreUtilities.cpp
@@ -1149,6 +1149,7 @@ struct FImportData
 	bool bIsLocalized = false;
 	FString FullName;
 	FPackage* Package = nullptr;
+	bool bIsMissingImport = false;
 
 	bool operator<(const FImportData& Other) const
 	{
@@ -2289,16 +2290,31 @@ int32 CreateTarget(
 						LastPackage->GlobalImportCount = I - LastPackage->FirstGlobalImport;
 					}
 					FPackage* FindPackage = PackageMap.FindRef(GlobalImport.ObjectName);
-					check(FindPackage);
-					FindPackage->FirstGlobalImport = I;
-					GlobalImport.Package = FindPackage;
-					LastPackage = FindPackage;
+					if (!FindPackage)
+					{
+						UE_LOG(LogIoStore, Warning, TEXT("Missing import package '%s'"), *GlobalImport.ObjectName.ToString());
+						LastPackage = nullptr;
+						GlobalImport.bIsMissingImport = true;
+					}
+					else
+					{
+						FindPackage->FirstGlobalImport = I;
+						LastPackage = FindPackage;
+						GlobalImport.Package = FindPackage;
+					}
 				}
 				else
 				{
 					int32* FindGlobalExport = GlobalExportsByFullName.Find(GlobalImport.FullName);
-					check(FindGlobalExport);
-					GlobalImport.GlobalExportIndex = *FindGlobalExport;
+					if (!FindGlobalExport)
+					{
+						UE_LOG(LogIoStore, Warning, TEXT("Missing import object '%s'"), *GlobalImport.FullName);
+						GlobalImport.bIsMissingImport = true;
+					}
+					else
+					{
+						GlobalImport.GlobalExportIndex = *FindGlobalExport;
+					}
 				}
 			}
 
@@ -2352,12 +2368,14 @@ int32 CreateTarget(
 
 		for (FPackage* Package : Packages)
 		{
+			bool bHasMissingImports = false;
 			Package->ImportedPackages.Reserve(Package->ImportCount);
 			for (int32 I = 0; I < Package->ImportCount; ++I)
 			{
 				int32 GlobalImportIndex = *GlobalImportsByFullName.Find(Package->ImportedFullNames[I]);
 
 				FImportData& ImportData = GlobalImports[GlobalImportIndex];
+				bHasMissingImports |= ImportData.bIsMissingImport;
 				Package->Imports.Add(ImportData.GlobalIndex);
 				if (ImportData.Package)
 				{
@@ -2365,6 +2383,7 @@ int32 CreateTarget(
 					ImportData.Package->ImportedByPackages.Add(Package);
 				}
 			}
+			UE_CLOG(bHasMissingImports, LogIoStore, Warning, TEXT("Missing imports for package '%s'"), *Package->Name.ToString());
 		}
 	}
 
@@ -2455,27 +2474,33 @@ int32 CreateTarget(
 						}
 						else
 						{
-							check(Import.GlobalExportIndex != - 1);
-							FExportData& Export = GlobalExports[Import.GlobalExportIndex];
-
-							FPackage* SourcePackage = PackageMap.FindRef(Export.SourcePackageName);
-							check(SourcePackage);
-
-							AddExternalExportArc(ExportGraph, *SourcePackage, Export.SourceIndex, PhaseFrom, *Package, I, PhaseTo);
-							Package->ImportedPreloadPackages.Add(SourcePackage);
-
-							LocalizedPackages.Reset();
-							SourceToLocalizedPackageMap.MultiFind(SourcePackage, LocalizedPackages);
-							for (FPackage* LocalizedPackage : LocalizedPackages)
+							if (Import.GlobalExportIndex != -1)
+							{
+								FExportData& Export = GlobalExports[Import.GlobalExportIndex];
+
+								FPackage* SourcePackage = PackageMap.FindRef(Export.SourcePackageName);
+								check(SourcePackage);
+
+								AddExternalExportArc(ExportGraph, *SourcePackage, Export.SourceIndex, PhaseFrom, *Package, I, PhaseTo);
+								Package->ImportedPreloadPackages.Add(SourcePackage);
+
+								LocalizedPackages.Reset();
+								SourceToLocalizedPackageMap.MultiFind(SourcePackage, LocalizedPackages);
+								for (FPackage* LocalizedPackage : LocalizedPackages)
+								{
+									UE_LOG(LogIoStore, Verbose, TEXT("For package '%s' (%d): Adding localized preload dependency '%s' in '%s'"),
+										*Package->Name.ToString(),
+										Package->GlobalPackageId,
+										*Export.ObjectName.ToString(),
+										*LocalizedPackage->Name.ToString());
+
+									AddExternalExportArc(ExportGraph, *LocalizedPackage, Export.SourceIndex, PhaseFrom, *Package, I, PhaseTo);
+									Package->ImportedPreloadPackages.Add(LocalizedPackage);
+								}
+							}
+							else
 							{
-								UE_LOG(LogIoStore, Verbose, TEXT("For package '%s' (%d): Adding localized preload dependency '%s' in '%s'"),
-									*Package->Name.ToString(),
-									Package->GlobalPackageId,
-									*Export.ObjectName.ToString(),
-									*LocalizedPackage->Name.ToString());
-
-								AddExternalExportArc(ExportGraph, *LocalizedPackage, Export.SourceIndex, PhaseFrom, *Package, I, PhaseTo);
-								Package->ImportedPreloadPackages.Add(LocalizedPackage);
+								UE_LOG(LogIoStore, Verbose, TEXT("Skipping export arc to '%s' due to missing import"), *Package->Name.ToString());
 							}
 						}
 					}
