Copying //UE4/Release-4.25Plus-Dev to Release-4.25Plus (//UE4/Release-4.25Plus)
#rnx

diff --git a/Engine/Source/Developer/IoStoreUtilities/Private/IoStoreUtilities.cpp b/Engine/Source/Developer/IoStoreUtilities/Private/IoStoreUtilities.cpp
--- a/Engine/Source/Developer/IoStoreUtilities/Private/IoStoreUtilities.cpp
+++ b/Engine/Source/Developer/IoStoreUtilities/Private/IoStoreUtilities.cpp
@@ -31,6 +31,8 @@
 #include "Misc/FileHelper.h"
 #include "Misc/ScopeLock.h"
 #include "Async/ParallelFor.h"
+#include "Async/AsyncFileHandle.h"
+#include "Async/Async.h"
 
 //PRAGMA_DISABLE_OPTIMIZATION
 
@@ -39,17 +41,50 @@ IMPLEMENT_MODULE(FDefaultModuleImpl, IoStoreUtilities);
 DEFINE_LOG_CATEGORY_STATIC(LogIoStore, Log, All);
 
 #define IOSTORE_CPU_SCOPE(NAME) TRACE_CPUPROFILER_EVENT_SCOPE(IoStore##NAME);
-#define IOSTORE_CPU_SCOPE_DATA(NAME, DATA) TRACE_CPUPROFILER_EVENT_SCOPE(NAME);
+#define IOSTORE_CPU_SCOPE_DATA(NAME, DATA) TRACE_CPUPROFILER_EVENT_SCOPE(IoStore##NAME);
 
 #define OUTPUT_CHUNKID_DIRECTORY 0
 #define OUTPUT_NAMEMAP_CSV 0
 #define OUTPUT_IMPORTMAP_CSV 0
-#define SKIP_WRITE_CONTAINER 0
-#define OUTPUT_CONTAINER_CSV 0
+#define OUTPUT_DEBUG_PACKAGE_HASHES 0
+
+static const FName DefaultCompressionMethod = NAME_Zlib;
+static const int64 DefaultCompressionBlockSize = 64 << 10;
+const int64 DefaultMemoryMappingAlignment = 16 << 10;
+
+class FNameReaderProxyArchive
+	: public FArchiveProxy
+{
+	const TArray<FNameEntryId>& NameMap;
+
+public:
+	using FArchiveProxy::FArchiveProxy;
+
+	FNameReaderProxyArchive(FArchive& InAr, const TArray<FNameEntryId>& InNameMap)
+		: FArchiveProxy(InAr)
+		, NameMap(InNameMap) { }
+
+	FArchive& operator<<(FName& Name)
+	{
+		int32 NameIndex, Number;
+		InnerArchive << NameIndex << Number;
+
+		if (!NameMap.IsValidIndex(NameIndex))
+		{
+			UE_LOG(LogStreaming, Fatal, TEXT("Bad name index %i/%i"), NameIndex, NameMap.Num());
+		}
+
+		const FNameEntryId MappedName = NameMap[NameIndex];
+		Name = FName::CreateFromDisplayId(MappedName, Number);
+
+		return *this;
+	}
+};
 
 struct FContainerSourceFile 
 {
 	FString NormalizedPath;
+	bool bNeedsCompression = false;
 };
 
 struct FContainerSourceSpec
@@ -60,8 +95,8 @@ struct FContainerSourceSpec
 
 struct FCookedFileStatData
 {
-	enum EFileExt { UMap, UAsset, UBulk, UPtnl };
-	enum EFileType { PackageHeader, BulkData };
+	enum EFileExt { UMap, UAsset, UExp, UBulk, UPtnl, UMappedBulk };
+	enum EFileType { PackageHeader, PackageData, BulkData };
 
 	int64 FileSize = 0;
 	EFileType FileType = PackageHeader;
@@ -77,63 +112,70 @@ struct FContainerTargetFile
 	FPackage* Package = nullptr;
 	FString NormalizedSourcePath;
 	FString TargetPath;
-	uint64 Size = 0;
+	uint64 SourceSize = 0;
+	uint64 TargetSize = 0;
+	uint64 Offset = 0;
+	uint64 Padding = 0;
+	uint64 Alignment = 0;
+	FIoChunkId ChunkId;
+	FSHAHash Hash;
 	bool bIsBulkData = false;
 	bool bIsOptionalBulkData = false;
+	bool bIsMemoryMappedBulkData = false;
+	bool bForceUncompressed = false;
+};
+
+struct FIoStoreArguments
+{
+	FString OutputDir;
+	FString CookedDir;
+	FString OutputReleaseVersionDir;
+	FString BasedOnReleaseVersionDir;
+	TArray<FContainerSourceSpec> Containers;
+	FCookedFileStatMap CookedFileStatMap;
+	TMap<FName, uint64> GameOrderMap;
+	TMap<FName, uint64> CookerOrderMap;
+	int64 MemoryMappingAlignment = 0;
+};
+
+struct FTocMetaEntry
+{
+	uint64 Offset = 0;
+	uint64 Size = 0;
+	FSHAHash Hash;
+
+	friend FArchive& operator<<(FArchive& Ar, FTocMetaEntry& TocMetaEntry)
+	{
+		Ar << TocMetaEntry.Offset;
+		Ar << TocMetaEntry.Size;
+		Ar << TocMetaEntry.Hash;
+		return Ar;
+	}
 };
 
 struct FContainerTargetSpec
 {
+	FString Name;
+	FString OutputPath;
 	FIoStoreWriter* IoStoreWriter;
 	TArray<FContainerTargetFile> TargetFiles;
-
 	TUniquePtr<FIoStoreEnvironment> IoStoreEnv;
+	TArray<FTocMetaEntry> PreviousBuildToc;
+	bool bIsCompressed = false;
 };
 
 struct FPackageAssetData
 {
-	FCriticalSection CriticalSection;
 	TArray<FObjectImport> ObjectImports;
 	TArray<FObjectExport> ObjectExports;
 	TArray<FPackageIndex> PreloadDependencies;
 };
 
-class FNameProxyArchive
-	: public FArchiveProxy
-{
-public:
-	using FArchiveProxy::FArchiveProxy;
-
-	FNameProxyArchive(FArchive& InAr, const TArray<FNameEntryId>& InNameMap)
-		: FArchiveProxy(InAr)
-		, NameMap(InNameMap) { }
-
-	virtual FArchive& operator<<(FName& Name) override
-	{
-		int32 NameIndex, Number;
-		InnerArchive << NameIndex << Number;
-
-		if (!NameMap.IsValidIndex(NameIndex))
-		{
-			UE_LOG(LogIoStore, Fatal, TEXT("Bad name index %i/%i"), NameIndex, NameMap.Num());
-		}
-
-		const FNameEntryId MappedName = NameMap[NameIndex];
-		Name = FName::CreateFromDisplayId(MappedName, Number);
-
-		return *this;
-	}
-
-private:
-	const TArray<FNameEntryId>& NameMap;
-};
-
 struct FPackage;
 using FPackageMap = TMap<FName, FPackage*>;
-using FSourceToLocalizedPackageIdMap = TMap<int32,int32>;
-using FSourceToLocalizedPackageMultimap = TMultiMap<FPackage*,FPackage*>;
-using FLocalizedToSourceImportIndexMap = TMap<int32, int32>;
-using FCulturePackageMap = TMap<FString, FSourceToLocalizedPackageIdMap>;
+using FPackageGlobalIdMap = TMap<FName, FPackageId>;
+using FSourceToLocalizedPackageMultimap = TMultiMap<FPackage*, FPackage*>;
+using FLocalizedToSourceImportIndexMap = TMap<FPackageObjectIndex, FPackageObjectIndex>;
 
 static constexpr TCHAR L10NPrefix[] = TEXT("/Game/L10N/");
 
@@ -172,8 +214,6 @@ class FNameMapBuilder
 public:
 	void MarkNamesAsReferenced(const TArray<FName>& Names, TArray<FNameEntryId>& OutNameMap, TArray<int32>& OutNameIndices)
 	{
-		FScopeLock _(&CriticalSection);
-
 		for (const FName& Name : Names)
 		{
 			const FNameEntryId Id = Name.GetComparisonIndex();
@@ -191,8 +231,6 @@ public:
 
 	void MarkNameAsReferenced(const FName& Name)
 	{
-		FScopeLock _(&CriticalSection);
-
 		const FNameEntryId Id = Name.GetComparisonIndex();
 		int32& Index = NameIndices.FindOrAdd(Id);
 		if (Index == 0)
@@ -224,8 +262,6 @@ public:
 
 	int32 MapName(const FName& Name) const
 	{
-		FScopeLock _(&CriticalSection);
-
 		const FNameEntryId Id = Name.GetComparisonIndex();
 		const int32* Index = NameIndices.Find(Id);
 		check(Index);
@@ -234,18 +270,50 @@ public:
 
 	void SerializeName(FArchive& A, const FName& N) const
 	{
-		FScopeLock _(&CriticalSection);
-
 		int32 NameIndex = MapName(N);
 		int32 NameNumber = N.GetNumber();
 		A << NameIndex << NameNumber;
 	}
 
+	const TMap<FNameEntryId, int32>& GetNameIndices() const
+	{
+		return NameIndices;
+	}
+
 	const TArray<FNameEntryId>& GetNameMap() const
 	{
 		return NameMap;
 	}
 
+	friend FArchive& operator<<(FArchive& Ar, FNameMapBuilder& NameMapBuilder)
+	{
+		if (Ar.IsSaving())
+		{
+			int32 NameCount = NameMapBuilder.NameMap.Num();
+			Ar << NameCount;
+			for (FNameEntryId NameEntryId : NameMapBuilder.NameMap)
+			{
+				const FNameEntry* NameEntry = FName::GetEntry(NameEntryId);
+				NameEntry->Write(Ar);
+			}
+		}
+		else
+		{
+			int32 NameCount;
+			Ar << NameCount;
+			for (int32 NameIndex = 0; NameIndex < NameCount; ++NameIndex)
+			{
+				FNameEntrySerialized NameEntrySerialized(ENAME_LinkerConstructor);
+				Ar << NameEntrySerialized;
+				FName Name(NameEntrySerialized);
+				FNameEntryId NameId = Name.GetComparisonIndex();
+				NameMapBuilder.NameMap.Add(NameId);
+				NameMapBuilder.NameIndices.Add(NameId, NameIndex + 1);
+			}
+		}
+		return Ar;
+	}
+
 #if OUTPUT_NAMEMAP_CSV
 	void SaveCsv(const FString& CsvFilePath)
 	{
@@ -277,7 +345,6 @@ public:
 #endif
 
 private:
-	mutable FCriticalSection CriticalSection;
 	TMap<FNameEntryId, int32> NameIndices;
 	TArray<FNameEntryId> NameMap;
 #if OUTPUT_NAMEMAP_CSV
@@ -323,27 +390,27 @@ FChunkIdCsv ChunkIdCsv;
 
 #endif
 
-static FIoChunkId CreateChunkId(int32 GlobalPackageId, uint16 ChunkIndex, EIoChunkType ChunkType, const TCHAR* DebugString)
+static FIoChunkId CreateChunkId(FPackageId GlobalPackageId, uint16 ChunkIndex, EIoChunkType ChunkType, const TCHAR* DebugString)
 {
-	FIoChunkId ChunkId = CreateIoChunkId(GlobalPackageId, ChunkIndex, ChunkType);
+	FIoChunkId ChunkId = CreateIoChunkId(GlobalPackageId.ToIndex(), ChunkIndex, ChunkType);
 #if OUTPUT_CHUNKID_DIRECTORY
 	ChunkIdCsv.AddChunk(GlobalPackageId, 0, ChunkIndex, (uint8)ChunkType, GetTypeHash(ChunkId), DebugString);
 #endif
 	return ChunkId;
 }
 
-static FIoChunkId CreateChunkIdForBulkData(int32 GlobalPackageId,EIoChunkType ChunkType, const TCHAR* DebugString)
+static FIoChunkId CreateChunkIdForBulkData(FPackageId GlobalPackageId,EIoChunkType ChunkType, const TCHAR* DebugString)
 {
-	FIoChunkId ChunkId = CreateIoChunkId(GlobalPackageId, 0, ChunkType);
+	FIoChunkId ChunkId = CreateIoChunkId(GlobalPackageId.ToIndex(), 0, ChunkType);
 #if OUTPUT_CHUNKID_DIRECTORY
 	ChunkIdCsv.AddChunk(GlobalPackageId, 0, 0, (uint8)ChunkType, GetTypeHash(ChunkId), DebugString);
 #endif
 	return ChunkId;
 }
 
-static FIoChunkId CreateChunkIdForBulkData(int32 GlobalPackageId, int64 BulkdataOffset, EIoChunkType ChunkType, const TCHAR* DebugString)
+static FIoChunkId CreateChunkIdForBulkData(FPackageId GlobalPackageId, int64 BulkdataOffset, EIoChunkType ChunkType, const TCHAR* DebugString)
 {
-	FIoChunkId ChunkId = CreateBulkdataChunkId(GlobalPackageId, BulkdataOffset, ChunkType);
+	FIoChunkId ChunkId = CreateBulkdataChunkId(GlobalPackageId.ToIndex(), BulkdataOffset, ChunkType);
 #if OUTPUT_CHUNKID_DIRECTORY
 	ChunkIdCsv.AddChunk(GlobalPackageId, 0, 0, (uint8)ChunkType, GetTypeHash(ChunkId), DebugString);
 #endif
@@ -423,7 +490,6 @@ struct FExportGraphNode
 	FPackage* Package;
 	FExportBundleEntry BundleEntry;
 	TSet<FExportGraphNode*> ExternalDependencies;
-	TSet<uint32> ScriptDependencies;
 	uint64 NodeIndex;
 };
 
@@ -478,27 +544,33 @@ private:
 	TMultiMap<FExportGraphNode*, FExportGraphNode*> Edges;
 };
 
+#if OUTPUT_DEBUG_PACKAGE_HASHES
+struct FPackageHashes
+{
+	FSHAHash UAssetHash;
+	FSHAHash UExpHash;
+	FSHAHash ExportBundleHash;
+};
+#endif
+
 struct FPackage
 {
 	FName Name;
 	FName SourcePackageName; // for localized packages
 	FString FileName;
-	FString RelativeFileName;
-	int32 GlobalPackageId = 0;
+	FPackageId GlobalPackageId;
 	FString Region; // for localized packages
-	int32 SourceGlobalPackageId = -1; // for localized packages
+	FPackageId SourceGlobalPackageId; // for localized packages
 	int32 ImportedPackagesSerializeCount = 0; // < ImportedPackages.Num() for source packages that have localized packages
 	uint32 PackageFlags = 0;
+	uint32 CookedHeaderSize = 0;
 	int32 NameCount = 0;
 	int32 ImportCount = 0;
+	int32 PreloadDependencyCount = 0;
 	int32 ExportCount = 0;
-	int32 FirstGlobalImport = -1;
-	int32 GlobalImportCount = -1;
 	int32 ImportIndexOffset = -1;
 	int32 ExportIndexOffset = -1;
 	int32 PreloadIndexOffset = -1;
-	int32 FirstExportBundleMetaEntry = -1;
-	int64 BulkDataStartOffset = -1;
 	int64 UExpSize = 0;
 	int64 UAssetSize = 0;
 	int64 SummarySize = 0;
@@ -506,7 +578,7 @@ struct FPackage
 	int64 NameMapSize = 0;
 	int64 ImportMapSize = 0;
 	int64 ExportMapSize = 0;
-	int64 ExportBundlesSize = 0;
+	int64 ExportBundlesHeaderSize = 0;
 
 	bool bIsLocalizedAndConformed = false;
 	bool bHasCircularImportDependencies = false;
@@ -521,11 +593,11 @@ struct FPackage
 	TArray<FNameEntryId> NameMap;
 	TArray<int32> NameIndices;
 
-	TArray<int32> Imports;
+	TArray<FPackageObjectIndex> Imports;
+	TArray<FPackageObjectIndex> PublicExports;
 	TArray<int32> Exports;
 	TArray<FArc> InternalArcs;
 	TMap<FPackage*, TArray<FArc>> ExternalArcs;
-	TArray<FArc> ScriptArcs;
 	
 	TArray<FExportBundle> ExportBundles;
 	TMap<FExportGraphNode*, uint32> ExportBundleMap;
@@ -537,9 +609,12 @@ struct FPackage
 	FPackageGraphNode* Node = nullptr;
 
 	TArray<uint8> HeaderData;
-	uint64 ExportsPayloadSize = 0;
+	uint64 ExportsSerialSize = 0;
 
 	uint64 DiskLayoutOrder = MAX_uint64;
+#if OUTPUT_DEBUG_PACKAGE_HASHES
+	FPackageHashes Hashes;
+#endif
 };
 
 struct FCircularImportChain
@@ -608,18 +683,11 @@ TArray<FPackage*> FPackageGraph::TopologicalSort() const
 	TArray<FPackage*> Result;
 	Result.Reserve(Nodes.Num());
 	
-	TSet<FPackageGraphNode*> UnmarkedNodes;
-	UnmarkedNodes.Append(Nodes);
-
 	struct
 	{
 		void Visit(FPackageGraphNode* Node)
 		{
-			if (Node->bPermanentMark)
-			{
-				return;
-			}
-			if (Node->bTemporaryMark)
+			if (Node->bPermanentMark || Node->bTemporaryMark)
 			{
 				return;
 			}
@@ -631,39 +699,39 @@ TArray<FPackage*> FPackageGraph::TopologicalSort() const
 			}
 			Node->bTemporaryMark = false;
 			Node->bPermanentMark = true;
-			UnmarkedNodes.Remove(Node);
-			Result.Insert(Node->Package, 0);
+			Result.Add(Node->Package);
 		}
 
-		TSet<FPackageGraphNode*>& UnmarkedNodes;
 		TMultiMap<FPackageGraphNode*, FPackageGraphNode*>& Edges;
 		TArray<FPackage*>& Result;
 
-	} Visitor{ UnmarkedNodes, EdgesCopy, Result };
+	} Visitor{ EdgesCopy, Result };
 
-	while (Result.Num() < Nodes.Num())
+	for (FPackageGraphNode* Node : Nodes)
 	{
-		auto It = UnmarkedNodes.CreateIterator();
-		FPackageGraphNode* UnmarkedNode = *It;
-		It.RemoveCurrent();
-		Visitor.Visit(UnmarkedNode);
+		Visitor.Visit(Node);
 	}
+	check(Result.Num() == Nodes.Num());
 
+	Algo::Reverse(Result);
 	return Result;
 }
 
 TArray<FExportGraphNode*> FExportGraph::ComputeLoadOrder(const TArray<FPackage*>& Packages) const
 {
+	IOSTORE_CPU_SCOPE(ComputeLoadOrder);
 	FPackageGraph PackageGraph;
-	for (FPackage* Package : Packages)
-	{
-		Package->Node = PackageGraph.AddNode(Package);
-	}
-	for (FPackage* Package : Packages)
 	{
-		for (FPackage* ImportedPackage : Package->ImportedPackages)
+		for (FPackage* Package : Packages)
+		{
+			Package->Node = PackageGraph.AddNode(Package);
+		}
+		for (FPackage* Package : Packages)
 		{
-			PackageGraph.AddImportDependency(ImportedPackage->Node, Package->Node);
+			for (FPackage* ImportedPackage : Package->ImportedPackages)
+			{
+				PackageGraph.AddImportDependency(ImportedPackage->Node, Package->Node);
+			}
 		}
 	}
 
@@ -672,8 +740,9 @@ TArray<FExportGraphNode*> FExportGraph::ComputeLoadOrder(const TArray<FPackage*>
 	int32 NodeCount = Nodes.Num();
 	TArray<uint32> NodesIncomingEdgeCount;
 	NodesIncomingEdgeCount.AddZeroed(NodeCount);
+
 	TMultiMap<FExportGraphNode*, FExportGraphNode*> EdgesCopy = Edges;
-	for (auto& KV : EdgesCopy)
+	for (auto& KV : Edges)
 	{
 		FExportGraphNode* ToNode = KV.Value;
 		++NodesIncomingEdgeCount[ToNode->NodeIndex];
@@ -696,7 +765,7 @@ TArray<FExportGraphNode*> FExportGraph::ComputeLoadOrder(const TArray<FPackage*>
 		{
 			while (Package->NodesWithNoIncomingEdges.Num() > 0)
 			{
-				FExportGraphNode* RemovedNode = Package->NodesWithNoIncomingEdges.Pop();
+				FExportGraphNode* RemovedNode = Package->NodesWithNoIncomingEdges.Pop(false);
 				LoadOrder.Add(RemovedNode);
 				for (auto EdgeIt = EdgesCopy.CreateKeyIterator(RemovedNode); EdgeIt; ++EdgeIt)
 				{
@@ -728,12 +797,6 @@ static void AddExternalExportArc(FExportGraph& ExportGraph, FPackage& FromPackag
 	ExportGraph.AddExternalDependency(FromNode, ToNode);
 }
 
-static void AddScriptArc(FPackage& Package, uint32 GlobalImportIndex, uint32 ExportIndex, EPreloadDependencyType Phase)
-{
-	FExportGraphNode* Node = Phase == PreloadDependencyType_Create ? Package.CreateExportNodes[ExportIndex] : Package.SerializeExportNodes[ExportIndex];
-	Node->ScriptDependencies.Add(GlobalImportIndex);
-}
-
 static void AddPostLoadArc(FPackage& FromPackage, FPackage& ToPackage)
 {
 	TArray<FArc>& ExternalArcs = ToPackage.ExternalArcs.FindOrAdd(&FromPackage);
@@ -758,12 +821,6 @@ static void AddUniqueExternalBundleArc(FPackage& FromPackage, uint32 FromBundleI
 	ExternalArcs.AddUnique({ FromNodeIndex, ToNodeIndex });
 }
 
-static void AddUniqueScriptBundleArc(FPackage& Package, uint32 GlobalImportIndex, uint32 BundleIndex)
-{
-	uint32 NodeIndex = EEventLoadNode2::Package_NumPhases + BundleIndex * EEventLoadNode2::ExportBundle_NumPhases + EEventLoadNode2::ExportBundle_Process;
-	Package.ScriptArcs.AddUnique({ GlobalImportIndex, NodeIndex });
-}
-
 static void AddReachablePackagesRecursive(FPackage& Package, FPackage& PackageWithImports, TSet<FPackage*>& Visited, bool bFirst)
 {
 	if (!bFirst)
@@ -839,7 +896,7 @@ static bool FindNewCircularImportChains(
 	return bFoundNew;
 }
 
-static void AddPostLoadDependencies(
+static void AddPackagePostLoadDependencies(
 	FPackage& Package,
 	TSet<FPackage*>& Visited,
 	TSet<FCircularImportChain>& CircularChains)
@@ -900,9 +957,27 @@ static void AddPostLoadDependencies(
 	*/
 }
 
+static int32 AddPostLoadDependencies(TArray<FPackage*>& Packages)
+{
+	IOSTORE_CPU_SCOPE(PostLoadDependencies);
+	UE_LOG(LogIoStore, Display, TEXT("Adding postload dependencies..."));
+
+	TSet<FPackage*> Visited;
+	TSet<FCircularImportChain> CircularChains;
+
+	for (FPackage* Package : Packages)
+	{
+		Visited.Reset();
+		Visited.Add(Package);
+		AddPackagePostLoadDependencies(*Package, Visited, CircularChains);
+	}
+	return CircularChains.Num();
+};
+
 static void BuildBundles(FExportGraph& ExportGraph, const TArray<FPackage*>& Packages)
 {
 	IOSTORE_CPU_SCOPE(BuildBundles)
+	UE_LOG(LogIoStore, Display, TEXT("Building bundles..."));
 
 	TArray<FExportGraphNode*> ExportLoadOrder = ExportGraph.ComputeLoadOrder(Packages);
 	FPackage* LastPackage = nullptr;
@@ -939,23 +1014,34 @@ static void BuildBundles(FExportGraph& ExportGraph, const TArray<FPackage*>& Pac
 				AddUniqueExternalBundleArc(*ExternalDependency->Package, *FindDependentBundleIndex, *Package, BundleIndex);
 			}
 		}
-		for (uint32 ScriptDependencyGlobalImportIndex : Node->ScriptDependencies)
-		{
-			AddUniqueScriptBundleArc(*Package, ScriptDependencyGlobalImportIndex, BundleIndex);
-		}
 		Bundle->Nodes.Add(Node);
 		Package->ExportBundleMap.Add(Node, BundleIndex);
 	}
 }
 
-static void CreateDiskLayout(const TArray<FPackage*>& Packages, const TMap<FString, uint64> PackageOrderMap, const TMap<FString, uint64>& CookerFileOpenOrderMap)
+struct FFileLayoutNode
+{
+	FFileLayoutNode* Prev = nullptr;
+	FFileLayoutNode* Next = nullptr;
+	FContainerTargetFile* TargetFile = nullptr;
+	uint64 IdealOrder = 0;
+	uint64 PreviousOffset = 0;
+	bool bHasPreviousOffset = false;
+};
+
+static void CreateDiskLayout(
+	TArray<FContainerTargetSpec>& ContainerTargets,
+	const TArray<FPackage*>& Packages,
+	const TMap<FName, uint64> GameOrderMap,
+	const TMap<FName, uint64>& CookerOrderMap,
+	uint64 CompressionBlockSize,
+	uint64 MemoryMappingAlignment)
 {
 	IOSTORE_CPU_SCOPE(CreateDiskLayout);
 
 	struct FCluster
 	{
 		TArray<FPackage*> Packages;
-		uint64 ClusterSize = 0;
 	};
 
 	TArray<FCluster*> Clusters;
@@ -965,20 +1051,20 @@ static void CreateDiskLayout(const TArray<FPackage*>& Packages, const TMap<FStri
 	struct FPackageAndOrder
 	{
 		FPackage* Package;
-		uint64 PackageLoadOrder;
+		uint64 GameOpenOrder;
 		uint64 CookerOpenOrder;
 
 		bool operator<(const FPackageAndOrder& Other) const
 		{
-			if (PackageLoadOrder != Other.PackageLoadOrder)
+			if (GameOpenOrder != Other.GameOpenOrder)
 			{
-				return PackageLoadOrder < Other.PackageLoadOrder;
+				return GameOpenOrder < Other.GameOpenOrder;
 			}
 			if (CookerOpenOrder != Other.CookerOpenOrder)
 			{
 				return CookerOpenOrder < Other.CookerOpenOrder;
 			}
-			// Fallback to reverse bundle order
+			// Fallback to reverse bundle order (so that packages are considered before their imports)
 			return Package->ExportBundles[0].LoadOrder > Other.Package->ExportBundles[0].LoadOrder;
 		}
 	};
@@ -989,30 +1075,30 @@ static void CreateDiskLayout(const TArray<FPackage*>& Packages, const TMap<FStri
 	{
 		FPackageAndOrder& Entry = SortedPackages.AddDefaulted_GetRef();
 		Entry.Package = Package;
-		const uint64* FindPackageLoadOrder = PackageOrderMap.Find(Package->Name.ToString());
-		Entry.PackageLoadOrder = FindPackageLoadOrder ? *FindPackageLoadOrder : MAX_uint64;
-		const uint64* FindCookerOpenOrder = CookerFileOpenOrderMap.Find(Package->RelativeFileName);
+		const uint64* FindGameOpenOrder = GameOrderMap.Find(Package->Name);
+		Entry.GameOpenOrder = FindGameOpenOrder ? *FindGameOpenOrder : MAX_uint64;
+		const uint64* FindCookerOpenOrder = CookerOrderMap.Find(Package->Name);
 		/*if (!FindCookerOpenOrder)
 		{
-			UE_LOG(LogIoStore, Warning, TEXT("Missing cooker order for package: %s"), *Package->RelativeFileName);
+			UE_LOG(LogIoStore, Warning, TEXT("Missing cooker order for package: %s"), *Package->Name.ToString());
 		}*/
 		Entry.CookerOpenOrder = FindCookerOpenOrder ? *FindCookerOpenOrder : MAX_uint64;
 	}
-	bool bHasPackageOrder = true;
+	bool bHasGameOrder = true;
 	bool bHasCookerOrder = true;
 	int32 LastAssignedCount = 0;
 	Algo::Sort(SortedPackages);
 	for (FPackageAndOrder& Entry : SortedPackages)
 	{
-		if (bHasPackageOrder && Entry.PackageLoadOrder == MAX_uint64)
+		if (bHasGameOrder && Entry.GameOpenOrder == MAX_uint64)
 		{
-			UE_LOG(LogIoStore, Display, TEXT("Ordered %d/%d packages using package load order"), AssignedPackages.Num(), Packages.Num());
+			UE_LOG(LogIoStore, Display, TEXT("Ordered %d/%d packages using game open order"), AssignedPackages.Num(), Packages.Num());
 			LastAssignedCount = AssignedPackages.Num();
-			bHasPackageOrder = false;
+			bHasGameOrder = false;
 		}
-		if (!bHasPackageOrder && bHasCookerOrder && Entry.CookerOpenOrder == MAX_uint64)
+		if (!bHasGameOrder && bHasCookerOrder && Entry.CookerOpenOrder == MAX_uint64)
 		{
-			UE_LOG(LogIoStore, Display, TEXT("Ordered %d/%d packages using cooker order"), AssignedPackages.Num() - LastAssignedCount, Packages.Num() - LastAssignedCount);
+			UE_LOG(LogIoStore, Display, TEXT("Ordered %d/%d packages using cooker open order"), AssignedPackages.Num() - LastAssignedCount, Packages.Num() - LastAssignedCount);
 			LastAssignedCount = AssignedPackages.Num();
 			bHasCookerOrder = false;
 		}
@@ -1047,11 +1133,6 @@ static void CreateDiskLayout(const TArray<FPackage*>& Packages, const TMap<FStri
 		{
 			return A->ExportBundles[0].LoadOrder < B->ExportBundles[0].LoadOrder;
 		});
-
-		for (FPackage* Package : Cluster->Packages)
-		{
-			Cluster->ClusterSize += Package->HeaderData.Num() + Package->ExportsPayloadSize;
-		}
 	}
 
 	uint64 LayoutIndex = 0;
@@ -1063,6 +1144,225 @@ static void CreateDiskLayout(const TArray<FPackage*>& Packages, const TMap<FStri
 		}
 		delete Cluster;
 	}
+
+	for (FContainerTargetSpec& ContainerTarget : ContainerTargets)
+	{
+		Algo::Sort(ContainerTarget.TargetFiles, [](const FContainerTargetFile& A, const FContainerTargetFile& B)
+		{
+			if (A.bIsMemoryMappedBulkData != B.bIsMemoryMappedBulkData)
+			{
+				return B.bIsMemoryMappedBulkData;
+			}
+			if (A.bIsBulkData != B.bIsBulkData)
+			{
+				return B.bIsBulkData;
+			}
+
+			return A.Package->DiskLayoutOrder < B.Package->DiskLayoutOrder;
+		});
+
+		TMap<FSHAHash, FTocMetaEntry> PreviousBuildTocMap;
+		for (const FTocMetaEntry& TocMetaEntry : ContainerTarget.PreviousBuildToc)
+		{
+			if (PreviousBuildTocMap.Find(TocMetaEntry.Hash))
+			{
+				UE_DEBUG_BREAK();
+			}
+			PreviousBuildTocMap.Add(TocMetaEntry.Hash, TocMetaEntry);
+		}
+
+		uint64 DiffCount = 0;
+		uint64 DiffSize = 0;
+		TArray<FFileLayoutNode*> LayoutNodes;
+		LayoutNodes.Reserve(ContainerTarget.TargetFiles.Num());
+		uint64 IdealOrder = 0;
+		TMap<uint64, FFileLayoutNode*> UnassignedNodes;
+		TMap<uint64, FFileLayoutNode*> AssignedNodes;
+		for (FContainerTargetFile& ContainerTargetFile : ContainerTarget.TargetFiles)
+		{
+			FFileLayoutNode* LayoutNode = new FFileLayoutNode();
+			LayoutNodes.Add(LayoutNode);
+			LayoutNode->TargetFile = &ContainerTargetFile;
+			LayoutNode->IdealOrder = IdealOrder++;
+			UnassignedNodes.Add(LayoutNode->IdealOrder, LayoutNode);
+
+			FTocMetaEntry* FindPreviousBuildEntry = PreviousBuildTocMap.Find(ContainerTargetFile.Hash);
+			if (!FindPreviousBuildEntry)
+			{
+				++DiffCount;
+				DiffSize += ContainerTargetFile.TargetSize;
+				//UE_LOG(LogIoStore, Display, TEXT("Diffing: %s"), *ContainerTargetFile.TargetPath);
+			}
+			else
+			{
+				LayoutNode->PreviousOffset = FindPreviousBuildEntry->Offset;
+				LayoutNode->bHasPreviousOffset = true;
+			}
+		}
+		UE_LOG(LogIoStore, Display, TEXT("%s: %d/%d new/modified entries (%fMB)"), *ContainerTarget.Name, DiffCount, ContainerTarget.TargetFiles.Num(), DiffSize / 1024.0 / 1024.0);
+
+		Algo::Sort(LayoutNodes, [](const FFileLayoutNode* A, const FFileLayoutNode* B)
+		{
+			if (A->bHasPreviousOffset != B->bHasPreviousOffset)
+			{
+				return A->bHasPreviousOffset;
+			}
+			else if (A->PreviousOffset != B->PreviousOffset)
+			{
+				return A->PreviousOffset < B->PreviousOffset;
+			}
+			return A->IdealOrder < B->IdealOrder;
+		});
+
+		FFileLayoutNode Head;
+		FFileLayoutNode Tail;
+		Head.Next = &Tail;
+		Tail.Prev = &Head;
+		FFileLayoutNode* PreviousNode = &Head;
+		for (FFileLayoutNode* LayoutNode : LayoutNodes)
+		{
+			check(UnassignedNodes.Contains(LayoutNode->IdealOrder));
+			check(!AssignedNodes.Contains(LayoutNode->IdealOrder));
+			
+			FFileLayoutNode* OrderAfterNode;
+			if (LayoutNode->bHasPreviousOffset)
+			{
+				OrderAfterNode = PreviousNode;
+			}
+			else
+			{
+				if (LayoutNode->IdealOrder == 0)
+				{
+					OrderAfterNode = &Head;
+				}
+				else
+				{
+					OrderAfterNode = AssignedNodes.FindRef(LayoutNode->IdealOrder - 1);
+					check(OrderAfterNode);
+				}
+			}
+
+			OrderAfterNode->Next->Prev = LayoutNode;
+			LayoutNode->Next = OrderAfterNode->Next;
+			LayoutNode->Prev = OrderAfterNode;
+			OrderAfterNode->Next = LayoutNode;
+			
+			AssignedNodes.Add(LayoutNode->IdealOrder, LayoutNode);
+			UnassignedNodes.Remove(LayoutNode->IdealOrder);
+
+			PreviousNode = LayoutNode;
+		}
+
+		TArray<bool> BlocksModificationState;
+		uint64 CurrentOffset = 0;
+		FFileLayoutNode* CurrentNode = Head.Next;
+		bool bIsAlignedForMemoryMappedBulkData = false;
+		while (CurrentNode != &Tail)
+		{
+			uint64 ModificationOffset = 0;
+			uint64 ModificationSize = 0;
+
+			if (CurrentNode->TargetFile->bIsMemoryMappedBulkData)
+			{
+				if (!bIsAlignedForMemoryMappedBulkData)
+				{
+					if (!IsAligned(CurrentOffset, CompressionBlockSize))
+					{
+						uint64 RemainingInBlock = Align(CurrentOffset, CompressionBlockSize) - CurrentOffset;
+						CurrentNode->TargetFile->Padding = RemainingInBlock;
+						CurrentOffset += RemainingInBlock;
+					}
+					CurrentNode->TargetFile->Alignment = MemoryMappingAlignment;
+					bIsAlignedForMemoryMappedBulkData = true;
+				}
+				else
+				{
+					if (!IsAligned(CurrentOffset, MemoryMappingAlignment))
+					{
+						uint64 RequiredPadding = Align(CurrentOffset, MemoryMappingAlignment) - CurrentOffset;
+						CurrentNode->TargetFile->Padding = RequiredPadding;
+						CurrentOffset += RequiredPadding;
+					}
+				}
+			}
+			else if (bIsAlignedForMemoryMappedBulkData)
+			{
+				if (!IsAligned(CurrentOffset, CompressionBlockSize))
+				{
+					uint64 RemainingInBlock = Align(CurrentOffset, CompressionBlockSize) - CurrentOffset;
+					CurrentNode->TargetFile->Padding = RemainingInBlock;
+					CurrentOffset += RemainingInBlock;
+				}
+				bIsAlignedForMemoryMappedBulkData = false;
+			}
+
+			if (CurrentNode->bHasPreviousOffset)
+			{
+				uint64 PreviousOffsetInBlock = CurrentNode->PreviousOffset % CompressionBlockSize;
+				uint64 CurrentOffsetInBlock = CurrentOffset % CompressionBlockSize;
+				if (CurrentOffsetInBlock != PreviousOffsetInBlock)
+				{
+					uint64 PadBytes = 0;
+					if (CurrentOffsetInBlock < PreviousOffsetInBlock)
+					{
+						PadBytes = PreviousOffsetInBlock - CurrentOffsetInBlock;
+					}
+					else if (PreviousOffsetInBlock < CurrentOffsetInBlock)
+					{
+						PadBytes = CompressionBlockSize - CurrentOffsetInBlock + PreviousOffsetInBlock;
+					}
+					ModificationOffset = CurrentOffset;
+					ModificationSize += PadBytes;
+					CurrentOffset += PadBytes;
+					CurrentNode->TargetFile->Padding = PadBytes;
+				}
+				check(CurrentOffset % CompressionBlockSize == CurrentNode->PreviousOffset % CompressionBlockSize);
+			}
+			else
+			{
+				ModificationOffset = CurrentOffset;
+				ModificationSize = CurrentNode->TargetFile->TargetSize;
+			}
+
+			if (ModificationSize)
+			{
+				int32 BeginBlockIndex = ModificationOffset / CompressionBlockSize;
+				int32 EndBlockIndex = Align(ModificationOffset + ModificationSize, CompressionBlockSize) / CompressionBlockSize;
+				if (BlocksModificationState.Num() < EndBlockIndex)
+				{
+					BlocksModificationState.AddZeroed(EndBlockIndex - BlocksModificationState.Num());
+				}
+				for (int32 BlockIndex = BeginBlockIndex; BlockIndex < EndBlockIndex; ++BlockIndex)
+				{
+					BlocksModificationState[BlockIndex] = true;
+				}
+			}
+
+			CurrentNode->TargetFile->Offset = CurrentOffset;
+			CurrentOffset += CurrentNode->TargetFile->TargetSize;
+			CurrentNode = CurrentNode->Next;
+		}
+
+		uint64 ModifiedBlocksCount = 0;
+		for (bool BlockModified : BlocksModificationState)
+		{
+			if (BlockModified)
+			{
+				++ModifiedBlocksCount;
+			}
+		}
+		UE_LOG(LogIoStore, Display, TEXT("%s: %d new/modified blocks (%fMB)"), *ContainerTarget.Name, ModifiedBlocksCount, ModifiedBlocksCount * CompressionBlockSize / 1024.0f / 1024.0f);
+
+		for (FFileLayoutNode* LayoutNode : LayoutNodes)
+		{
+			delete LayoutNode;
+		}
+
+		Algo::Sort(ContainerTarget.TargetFiles, [](const FContainerTargetFile& A, const FContainerTargetFile& B)
+		{
+			return A.Offset < B.Offset;
+		});
+	}
 }
 
 static EIoChunkType BulkdataTypeToChunkIdType(FPackageStoreBulkDataManifest::EBulkdataType Type)
@@ -1073,123 +1373,94 @@ static EIoChunkType BulkdataTypeToChunkIdType(FPackageStoreBulkDataManifest::EBu
 		return EIoChunkType::BulkData;
 	case FPackageStoreBulkDataManifest::EBulkdataType::Optional:
 		return EIoChunkType::OptionalBulkData;
+	case FPackageStoreBulkDataManifest::EBulkdataType::MemoryMapped:
+		return EIoChunkType::MemoryMappedBulkData;
 	default:
 		UE_LOG(LogIoStore, Error, TEXT("Invalid EBulkdataType (%d) found!"), Type);
 		return EIoChunkType::Invalid;
 	}
 }
 
-static bool WriteBulkData(	const FString& Filename, FPackageStoreBulkDataManifest::EBulkdataType Type, const FPackage& Package, 
-							const FPackageStoreBulkDataManifest& BulkDataManifest, FIoStoreWriter* IoStoreWriter)
+static bool MapAdditionalBulkDataChunks(const FContainerTargetFile& TargetFile, const FPackageStoreBulkDataManifest& BulkDataManifest, FIoStoreWriter* IoStoreWriter)
 {
-
+	const FPackage& Package = *TargetFile.Package;
 	const FPackageStoreBulkDataManifest::FPackageDesc* PackageDesc = BulkDataManifest.Find(Package.FileName);
 	if (PackageDesc != nullptr)
 	{
-		const EIoChunkType ChunkIdType = BulkdataTypeToChunkIdType(Type);
-
-		const FIoChunkId BulkDataChunkId = CreateChunkIdForBulkData(Package.GlobalPackageId, ChunkIdType, *Package.FileName);
-
-#if !SKIP_WRITE_CONTAINER		
-		FIoBuffer IoBuffer;
+		FPackageStoreBulkDataManifest::EBulkdataType BulkDataType = TargetFile.bIsOptionalBulkData
+			? FPackageStoreBulkDataManifest::EBulkdataType::Optional
+			: TargetFile.bIsMemoryMappedBulkData
+				? FPackageStoreBulkDataManifest::EBulkdataType::MemoryMapped
+				: FPackageStoreBulkDataManifest::EBulkdataType::Normal;
 
-		TUniquePtr<FArchive> BulkAr(IFileManager::Get().CreateFileReader(*Filename));
-		if (BulkAr != nullptr)
-		{
-			uint8* BulkBuffer = static_cast<uint8*>(FMemory::Malloc(BulkAr->TotalSize()));
-			BulkAr->Serialize(BulkBuffer, BulkAr->TotalSize());
-			IoBuffer = FIoBuffer(FIoBuffer::AssumeOwnership, BulkBuffer, BulkAr->TotalSize());
-
-			BulkAr->Close();
-		}
-		
-		const FIoStatus AppendResult = IoStoreWriter->Append(BulkDataChunkId, IoBuffer, *FPaths::GetCleanFilename(Filename));
-		if (!AppendResult.IsOk())
-		{
-			UE_LOG(LogIoStore, Error, TEXT("Failed to append bulkdata for '%s' due to: %s"), *Package.FileName, *AppendResult.ToString());
-			return false;
-		}
-#endif
+		const EIoChunkType ChunkIdType = BulkdataTypeToChunkIdType(BulkDataType);
 
 		// Create additional mapping chunks as needed
 		for (const FPackageStoreBulkDataManifest::FPackageDesc::FBulkDataDesc& BulkDataDesc : PackageDesc->GetDataArray())
 		{
-			if (BulkDataDesc.Type == Type)
+			if (BulkDataDesc.Type == BulkDataType)
 			{
 				const FIoChunkId AccessChunkId = CreateChunkIdForBulkData(Package.GlobalPackageId, BulkDataDesc.ChunkId, ChunkIdType, *Package.FileName);
-#if !SKIP_WRITE_CONTAINER	
-				const FIoStatus PartialResult = IoStoreWriter->MapPartialRange(BulkDataChunkId, BulkDataDesc.Offset, BulkDataDesc.Size, AccessChunkId);
+				const FIoStatus PartialResult = IoStoreWriter->MapPartialRange(TargetFile.ChunkId, BulkDataDesc.Offset, BulkDataDesc.Size, AccessChunkId);
 				if (!PartialResult.IsOk())
 				{
 					UE_LOG(LogIoStore, Warning, TEXT("Failed to map partial range for '%s' due to: %s"), *Package.FileName, *PartialResult.ToString());
 				}
-#endif
 			}
 		}
 	}
-	else if(IFileManager::Get().FileExists(*Filename))
+	else
 	{
-		UE_LOG(LogIoStore, Error, TEXT("Unable to find an entry in the bulkdata manifest for '%s' the file might be out of date!"), *Package.FileName);
+		UE_LOG(LogIoStore, Warning, TEXT("Unable to find an entry in the bulkdata manifest for '%s' the file might be out of date!"), *Package.FileName);
 		return false;
 	}
 
 	return true;
 }
 
-struct FImportData
+struct FScriptImportData
+{
+	FName ObjectName;
+	FString FullName;
+	FPackageObjectIndex GlobalIndex;
+	FPackageObjectIndex OuterIndex;
+	FPackageObjectIndex OutermostIndex;
+	FPackageObjectIndex CDOClassIndex;
+	FString CDOClassFullName;
+	bool bInitialized = false;
+};
+
+struct FPackageImportData
 {
-	int32 GlobalIndex = -1;
-	int32 OuterIndex = -1;
-	int32 OutermostIndex = -1;
-	int32 GlobalExportIndex = -1;
-	int32 RefCount = 0;
 	FName ObjectName;
-	bool bIsPackage = false;
-	bool bIsScript = false;
-	bool bIsLocalized = false;
 	FString FullName;
+	FPackageObjectIndex GlobalIndex;
+	int32 GlobalExportIndex = -1;
 	FPackage* Package = nullptr;
+	bool bIsLocalized = false;
 	bool bIsMissingImport = false;
-
-	bool operator<(const FImportData& Other) const
-	{
-		if (bIsScript != Other.bIsScript)
-		{
-			return bIsScript;
-		}
-		if (bIsLocalized != Other.bIsLocalized)
-		{
-			return Other.bIsLocalized;
-		}
-		if (OutermostIndex != Other.OutermostIndex)
-		{
-			return OutermostIndex < Other.OutermostIndex;
-		}
-		if (bIsPackage != Other.bIsPackage)
-		{
-			return bIsPackage;
-		}
-		return FullName < Other.FullName;
-	}
-};
+	bool bInitialized = false;
+};
 
 struct FExportData
 {
 	int32 GlobalIndex = -1;
-	FName SourcePackageName;
 	FName ObjectName;
 	int32 SourceIndex = -1;
-	int32 GlobalImportIndex = -1;
-	FPackageIndex OuterIndex;
-	FPackageIndex ClassIndex;
-	FPackageIndex SuperIndex;
-	FPackageIndex TemplateIndex;
+	FPackageObjectIndex GlobalImportIndex;
+	FPackageObjectIndex OuterIndex;
+	FPackageObjectIndex ClassIndex;
+	FPackageObjectIndex SuperIndex;
+	FPackageObjectIndex TemplateIndex;
 	FString FullName;
 
+	FPackage* Package = nullptr;
 	FExportGraphNode* CreateNode = nullptr;
 	FExportGraphNode* SerializeNode = nullptr;
 };
 
+using FImportObjectsByFullName = TMap<FString, FPackageObjectIndex>;
+
 template <typename T>
 struct FGlobalObjects
 {
@@ -1197,9 +1468,17 @@ struct FGlobalObjects
 	TMap<FString, int32> ObjectsByFullName;
 };
 
-using FGlobalImports = FGlobalObjects<FImportData>;
+using FGlobalScriptImports = TArray<FScriptImportData>;
+using FGlobalPackageImports = TArray<FPackageImportData>;
 using FGlobalExports = FGlobalObjects<FExportData>;
 
+struct FGlobalImports
+{
+	FGlobalScriptImports ScriptImports;
+	FGlobalPackageImports PackageImports;
+	FImportObjectsByFullName ObjectsByFullName;
+};
+
 struct FGlobalPackageData
 {
 	FGlobalImports Imports;
@@ -1210,69 +1489,144 @@ static void FindImport(
 	FGlobalImports& GlobalImports,
 	TArray<FString>& TempFullNames,
 	FObjectImport* ImportMap,
-	const int32 LocalImportIndex)
+	const int32 LocalImportIndex,
+	const FPackageMap& PackageMap,
+	FPackage*& CurrentPackage)
 {
 	FObjectImport* Import = &ImportMap[LocalImportIndex];
 	FString& FullName = TempFullNames[LocalImportIndex];
+
+	if (Import->OuterIndex.IsNull())
+	{
+		CurrentPackage = PackageMap.FindRef(Import->ObjectName);
+	}
+
 	if (FullName.Len() == 0)
 	{
 		if (Import->OuterIndex.IsNull())
 		{
 			Import->ObjectName.AppendString(FullName);
-			const int32* FindGlobalImport = GlobalImports.ObjectsByFullName.Find(FullName);
-			if (!FindGlobalImport)
-			{
-				// first time, assign global index for this root package
-				const int32 GlobalImportIndex = GlobalImports.Objects.Num();
-				GlobalImports.ObjectsByFullName.Add(FullName, GlobalImportIndex);
-				FImportData& GlobalImport = GlobalImports.Objects.AddDefaulted_GetRef();
-				GlobalImport.GlobalIndex = GlobalImportIndex;
-				GlobalImport.OutermostIndex = GlobalImportIndex;
-				GlobalImport.OuterIndex = -1;
-				GlobalImport.ObjectName = Import->ObjectName;
-				GlobalImport.bIsPackage = true;
-				GlobalImport.bIsScript = FullName.StartsWith(TEXT("/Script/"));
-				GlobalImport.bIsLocalized = FullName.Contains(L10NPrefix);
-				GlobalImport.FullName = FullName;
-				GlobalImport.RefCount = 1;
-			}
-			else
+
+			const bool bIsScript = FullName.StartsWith(TEXT("/Script/"));
+			if (bIsScript)
 			{
-				++GlobalImports.Objects[*FindGlobalImport].RefCount;
+				FScriptImportData* ScriptImport;
+				FPackageObjectIndex FindGlobalImportIndex = GlobalImports.ObjectsByFullName.FindRef(FullName);
+				if (FindGlobalImportIndex.IsNull())
+				{
+					// assign global index for this script UPackage
+					FPackageObjectIndex::Type IndexType = FPackageObjectIndex::ScriptImport;
+					TArray<FScriptImportData>& ScriptImports = GlobalImports.ScriptImports;
+					FPackageObjectIndex GlobalImportIndex(IndexType, ScriptImports.Num());
+					GlobalImports.ObjectsByFullName.Add(FullName, GlobalImportIndex);
+
+					ScriptImport = &ScriptImports.AddDefaulted_GetRef();
+					ScriptImport->GlobalIndex = GlobalImportIndex;
+					ScriptImport->FullName = FullName;
+				}
+				else
+				{
+					ScriptImport = &GlobalImports.ScriptImports[FindGlobalImportIndex.GetIndex()];
+				}
+				if (!ScriptImport->bInitialized)
+				{
+					ScriptImport->OuterIndex = FPackageObjectIndex();
+					ScriptImport->ObjectName = Import->ObjectName;
+					ScriptImport->bInitialized = true;
+				}
 			}
 		}
 		else
 		{
 			const int32 LocalOuterIndex = Import->OuterIndex.ToImport();
-			FindImport(GlobalImports, TempFullNames, ImportMap, LocalOuterIndex);
-			FString& OuterName = TempFullNames[LocalOuterIndex];
-			ensure(OuterName.Len() > 0);
+			FindImport(GlobalImports, TempFullNames, ImportMap, LocalOuterIndex, PackageMap, CurrentPackage);
+			const FString& OuterName = TempFullNames[LocalOuterIndex];
+			check(OuterName.Len() > 0);
 
 			FullName.Append(OuterName);
 			FullName.AppendChar(TEXT('/'));
 			Import->ObjectName.AppendString(FullName);
-			const int32* FindGlobalImport = GlobalImports.ObjectsByFullName.Find(FullName);
-			if (!FindGlobalImport)
-			{
-				// first time, assign global index for this intermediate import
-				const int32 GlobalImportIndex = GlobalImports.Objects.Num();
-				GlobalImports.ObjectsByFullName.Add(FullName, GlobalImportIndex);
-				FImportData& GlobalImport = GlobalImports.Objects.AddDefaulted_GetRef();
-				const int32* FindOuterGlobalImport = GlobalImports.ObjectsByFullName.Find(OuterName);
-				check(FindOuterGlobalImport);
-				const FImportData& OuterGlobalImport = GlobalImports.Objects[*FindOuterGlobalImport];
-				GlobalImport.GlobalIndex = GlobalImportIndex;
-				GlobalImport.OutermostIndex = OuterGlobalImport.OutermostIndex;
-				GlobalImport.OuterIndex = OuterGlobalImport.GlobalIndex;
-				GlobalImport.ObjectName = Import->ObjectName;
-				GlobalImport.bIsScript = OuterGlobalImport.bIsScript;
-				GlobalImport.bIsLocalized = OuterGlobalImport.bIsLocalized;
-				GlobalImport.FullName = FullName;
-				GlobalImport.RefCount = 1;
+
+			const bool bIsScript = FullName.StartsWith(TEXT("/Script/"));
+			FPackageObjectIndex FindGlobalImportIndex = GlobalImports.ObjectsByFullName.FindRef(FullName);
+			FPackageObjectIndex FindOuterGlobalImport = GlobalImports.ObjectsByFullName.FindRef(OuterName);
+			if (bIsScript)
+			{
+				check(FindOuterGlobalImport.IsScriptImport());
+
+				FScriptImportData* ScriptImport;
+				if (FindGlobalImportIndex.IsNull())
+				{
+					// assign global index for this script UObject
+					TArray<FScriptImportData>& ScriptImports = GlobalImports.ScriptImports;
+					FPackageObjectIndex GlobalImportIndex(FPackageObjectIndex::ScriptImport, ScriptImports.Num());
+					GlobalImports.ObjectsByFullName.Add(FullName, GlobalImportIndex);
+
+					ScriptImport = &GlobalImports.ScriptImports.AddDefaulted_GetRef();
+					ScriptImport->GlobalIndex = GlobalImportIndex;
+					ScriptImport->FullName = FullName;
+				}
+				else
+				{
+					ScriptImport = &GlobalImports.ScriptImports[FindGlobalImportIndex.GetIndex()];
+				}
+				if (!ScriptImport->bInitialized)
+				{
+					ScriptImport->OuterIndex = FindOuterGlobalImport;
+					ScriptImport->ObjectName = Import->ObjectName;
+
+					const FScriptImportData& OuterScriptImport = GlobalImports.ScriptImports[FindOuterGlobalImport.GetIndex()];
+					if (OuterScriptImport.CDOClassFullName.Len() > 0)
+					{
+						ScriptImport->CDOClassFullName = OuterScriptImport.CDOClassFullName;
+					}
+					else
+					{
+						TCHAR NameBuffer[FName::StringBufferSize];
+						uint32 Len = Import->ObjectName.ToString(NameBuffer);
+						if (FCString::Strncmp(NameBuffer, TEXT("Default__"), 9) == 0)
+						{
+							ScriptImport->CDOClassFullName.Append(OuterName);
+							ScriptImport->CDOClassFullName.AppendChar(TEXT('/'));
+							ScriptImport->CDOClassFullName.AppendChars(NameBuffer + 9, Len - 9);
+						}
+					}
+					ScriptImport->bInitialized = true;
+				}
 			}
 			else
 			{
-				++GlobalImports.Objects[*FindGlobalImport].RefCount;
+				FPackageImportData* PackageImport;
+				if (FindGlobalImportIndex.IsNull())
+				{
+					TArray<FPackageImportData>& PackageImports = GlobalImports.PackageImports;
+					FPackageObjectIndex GlobalImportIndex(FPackageObjectIndex::PackageImport, PackageImports.Num());
+					GlobalImports.ObjectsByFullName.Add(FullName, GlobalImportIndex);
+
+					PackageImport = &GlobalImports.PackageImports.AddDefaulted_GetRef();
+					PackageImport->GlobalIndex = GlobalImportIndex;
+					PackageImport->FullName = FullName;
+				}
+				else
+				{
+					PackageImport = &GlobalImports.PackageImports[FindGlobalImportIndex.GetIndex()];
+				}
+				if (!PackageImport->bInitialized)
+				{
+					PackageImport->ObjectName = Import->ObjectName;
+					PackageImport->bIsLocalized = FullName.Contains(L10NPrefix);
+
+					if (FindOuterGlobalImport.IsPackageImport())
+					{
+						const FPackageImportData& OuterPackageImport = GlobalImports.PackageImports[FindOuterGlobalImport.GetIndex()];
+						PackageImport->Package = OuterPackageImport.Package;
+					}
+					else
+					{
+						PackageImport->Package = CurrentPackage;
+					}
+					PackageImport->bInitialized = true;
+				}
 			}
 		}
 	}
@@ -1283,7 +1637,7 @@ static void FindExport(
 	TArray<FString>& TempFullNames,
 	const FObjectExport* ExportMap,
 	const int32 LocalExportIndex,
-	const FName& PackageName)
+	FPackage* Package)
 {
 	const FObjectExport* Export = ExportMap + LocalExportIndex;
 	FString& FullName = TempFullNames[LocalExportIndex];
@@ -1292,7 +1646,7 @@ static void FindExport(
 	{
 		if (Export->OuterIndex.IsNull())
 		{
-			PackageName.AppendString(FullName);
+			Package->Name.AppendString(FullName);
 			FullName.AppendChar(TEXT('/'));
 			Export->ObjectName.AppendString(FullName);
 		}
@@ -1300,10 +1654,9 @@ static void FindExport(
 		{
 			check(Export->OuterIndex.IsExport());
 
-			FindExport(GlobalExports, TempFullNames, ExportMap, Export->OuterIndex.ToExport(), PackageName);
+			FindExport(GlobalExports, TempFullNames, ExportMap, Export->OuterIndex.ToExport(), Package);
 			FString& OuterName = TempFullNames[Export->OuterIndex.ToExport()];
 			check(OuterName.Len() > 0);
-
 			FullName.Append(OuterName);
 			FullName.AppendChar(TEXT('/'));
 			Export->ObjectName.AppendString(FullName);
@@ -1313,7 +1666,7 @@ static void FindExport(
 		GlobalExports.ObjectsByFullName.Add(FullName, GlobalExportIndex);
 		FExportData& ExportData = GlobalExports.Objects.AddDefaulted_GetRef();
 		ExportData.GlobalIndex = GlobalExportIndex;
-		ExportData.SourcePackageName = PackageName;
+		ExportData.Package = Package;
 		ExportData.ObjectName = Export->ObjectName;
 		ExportData.SourceIndex = LocalExportIndex;
 		ExportData.FullName = FullName;
@@ -1323,7 +1676,8 @@ static void FindExport(
 FPackage* FindOrAddPackage(
 	const TCHAR* RelativeFileName,
 	TArray<FPackage*>& Packages,
-	FPackageMap& PackageMap)
+	FPackageMap& PackageMap,
+	FPackageGlobalIdMap& PackageGlobalIdMap)
 {
 	FString PackageName;
 	FString ErrorMessage;
@@ -1341,7 +1695,16 @@ FPackage* FindOrAddPackage(
 		Package = new FPackage();
 		Package->Name = PackageFName;
 		Package->SourcePackageName = *RemapLocalizationPathIfNeeded(PackageName, &Package->Region);
-		Package->GlobalPackageId = Packages.Num();
+		const FPackageId* FindPackageGlobalId = PackageGlobalIdMap.Find(PackageFName);
+		if (FindPackageGlobalId)
+		{
+			Package->GlobalPackageId = *FindPackageGlobalId;
+		}
+		else
+		{
+			Package->GlobalPackageId = FPackageId::FromIndex(PackageGlobalIdMap.Num());
+			PackageGlobalIdMap.Add(PackageFName, Package->GlobalPackageId);
+		}
 		Packages.Add(Package);
 		PackageMap.Add(PackageFName, Package);
 	}
@@ -1351,7 +1714,7 @@ FPackage* FindOrAddPackage(
 
 static bool ConformLocalizedPackage(
 	const FPackageMap& PackageMap,
-	const TArray<FImportData>& GlobalImports,
+	const FGlobalImports& GlobalImports,
 	const FPackage& SourcePackage,
 	FPackage& LocalizedPackage,
 	TArray<FExportData>& GlobalExports,
@@ -1366,9 +1729,9 @@ static bool ConformLocalizedPackage(
 		TEXT("For culture '%s': Localized package '%s' (%d) for source package '%s' (%d)  - Has ExportCount %d vs. %d"),
 			*LocalizedPackage.Region,
 			*LocalizedPackage.Name.ToString(),
-			LocalizedPackage.GlobalPackageId,
+			LocalizedPackage.GlobalPackageId.ToIndexForDebugging(),
 			*LocalizedPackage.SourcePackageName.ToString(),
-			SourcePackage.GlobalPackageId,
+			SourcePackage.GlobalPackageId.ToIndexForDebugging(),
 			LocalizedPackage.ExportCount,
 			SourcePackage.ExportCount);
 
@@ -1407,11 +1770,40 @@ static bool ConformLocalizedPackage(
 		return nullptr;
 	};
 
+	auto AppendMismatchMessage = [&GlobalImports, &GlobalExports, &LocalizedPackage, &SourcePackage](
+		const TCHAR* Text, FName ExportName, FPackageObjectIndex LocIndex, FPackageObjectIndex SrcIndex, FString& FailReason)
+	{
+		FString LocString = 
+			LocIndex.IsPackageImport() ?
+			GlobalImports.PackageImports[LocIndex.GetIndex()].ObjectName.ToString() :
+			LocIndex.IsScriptImport() ?
+			GlobalImports.ScriptImports[LocIndex.GetIndex()].ObjectName.ToString() :
+			LocIndex.IsExport() ?
+			GlobalExports[LocalizedPackage.Exports[LocIndex.GetIndex()]].ObjectName.ToString() :
+			TEXT("");
+		FString SrcString = 
+			SrcIndex.IsPackageImport() ?
+			GlobalImports.PackageImports[SrcIndex.GetIndex()].ObjectName.ToString() :
+			SrcIndex.IsScriptImport() ?
+			GlobalImports.ScriptImports[SrcIndex.GetIndex()].ObjectName.ToString() :
+			SrcIndex.IsExport() ?
+			GlobalExports[LocalizedPackage.Exports[SrcIndex.GetIndex()]].ObjectName.ToString() :
+			TEXT("");
+
+		FailReason.Appendf(TEXT("Public export '%s' has %s %s (%d) vs. %s (%d)"),
+			*ExportName.ToString(),
+			Text,
+			*LocString,
+			LocIndex.GetIndex(),
+			*SrcString,
+			SrcIndex.GetIndex());
+	};
+
 	const int32 LocalizedPackageNameLen = LocalizedPackage.Name.GetStringLength();
 	const int32 SourcePackageNameLen = SourcePackage.Name.GetStringLength();
 
-	TArray <TPair<int32, int32>, TInlineAllocator<64>> MatchingPublicExports;
-	MatchingPublicExports.Reserve(ExportCount);
+	TArray <TPair<int32, int32>, TInlineAllocator<64>> NewPublicExports;
+	NewPublicExports.Reserve(ExportCount);
 
 	bool bSuccess = true;
 	int32 LocalizedIndex = 0;
@@ -1433,9 +1825,9 @@ static bool ConformLocalizedPackage(
 				TEXT("Culture '%s': Localized package '%s' (%d) for source package '%s' (%d) - Has some bad data from an earlier phase."),
 				*LocalizedPackage.Region,
 				*LocalizedPackage.Name.ToString(),
-				LocalizedPackage.GlobalPackageId,
+				LocalizedPackage.GlobalPackageId.ToIndexForDebugging(),
 				*LocalizedPackage.SourcePackageName.ToString(),
-				SourcePackage.GlobalPackageId)
+				SourcePackage.GlobalPackageId.ToIndexForDebugging())
 			return false;
 		}
 
@@ -1444,17 +1836,17 @@ static bool ConformLocalizedPackage(
 		{
 			++LocalizedIndex;
 
-			if (LocalizedExportData.GlobalImportIndex != -1)
+			if (LocalizedExportData.GlobalImportIndex.IsImport())
 			{
-				FailReason.Appendf(TEXT("Public localized export '%s' is missing in the source package"),
-					*LocalizedExportData.ObjectName.ToString());
+				// public localized export is missing in the source package, so just keep it as it is
+				NewPublicExports.Emplace(LocalizedIndex - 1, 1);
 			}
 		}
 		else if (CompareResult > 0)
 		{
 			++SourceIndex;
 
-			if (SourceExportData.GlobalImportIndex != -1)
+			if (SourceExportData.GlobalImportIndex.IsImport())
 			{
 				FailReason.Appendf(TEXT("Public source export '%s' is missing in the localized package"),
 					*SourceExportData.ObjectName.ToString());
@@ -1465,51 +1857,30 @@ static bool ConformLocalizedPackage(
 			++LocalizedIndex;
 			++SourceIndex;
 
-		if (SourceExportData.GlobalImportIndex != -1)
-		{
-				if (LocalizedExportData.ClassIndex != SourceExportData.ClassIndex)
+			if (SourceExportData.GlobalImportIndex.IsImport())
 			{
-					const FImportData& LocalizedImportData = GlobalImports[LocalizedExportData.ClassIndex.ToImport()];
-					const FImportData& SourceImportData = GlobalImports[SourceExportData.ClassIndex.ToImport()];
-
-					FailReason.Appendf(TEXT("Public export '%s' has class %s (%d) vs. %s (%d)"),
-					*LocalizedExportData.ObjectName.ToString(),
-						*LocalizedImportData.ObjectName.ToString(),
-						LocalizedExportData.ClassIndex.ForDebugging(),
-						*SourceImportData.ObjectName.ToString(),
-						SourceExportData.ClassIndex.ForDebugging());
-			}
+				if (LocalizedExportData.ClassIndex != SourceExportData.ClassIndex)
+				{
+					AppendMismatchMessage(TEXT("class"), LocalizedExportData.ObjectName,
+						LocalizedExportData.ClassIndex, SourceExportData.ClassIndex, FailReason);
+				}
 				else if (LocalizedExportData.TemplateIndex != SourceExportData.TemplateIndex)
-			{
-					const FImportData& LocalizedImportData = GlobalImports[LocalizedExportData.TemplateIndex.ToImport()];
-					const FImportData& SourceImportData = GlobalImports[SourceExportData.TemplateIndex.ToImport()];
-
-					FailReason.Appendf(TEXT("Public export '%s' has template %s (%d) vs. %s (%d)"),
-						*LocalizedExportData.ObjectName.ToString(),
-					*LocalizedImportData.ObjectName.ToString(),
-					LocalizedExportData.ClassIndex.ForDebugging(),
-					*SourceImportData.ObjectName.ToString(),
-					SourceExportData.ClassIndex.ForDebugging());
-			}
-			else if (LocalizedExportData.SuperIndex != SourceExportData.SuperIndex)
-			{
-				const FImportData& LocalizedImportData = GlobalImports[LocalizedExportData.ClassIndex.ToImport()];
-				const FImportData& SourceImportData = GlobalImports[SourceExportData.ClassIndex.ToImport()];
-
-					FailReason.Appendf(TEXT("Public export '%s' has super %s (%d) vs. %s (%d)"),
-						*LocalizedExportData.ObjectName.ToString(),
-					*LocalizedImportData.ObjectName.ToString(),
-					LocalizedExportData.ClassIndex.ForDebugging(),
-					*SourceImportData.ObjectName.ToString(),
-					SourceExportData.ClassIndex.ForDebugging());
+				{
+					AppendMismatchMessage(TEXT("template"), LocalizedExportData.ObjectName,
+						LocalizedExportData.TemplateIndex, SourceExportData.TemplateIndex, FailReason);
+				}
+				else if (LocalizedExportData.SuperIndex != SourceExportData.SuperIndex)
+				{
+					AppendMismatchMessage(TEXT("super"), LocalizedExportData.ObjectName,
+						LocalizedExportData.SuperIndex, SourceExportData.SuperIndex, FailReason);
+				}
+				else
+				{
+					NewPublicExports.Emplace(LocalizedIndex - 1, SourceIndex - 1);
+				}
 			}
-			else
+			else if (LocalizedExportData.GlobalImportIndex.IsImport())
 			{
-					MatchingPublicExports.Emplace(LocalizedIndex - 1, SourceIndex - 1);
-			}
-		}
-		else if (LocalizedExportData.GlobalImportIndex != -1)
-		{
 				FailReason.Appendf(TEXT("Public localized export '%s' exists in the source package")
 					TEXT(", but is not part of the source package interface (it is not imported by any other package)."),
 					*LocalizedExportData.ObjectName.ToString());
@@ -1522,9 +1893,9 @@ static bool ConformLocalizedPackage(
 				TEXT("Culture '%s': Localized package '%s' (%d) for '%s' (%d) - %s"),
 				*LocalizedPackage.Region,
 				*LocalizedPackage.Name.ToString(),
-				LocalizedPackage.GlobalPackageId,
+				LocalizedPackage.GlobalPackageId.ToIndexForDebugging(),
 				*LocalizedPackage.SourcePackageName.ToString(),
-				SourcePackage.GlobalPackageId,
+				SourcePackage.GlobalPackageId.ToIndexForDebugging(),
 				*FailReason);
 			bSuccess = false;
 		}
@@ -1532,229 +1903,403 @@ static bool ConformLocalizedPackage(
 
 	if (bSuccess)
 	{
-		for (TPair<int32, int32>& Pair : MatchingPublicExports)
+		LocalizedPackage.PublicExports.Reset();
+		for (TPair<int32, int32>& Pair : NewPublicExports)
 		{
 			FExportData& LocalizedExportData = GlobalExports[LocalizedPackage.Exports[Pair.Key]];
-			const FExportData& SourceExportData = GlobalExports[SourcePackage.Exports[Pair.Value]];
-
-			LocalizedToSourceImportIndexMap.Add(
-				LocalizedExportData.GlobalImportIndex,
-				SourceExportData.GlobalImportIndex);
+			if (Pair.Value != -1)
+			{
+				const FExportData& SourceExportData = GlobalExports[SourcePackage.Exports[Pair.Value]];
 
-			LocalizedExportData.GlobalImportIndex = SourceExportData.GlobalImportIndex;
+				LocalizedToSourceImportIndexMap.Add(
+					LocalizedExportData.GlobalImportIndex,
+					SourceExportData.GlobalImportIndex);
 
+				LocalizedExportData.GlobalImportIndex = SourceExportData.GlobalImportIndex;
+			}
+			LocalizedPackage.PublicExports.Add(LocalizedExportData.GlobalImportIndex);
 		}
-	LocalizedPackage.FirstGlobalImport = SourcePackage.FirstGlobalImport;
-	LocalizedPackage.GlobalImportCount = SourcePackage.GlobalImportCount;
 	}
 
 	return bSuccess;
 }
 
-void FinalizePackageHeader(
-	FPackage* Package,
+static void AddPreloadDependencies(
+	const FPackageAssetData& PackageAssetData,
+	const FGlobalPackageData& GlobalPackageData,
+	const FSourceToLocalizedPackageMultimap& SourceToLocalizedPackageMap,
+	FExportGraph& ExportGraph,
+	TArray<FPackage*>& Packages)
+{
+	IOSTORE_CPU_SCOPE(PreLoadDependencies);
+	UE_LOG(LogIoStore, Display, TEXT("Adding preload dependencies..."));
+
+	const FGlobalPackageImports& PackageImports = GlobalPackageData.Imports.PackageImports;
+	const TArray<FExportData>& GlobalExports = GlobalPackageData.Exports.Objects;
+
+	TArray<FPackage*> LocalizedPackages;
+	for (FPackage* Package : Packages)
+	{
+		// Convert PreloadDependencies to arcs
+		for (int32 I = 0; I < Package->ExportCount; ++I)
+		{
+			const FObjectExport& ObjectExport = PackageAssetData.ObjectExports[Package->ExportIndexOffset + I];
+			int32 PreloadDependenciesBaseIndex = Package->PreloadIndexOffset;
+
+			FPackageIndex ExportPackageIndex = FPackageIndex::FromExport(I);
+
+			auto AddPreloadArc = [&](FPackageIndex Dep, EPreloadDependencyType PhaseFrom, EPreloadDependencyType PhaseTo)
+			{
+				if (Dep.IsExport())
+				{
+					AddInternalExportArc(ExportGraph, *Package, Dep.ToExport(), PhaseFrom, I, PhaseTo);
+				}
+				else
+				{
+					FPackageObjectIndex ImportIndex = Package->Imports[Dep.ToImport()];
+					if (ImportIndex.IsPackageImport())
+					{
+						const FPackageImportData& Import = PackageImports[ImportIndex.GetIndex()];
+						check(Import.GlobalIndex == ImportIndex);
+						if (Import.GlobalExportIndex != -1)
+						{
+							const FExportData& Export = GlobalExports[Import.GlobalExportIndex];
+
+							AddExternalExportArc(ExportGraph, *Export.Package, Export.SourceIndex, PhaseFrom, *Package, I, PhaseTo);
+							Package->ImportedPreloadPackages.Add(Export.Package);
+
+							LocalizedPackages.Reset();
+							SourceToLocalizedPackageMap.MultiFind(Export.Package, LocalizedPackages);
+							for (FPackage* LocalizedPackage : LocalizedPackages)
+							{
+								UE_LOG(LogIoStore, Verbose, TEXT("For package '%s' (%d): Adding localized preload dependency '%s' in '%s'"),
+									*Package->Name.ToString(),
+									Package->GlobalPackageId.ToIndexForDebugging(),
+									*Export.ObjectName.ToString(),
+									*LocalizedPackage->Name.ToString());
+
+								AddExternalExportArc(ExportGraph, *LocalizedPackage, Export.SourceIndex, PhaseFrom, *Package, I, PhaseTo);
+								Package->ImportedPreloadPackages.Add(LocalizedPackage);
+							}
+						}
+						else
+						{
+							UE_LOG(LogIoStore, Verbose, TEXT("Skipping export arc to '%s' due to missing import"), *Package->Name.ToString());
+						}
+					}
+				}
+			};
+
+			if (PreloadDependenciesBaseIndex >= 0 && ObjectExport.FirstExportDependency >= 0)
+			{
+				int32 RunningIndex = PreloadDependenciesBaseIndex + ObjectExport.FirstExportDependency;
+				for (int32 Index = ObjectExport.SerializationBeforeSerializationDependencies; Index > 0; Index--)
+				{
+					FPackageIndex Dep = PackageAssetData.PreloadDependencies[RunningIndex++];
+					check(!Dep.IsNull());
+					AddPreloadArc(Dep, PreloadDependencyType_Serialize, PreloadDependencyType_Serialize);
+				}
+
+				for (int32 Index = ObjectExport.CreateBeforeSerializationDependencies; Index > 0; Index--)
+				{
+					FPackageIndex Dep = PackageAssetData.PreloadDependencies[RunningIndex++];
+					check(!Dep.IsNull());
+					AddPreloadArc(Dep, PreloadDependencyType_Create, PreloadDependencyType_Serialize);
+				}
+
+				for (int32 Index = ObjectExport.SerializationBeforeCreateDependencies; Index > 0; Index--)
+				{
+					FPackageIndex Dep = PackageAssetData.PreloadDependencies[RunningIndex++];
+					check(!Dep.IsNull());
+					AddPreloadArc(Dep, PreloadDependencyType_Serialize, PreloadDependencyType_Create);
+				}
+
+				for (int32 Index = ObjectExport.CreateBeforeCreateDependencies; Index > 0; Index--)
+				{
+					FPackageIndex Dep = PackageAssetData.PreloadDependencies[RunningIndex++];
+					check(!Dep.IsNull());
+					// can't create this export until these things are created
+					AddPreloadArc(Dep, PreloadDependencyType_Create, PreloadDependencyType_Create);
+				}
+			}
+		}
+	}
+};
+
+void FinalizePackageHeaders(
+	const TArray<FPackage*>& Packages,
 	const FNameMapBuilder& NameMapBuilder,
 	const TArray<FObjectExport>& ObjectExports,
 	const TArray<FExportData>& GlobalExports,
-	const TMap<FString, int32>& GlobalImportsByFullName,
-	FExportBundleMetaEntry* ExportBundleMetaEntries)
+	const FImportObjectsByFullName& GlobalImportsByFullName)
 {
-	// Temporary Archive for serializing ImportMap
-	FBufferWriter ImportMapArchive(nullptr, 0, EBufferWriterFlags::AllowResize | EBufferWriterFlags::TakeOwnership);
-	for (int32 GlobalImportIndex : Package->Imports)
+	IOSTORE_CPU_SCOPE(FinalizePackageHeaders);
+	UE_LOG(LogIoStore, Display, TEXT("Finalizing package headers..."));
+
+	for (FPackage* Package : Packages)
 	{
-		ImportMapArchive << GlobalImportIndex;
-	}
-	Package->ImportMapSize = ImportMapArchive.Tell();
+		// Temporary Archive for serializing ImportMap
+		FBufferWriter ImportMapArchive(nullptr, 0, EBufferWriterFlags::AllowResize | EBufferWriterFlags::TakeOwnership);
+		for (FPackageObjectIndex& GlobalImportIndex : Package->Imports)
+		{
+			ImportMapArchive << GlobalImportIndex;
+		}
+		Package->ImportMapSize = ImportMapArchive.Tell();
 
-	// Temporary Archive for serializing EDL graph data
-	FBufferWriter GraphArchive(nullptr, 0, EBufferWriterFlags::AllowResize | EBufferWriterFlags::TakeOwnership);
+		// Temporary Archive for serializing EDL graph data
+		FBufferWriter GraphArchive(nullptr, 0, EBufferWriterFlags::AllowResize | EBufferWriterFlags::TakeOwnership);
 
-	int32 InternalArcCount = Package->InternalArcs.Num();
-	GraphArchive << InternalArcCount;
-	GraphArchive.Serialize(Package->InternalArcs.GetData(), Package->InternalArcs.Num() * sizeof(FArc));
+		int32 InternalArcCount = Package->InternalArcs.Num();
+		GraphArchive << InternalArcCount;
+		GraphArchive.Serialize(Package->InternalArcs.GetData(), Package->InternalArcs.Num() * sizeof(FArc));
 
-	int32 ReferencedPackagesCount = Package->ExternalArcs.Num();
-	GraphArchive << ReferencedPackagesCount;
-	for (auto& KV : Package->ExternalArcs)
-	{
-		FPackage* ImportedPackage = KV.Key;
-		TArray<FArc>& Arcs = KV.Value;
-		int32 ExternalArcCount = Arcs.Num();
+		int32 ReferencedPackagesCount = Package->ExternalArcs.Num();
+		GraphArchive << ReferencedPackagesCount;
+		for (auto& KV : Package->ExternalArcs)
+		{
+			FPackage* ImportedPackage = KV.Key;
+			TArray<FArc>& Arcs = KV.Value;
+			int32 ExternalArcCount = Arcs.Num();
 
-		GraphArchive << ImportedPackage->GlobalPackageId;
-		GraphArchive << ExternalArcCount;
-		GraphArchive.Serialize(Arcs.GetData(), ExternalArcCount * sizeof(FArc));
-	}
-	Package->UGraphSize = GraphArchive.Tell();
+			GraphArchive << ImportedPackage->GlobalPackageId;
+			GraphArchive << ExternalArcCount;
+			GraphArchive.Serialize(Arcs.GetData(), ExternalArcCount * sizeof(FArc));
+		}
+		Package->UGraphSize = GraphArchive.Tell();
 
-	// Temporary Archive for serializing export map data
-	FBufferWriter ExportMapArchive(nullptr, 0, EBufferWriterFlags::AllowResize | EBufferWriterFlags::TakeOwnership);
-	for (int32 I = 0; I < Package->ExportCount; ++I)
-	{
-		const FObjectExport& ObjectExport = ObjectExports[Package->ExportIndexOffset + I];
-		const FExportData& ExportData = GlobalExports[Package->Exports[I]];
+		// Temporary Archive for serializing export map data
+		FBufferWriter ExportMapArchive(nullptr, 0, EBufferWriterFlags::AllowResize | EBufferWriterFlags::TakeOwnership);
+		for (int32 I = 0; I < Package->ExportCount; ++I)
+		{
+			const FObjectExport& ObjectExport = ObjectExports[Package->ExportIndexOffset + I];
+			const FExportData& ExportData = GlobalExports[Package->Exports[I]];
+
+			FExportMapEntry ExportMapEntry;
+			ExportMapEntry.CookedSerialOffset = ObjectExport.SerialOffset;
+			ExportMapEntry.CookedSerialSize = ObjectExport.SerialSize;
+			ExportMapEntry.ObjectName.NameEntryIndex = NameMapBuilder.MapName(ObjectExport.ObjectName);
+			ExportMapEntry.ObjectName.Number = ObjectExport.ObjectName.GetNumber();
+			ExportMapEntry.OuterIndex = ExportData.OuterIndex;
+			ExportMapEntry.ClassIndex = ExportData.ClassIndex;
+			ExportMapEntry.SuperIndex = ExportData.SuperIndex;
+			ExportMapEntry.TemplateIndex = ExportData.TemplateIndex;
+			ExportMapEntry.GlobalImportIndex = ExportData.GlobalImportIndex;
+			ExportMapEntry.ObjectFlags = ObjectExport.ObjectFlags;
+			ExportMapEntry.FilterFlags = EExportFilterFlags::None;
+			if (ObjectExport.bNotForClient)
+			{
+				ExportMapEntry.FilterFlags = EExportFilterFlags::NotForClient;
+			}
+			else if (ObjectExport.bNotForServer)
+			{
+				ExportMapEntry.FilterFlags = EExportFilterFlags::NotForServer;
+			}
+
+			ExportMapArchive << ExportMapEntry;
+		}
+		Package->ExportMapSize = ExportMapArchive.Tell();
 
-		FExportMapEntry ExportMapEntry;
-		ExportMapEntry.SerialSize = ObjectExport.SerialSize;
-		ExportMapEntry.ObjectName[0] = NameMapBuilder.MapName(ObjectExport.ObjectName);
-		ExportMapEntry.ObjectName[1] = ObjectExport.ObjectName.GetNumber();
-		ExportMapEntry.OuterIndex = ExportData.OuterIndex;
-		ExportMapEntry.ClassIndex = ExportData.ClassIndex;
-		ExportMapEntry.SuperIndex = ExportData.SuperIndex;
-		ExportMapEntry.TemplateIndex = ExportData.TemplateIndex;
-		ExportMapEntry.GlobalImportIndex = ExportData.GlobalImportIndex;
-		ExportMapEntry.ObjectFlags = ObjectExport.ObjectFlags;
-		ExportMapEntry.FilterFlags = EExportFilterFlags::None;
-		if (ObjectExport.bNotForClient)
+		// Temporary archive for serializing export bundle data
+		FBufferWriter ExportBundlesArchive(nullptr, 0, EBufferWriterFlags::AllowResize | EBufferWriterFlags::TakeOwnership);
+		uint32 ExportBundleEntryIndex = 0;
+		for (FExportBundle& ExportBundle : Package->ExportBundles)
 		{
-			ExportMapEntry.FilterFlags = EExportFilterFlags::NotForClient;
+			const uint32 EntryCount = ExportBundle.Nodes.Num();
+			FExportBundleHeader ExportBundleHeader { ExportBundleEntryIndex, EntryCount };
+			ExportBundlesArchive << ExportBundleHeader ;
+
+			ExportBundleEntryIndex += EntryCount; 
 		}
-		else if (ObjectExport.bNotForServer)
+		for (FExportBundle& ExportBundle : Package->ExportBundles)
 		{
-			ExportMapEntry.FilterFlags = EExportFilterFlags::NotForServer;
+			for (FExportGraphNode* ExportNode : ExportBundle.Nodes)
+			{
+				ExportBundlesArchive << ExportNode->BundleEntry;
+			}
 		}
+		Package->ExportBundlesHeaderSize = ExportBundlesArchive.Tell();
 
-		ExportMapArchive << ExportMapEntry;
-	}
-	Package->ExportMapSize = ExportMapArchive.Tell();
+		uint64 PackageHeaderSize =
+			sizeof(FPackageSummary)
+			+ Package->NameMapSize
+			+ Package->ImportMapSize
+			+ Package->ExportMapSize
+			+ Package->ExportBundlesHeaderSize
+			+ Package->UGraphSize;
 
-	// Temporary archive for serializing export bundle data
-	FBufferWriter ExportBundlesArchive(nullptr, 0, EBufferWriterFlags::AllowResize | EBufferWriterFlags::TakeOwnership);
-	uint32 ExportBundleEntryIndex = 0;
-	for (FExportBundle& ExportBundle : Package->ExportBundles)
-	{
-		const uint32 EntryCount = ExportBundle.Nodes.Num();
-		FExportBundleHeader ExportBundleHeader { ExportBundleEntryIndex, EntryCount };
-		ExportBundlesArchive << ExportBundleHeader ;
+		Package->HeaderData.AddZeroed(PackageHeaderSize);
+		uint8* PackageHeaderBuffer = Package->HeaderData.GetData();
+		FPackageSummary* PackageSummary = reinterpret_cast<FPackageSummary*>(PackageHeaderBuffer);
 
-		ExportBundleEntryIndex += EntryCount; 
-	}
-	for (FExportBundle& ExportBundle : Package->ExportBundles)
-	{
-		for (FExportGraphNode* ExportNode : ExportBundle.Nodes)
+		PackageSummary->PackageFlags = Package->PackageFlags;
+		PackageSummary->CookedHeaderSize = Package->CookedHeaderSize;
+		PackageSummary->GraphDataSize = Package->UGraphSize;
+
+		FBufferWriter SummaryArchive(PackageHeaderBuffer, PackageHeaderSize);
+		SummaryArchive.Seek(sizeof(FPackageSummary));
+
+		// NameMap data
 		{
-			ExportBundlesArchive << ExportNode->BundleEntry;
+			PackageSummary->NameMapOffset = SummaryArchive.Tell();
+			SummaryArchive.Serialize(Package->NameIndices.GetData(), Package->NameMapSize);
 		}
-	}
-	Package->ExportBundlesSize = ExportBundlesArchive.Tell();
 
-	Package->NameMapSize = Package->NameIndices.Num() * Package->NameIndices.GetTypeSize();
+		// ImportMap data
+		{
+			check(ImportMapArchive.Tell() == Package->ImportMapSize);
+			PackageSummary->ImportMapOffset = SummaryArchive.Tell();
+			SummaryArchive.Serialize(ImportMapArchive.GetWriterData(), ImportMapArchive.Tell());
+		}
 
-	uint64 PackageHeaderSize =
-		sizeof(FPackageSummary)
-		+ Package->NameMapSize
-		+ Package->ImportMapSize
-		+ Package->ExportMapSize
-		+ Package->ExportBundlesSize
-		+ Package->UGraphSize;
+		// ExportMap data
+		{
+			check(ExportMapArchive.Tell() == Package->ExportMapSize);
+			PackageSummary->ExportMapOffset = SummaryArchive.Tell();
+			SummaryArchive.Serialize(ExportMapArchive.GetWriterData(), ExportMapArchive.Tell());
+		}
 
-	Package->HeaderData.AddZeroed(PackageHeaderSize);
-	uint8* PackageHeaderBuffer = Package->HeaderData.GetData();
-	FPackageSummary* PackageSummary = reinterpret_cast<FPackageSummary*>(PackageHeaderBuffer);
+		// ExportBundle data
+		{
+			check(ExportBundlesArchive.Tell() == Package->ExportBundlesHeaderSize);
+			PackageSummary->ExportBundlesOffset = SummaryArchive.Tell();
+			SummaryArchive.Serialize(ExportBundlesArchive.GetWriterData(), ExportBundlesArchive.Tell());
+		}
 
-	PackageSummary->PackageFlags = Package->PackageFlags;
-	PackageSummary->GraphDataSize = Package->UGraphSize;
-	PackageSummary->BulkDataStartOffset = Package->BulkDataStartOffset;
-	const int32* FindGlobalImportIndexForPackage = GlobalImportsByFullName.Find(Package->Name.ToString());
-	PackageSummary->GlobalImportIndex = FindGlobalImportIndexForPackage ? *FindGlobalImportIndexForPackage : -1;
+		// Graph data
+		{
+			check(GraphArchive.Tell() == Package->UGraphSize);
+			PackageSummary->GraphDataOffset = SummaryArchive.Tell();
+			SummaryArchive.Serialize(GraphArchive.GetWriterData(), GraphArchive.Tell());
+		}
+	}
+}
 
-	FBufferWriter SummaryArchive(PackageHeaderBuffer, PackageHeaderSize);
-	SummaryArchive.Seek(sizeof(FPackageSummary));
+void SerializePackageStoreEntries(
+	TArray<FPackage*> Packages,
+	FNameMapBuilder& NameMapBuilder,
+	FArchive& StoreTocArchive,
+	FArchive& StoreDataArchive)
+{
+	IOSTORE_CPU_SCOPE(SerializePackageStoreEntries);
+	UE_LOG(LogIoStore, Display, TEXT("Finalizing package store..."));
 
-	// NameMap data
-	{
-		PackageSummary->NameMapOffset = SummaryArchive.Tell();
-		SummaryArchive.Serialize(Package->NameIndices.GetData(), Package->NameMapSize);
-	}
+	int32 StoreTocSize = Packages.Num() * sizeof(FPackageStoreEntry);
 
-	// ImportMap data
+	auto SerializePackageEntryCArrayHeader = [&StoreTocSize,&StoreTocArchive,&StoreDataArchive](int32 Count)
 	{
-		check(ImportMapArchive.Tell() == Package->ImportMapSize);
-		PackageSummary->ImportMapOffset = SummaryArchive.Tell();
-		SummaryArchive.Serialize(ImportMapArchive.GetWriterData(), ImportMapArchive.Tell());
-	}
+		const int32 RemainingTocSize = StoreTocSize - StoreTocArchive.Tell();
+		const int32 OffsetFromThis = RemainingTocSize + StoreDataArchive.Tell();
+		uint32 ArrayNum = Count > 0 ? Count : 0; 
+		uint32 OffsetToDataFromThis = ArrayNum > 0 ? OffsetFromThis : 0;
 
-	// ExportMap data
-	{
-		check(ExportMapArchive.Tell() == Package->ExportMapSize);
-		PackageSummary->ExportMapOffset = SummaryArchive.Tell();
-		SummaryArchive.Serialize(ExportMapArchive.GetWriterData(), ExportMapArchive.Tell());
-	}
+		StoreTocArchive << ArrayNum;
+		StoreTocArchive << OffsetToDataFromThis;
+	};
 
-	// ExportBundle data
-	{
-		check(ExportBundlesArchive.Tell() == Package->ExportBundlesSize);
-		PackageSummary->ExportBundlesOffset = SummaryArchive.Tell();
-		SummaryArchive.Serialize(ExportBundlesArchive.GetWriterData(), ExportBundlesArchive.Tell());
-	}
 
-	// Graph data
+	for (FPackage* Package : Packages)
 	{
-		check(GraphArchive.Tell() == Package->UGraphSize);
-		PackageSummary->GraphDataOffset = SummaryArchive.Tell();
-		SummaryArchive.Serialize(GraphArchive.GetWriterData(), GraphArchive.Tell());
-	}
+		NameMapBuilder.MarkNameAsReferenced(Package->Name);
+		NameMapBuilder.SerializeName(StoreTocArchive, Package->Name);
+		StoreTocArchive << Package->SourceGlobalPackageId;
+		StoreTocArchive << Package->ExportCount;
 
-	// Export bundle payload
-	{
-		check(Package->ExportBundles.Num());
+		// Global imported packages meta data
+		{
+			SerializePackageEntryCArrayHeader(Package->ImportedPackagesSerializeCount);
+			for (int32 I = 0; I < Package->ImportedPackagesSerializeCount; ++I)
+			{
+				FPackage* ImportedPackage = Package->ImportedPackages[I];
+				StoreDataArchive << ImportedPackage->GlobalPackageId;
+			}
+		}
 
-		for (int32 ExportBundleIndex = 0; ExportBundleIndex < Package->ExportBundles.Num(); ++ExportBundleIndex)
+		// Global public exports meta data
 		{
-			FExportBundle& ExportBundle = Package->ExportBundles[ExportBundleIndex];
-			FExportBundleMetaEntry* ExportBundleMetaEntry = ExportBundleMetaEntries + Package->FirstExportBundleMetaEntry + ExportBundleIndex;
-			ExportBundleMetaEntry->LoadOrder = ExportBundle.LoadOrder;
-			for (FExportGraphNode* Node : ExportBundle.Nodes)
+			SerializePackageEntryCArrayHeader(Package->PublicExports.Num());
+			for (FPackageObjectIndex ObjectIndex : Package->PublicExports)
 			{
-				if (Node->BundleEntry.CommandType == FExportBundleEntry::ExportCommandType_Serialize)
-				{
-					const FObjectExport& ObjectExport = ObjectExports[Package->ExportIndexOffset + Node->BundleEntry.LocalExportIndex];
-					Package->ExportsPayloadSize += ObjectExport.SerialSize;
-				}
+				StoreDataArchive << ObjectIndex;
+			}
+		}
+
+		// Global export bundle meta data
+		{
+			// TODO: Request IO for each package export bundle individually in the loader,
+			//       or just store the data for the first entry directly on the package entry itself
+			SerializePackageEntryCArrayHeader(Package->ExportBundles.Num());
+			FExportBundleMetaEntry ExportBundleMetaEntry;
+			ExportBundleMetaEntry.PayloadSize = Package->HeaderData.Num() + Package->ExportsSerialSize;
+			for (const FExportBundle& ExportBundle : Package->ExportBundles)
+			{
+				ExportBundleMetaEntry.LoadOrder = ExportBundle.LoadOrder;
+				StoreDataArchive << ExportBundleMetaEntry;
+
+				ExportBundleMetaEntry.PayloadSize = 0; // currently only specified for the first entry
 			}
 		}
-		FExportBundleMetaEntry* FirstExportBundleMetaEntry = ExportBundleMetaEntries + Package->FirstExportBundleMetaEntry;
-		FirstExportBundleMetaEntry->PayloadSize = Package->HeaderData.Num() + Package->ExportsPayloadSize;
 	}
 }
 
-void SerializePackageData(
-	FPackage* Package,
-	const TArray<FObjectExport>& ObjectExports,
-	FIoStoreWriter* IoStoreWriter)
-{
-	FString UExpFileName = FPaths::ChangeExtension(Package->FileName, TEXT(".uexp"));
-	TUniquePtr<FArchive> ExpAr(IFileManager::Get().CreateFileReader(*UExpFileName));
-	Package->UExpSize = ExpAr->TotalSize();
-#if !SKIP_WRITE_CONTAINER
-	uint8* ExportsBuffer = static_cast<uint8*>(FMemory::Malloc(ExpAr->TotalSize()));
-	ExpAr->Serialize(ExportsBuffer, ExpAr->TotalSize());
-	ExpAr->Close();
-	
-	uint64 BundleBufferSize = Package->HeaderData.Num() + Package->ExportsPayloadSize;
-	uint8* BundleBuffer = static_cast<uint8*>(FMemory::Malloc(BundleBufferSize));
-	FMemory::Memcpy(BundleBuffer, Package->HeaderData.GetData(), Package->HeaderData.Num());
+template<class T>
+static void SerializeImportNames(
+	const FNameMapBuilder& NameMapBuilder,
+	const TArray<T>& Imports,
+	FArchive& Ar)
+{
+	int32 Count = Imports.Num();
+	Ar << Count;
+	for (const T& ImportData : Imports)
+	{
+		NameMapBuilder.SerializeName(Ar, ImportData.ObjectName);
+	}
+};
+
+static void SerializeInitialLoad(
+	const FNameMapBuilder& NameMapBuilder,
+	const FGlobalScriptImports& GlobalScriptImports,
+	FArchive& InitialLoadArchive)
+{
+	IOSTORE_CPU_SCOPE(SerializeInitialLoad);
+	UE_LOG(LogIoStore, Display, TEXT("Finalizing initial load..."));
+
+	int32 NumScriptObjects = GlobalScriptImports.Num();
+	InitialLoadArchive << NumScriptObjects; 
+
+	for (const FScriptImportData& ImportData : GlobalScriptImports)
+	{
+		FScriptObjectEntry Entry;
+		Entry.ObjectName.Index = NameMapBuilder.MapName(ImportData.ObjectName);
+		Entry.ObjectName.Number = ImportData.ObjectName.GetNumber();
+		Entry.OuterIndex = ImportData.OuterIndex;
+		Entry.CDOClassIndex = ImportData.CDOClassIndex;
+
+		InitialLoadArchive << Entry; 
+	}
+};
+
+static FIoBuffer CreateExportBundleBuffer(const FPackage* Package, const TArray<FObjectExport>& ObjectExports, const FIoBuffer UExpBuffer)
+{
+	uint64 BundleBufferSize = Package->HeaderData.Num() + Package->ExportsSerialSize;
+	FIoBuffer BundleBuffer(BundleBufferSize);
+	FMemory::Memcpy(BundleBuffer.Data(), Package->HeaderData.GetData(), Package->HeaderData.Num());
 	uint64 BundleBufferOffset = Package->HeaderData.Num();
-	for (FExportBundle& ExportBundle : Package->ExportBundles)
+	for (const FExportBundle& ExportBundle : Package->ExportBundles)
 	{
-		for (FExportGraphNode* Node : ExportBundle.Nodes)
+		for (const FExportGraphNode* Node : ExportBundle.Nodes)
 		{
 			if (Node->BundleEntry.CommandType == FExportBundleEntry::ExportCommandType_Serialize)
 			{
 				const FObjectExport& ObjectExport = ObjectExports[Package->ExportIndexOffset + Node->BundleEntry.LocalExportIndex];
 				const int64 Offset = ObjectExport.SerialOffset - Package->UAssetSize;
-				FMemory::Memcpy(BundleBuffer + BundleBufferOffset, ExportsBuffer + Offset, ObjectExport.SerialSize);
+				check(uint64(Offset + ObjectExport.SerialSize) <= UExpBuffer.DataSize());
+				FMemory::Memcpy(BundleBuffer.Data() + BundleBufferOffset, UExpBuffer.Data() + Offset, ObjectExport.SerialSize);
 				BundleBufferOffset += ObjectExport.SerialSize;
 			}
 		}
 	}
-
-	FIoBuffer IoBuffer(FIoBuffer::Wrap, BundleBuffer, BundleBufferSize);
-	IoStoreWriter->Append(CreateChunkId(Package->GlobalPackageId, 0, EIoChunkType::ExportBundleData, *Package->FileName), IoBuffer, *FPaths::GetCleanFilename(Package->FileName));
-	FMemory::Free(BundleBuffer);
-	FMemory::Free(ExportsBuffer);
-#endif
+	check(BundleBufferOffset == BundleBuffer.DataSize());
+	return BundleBuffer;
 }
 
 static void ParsePackageAssets(
@@ -1765,105 +2310,119 @@ static void ParsePackageAssets(
 	FExportGraph& ExportGraph)
 {
 	IOSTORE_CPU_SCOPE(ParsePackageAssets);
+	UE_LOG(LogIoStore, Display, TEXT("Parsing packages..."));
 
 	TAtomic<int32> ReadCount {0};
 	TAtomic<int32> ParseCount {0};
 	const int32 TotalPackageCount = Packages.Num();
 
 	TArray<FPackageFileSummary> PackageFileSummaries;
-	TArray<TUniquePtr<uint8[]>> PackageAssetBuffers;
-
 	PackageFileSummaries.SetNum(TotalPackageCount);
+
+	uint8* UAssetMemory = nullptr;
+	TArray<uint8*> PackageAssetBuffers;
 	PackageAssetBuffers.SetNum(TotalPackageCount);
 
 	UE_LOG(LogIoStore, Display, TEXT("Reading package assets..."));
 	{
-		IOSTORE_CPU_SCOPE(ReadAssets);
+		IOSTORE_CPU_SCOPE(ReadUAssetFiles);
 
-		ParallelFor(TotalPackageCount,[
-			&ReadCount,
-			&PackageAssetBuffers,
-			&PackageFileSummaries,
-			&Packages,
-			&PackageAssetData](int32 Index)
+		uint64 TotalUAssetSize = 0;
+		for (const FPackage* Package : Packages)
 		{
-			TUniquePtr<uint8[]>& PackageBuffer = PackageAssetBuffers[Index];
-			FPackageFileSummary& Summary = PackageFileSummaries[Index];
-			FPackage& Package = *Packages[Index];
-
-			IOSTORE_CPU_SCOPE_DATA(ReadPackage, TCHAR_TO_ANSI(*Package.FileName));
+			TotalUAssetSize += Package->UAssetSize;
+		}
+		UAssetMemory = reinterpret_cast<uint8*>(FMemory::Malloc(TotalUAssetSize));
+		uint8* UAssetMemoryPtr = UAssetMemory;
+		for (int32 Index = 0; Index < TotalPackageCount; ++Index)
+		{
+			PackageAssetBuffers[Index] = UAssetMemoryPtr;
+			UAssetMemoryPtr += Packages[Index]->UAssetSize;
+		}
 
-			UE_CLOG(++ReadCount % 1000 == 0, LogIoStore, Display, TEXT("Reading %d/%d: '%s'"), ReadCount.Load(), Packages.Num(), *Package.FileName);
+		TAtomic<uint64> CurrentFileIndex{ 0 };
+		ParallelFor(TotalPackageCount, [&ReadCount, &PackageAssetBuffers, &Packages, &CurrentFileIndex](int32 Index)
+		{
+			TRACE_CPUPROFILER_EVENT_SCOPE(ReadUAssetFile);
+			FPackage* Package = Packages[Index];
+			uint8* Buffer = PackageAssetBuffers[Index];
+			IFileHandle* FileHandle = FPlatformFileManager::Get().GetPlatformFile().OpenRead(*Package->FileName);
+			check(FileHandle);
+			bool bSuccess = FileHandle->Read(Buffer, Package->UAssetSize);
+			check(bSuccess);
+			delete FileHandle;
+			uint64 LocalFileIndex = CurrentFileIndex.IncrementExchange() + 1;
+			UE_CLOG(LocalFileIndex % 1000 == 0, LogIoStore, Display, TEXT("Reading %d/%d: '%s'"), LocalFileIndex, Packages.Num(), *Package->FileName);
+		}, EParallelForFlags::Unbalanced);
+	}
 
-			{
-				TUniquePtr<FArchive> FileAr(IFileManager::Get().CreateFileReader(*Package.FileName));
-				check(FileAr);
-				Package.UAssetSize = FileAr->TotalSize();
+	{
+		IOSTORE_CPU_SCOPE(SerializeSummaries);
 
-				PackageBuffer = MakeUnique<uint8[]>(Package.UAssetSize);
-				FileAr->Serialize(PackageBuffer.Get(), Package.UAssetSize);
-				FileAr->Close();
-			}
+		ParallelFor(TotalPackageCount, [
+				&ReadCount,
+				&PackageAssetBuffers,
+				&PackageFileSummaries,
+				&Packages](int32 Index)
+		{
+			uint8* PackageBuffer = PackageAssetBuffers[Index];
+			FPackageFileSummary& Summary = PackageFileSummaries[Index];
+			FPackage& Package = *Packages[Index];
 
-			TArrayView<const uint8> MemView(PackageBuffer.Get(), Package.UAssetSize);
+			TArrayView<const uint8> MemView(PackageBuffer, Package.UAssetSize);
 			FMemoryReaderView Ar(MemView);
 			Ar << Summary;
 
 			Package.SummarySize = Ar.Tell();
 			Package.NameCount = Summary.NameCount;
 			Package.ImportCount = Summary.ImportCount;
+			Package.PreloadDependencyCount = Summary.PreloadDependencyCount;
 			Package.ExportCount = Summary.ExportCount;
 			Package.PackageFlags = Summary.PackageFlags;
-			Package.BulkDataStartOffset = Summary.BulkDataStartOffset;
+			Package.CookedHeaderSize = Summary.TotalHeaderSize;
 
-			if (Summary.ImportCount > 0)
-			{
-				FScopeLock _(&PackageAssetData.CriticalSection);
-				Package.ImportIndexOffset = PackageAssetData.ObjectImports.Num();
-				PackageAssetData.ObjectImports.AddUninitialized(Summary.ImportCount);
-			}
-			
-			if (Summary.PreloadDependencyCount > 0)
-			{
-				FScopeLock _(&PackageAssetData.CriticalSection);
-				Package.PreloadIndexOffset = PackageAssetData.PreloadDependencies.Num();
-				PackageAssetData.PreloadDependencies.AddUninitialized(Summary.PreloadDependencyCount);
-			}
+		}, EParallelForFlags::Unbalanced);
+	}
 
-			if (Summary.ExportCount > 0)
-			{
-				FScopeLock _(&PackageAssetData.CriticalSection);
-				Package.ExportIndexOffset = PackageAssetData.ObjectExports.Num();
-				PackageAssetData.ObjectExports.AddUninitialized(Summary.ExportCount);
-			}
+	int32 TotalImportCount = 0;
+	int32 TotalPreloadDependencyCount = 0;
+	int32 TotalExportCount = 0;
+	for (FPackage* Package : Packages)
+	{
+		if (Package->ImportCount > 0)
+		{
+			Package->ImportIndexOffset = TotalImportCount;
+			TotalImportCount += Package->ImportCount;
+		}
 
-		}, EParallelForFlags::Unbalanced);
+		if (Package->PreloadDependencyCount > 0)
+		{
+			Package->PreloadIndexOffset = TotalPreloadDependencyCount;
+			TotalPreloadDependencyCount += Package->PreloadDependencyCount;
+		}
+
+		if (Package->ExportCount > 0)
+		{
+			Package->ExportIndexOffset = TotalExportCount;
+			TotalExportCount += Package->ExportCount;
+		}
 	}
+	PackageAssetData.ObjectImports.AddUninitialized(TotalImportCount);
+	PackageAssetData.PreloadDependencies.AddUninitialized(TotalPreloadDependencyCount);
+	PackageAssetData.ObjectExports.AddUninitialized(TotalExportCount);
 
 	UE_LOG(LogIoStore, Display, TEXT("Parsing package assets..."));
 	{
 		IOSTORE_CPU_SCOPE(SerializeAssets);
 
-		ParallelFor(TotalPackageCount,[
-				&ParseCount,
-				&PackageAssetBuffers,
-				&PackageFileSummaries,
-				&Packages,
-				&NameMapBuilder,
-				&PackageAssetData,
-				&GlobalPackageData,
-				&ExportGraph](int32 Index)
+		for (int32 PackageIndex = 0; PackageIndex < TotalPackageCount; ++PackageIndex)
 		{
-			TUniquePtr<uint8[]>& PackageBuffer = PackageAssetBuffers[Index];
-			const FPackageFileSummary& Summary = PackageFileSummaries[Index];
-			FPackage& Package = *Packages[Index];
-			TArrayView<const uint8> MemView(PackageBuffer.Get(), Package.UAssetSize);
+			uint8* PackageBuffer = PackageAssetBuffers[PackageIndex];
+			const FPackageFileSummary& Summary = PackageFileSummaries[PackageIndex];
+			FPackage& Package = *Packages[PackageIndex];
+			TArrayView<const uint8> MemView(PackageBuffer, Package.UAssetSize);
 			FMemoryReaderView Ar(MemView);
 
-			IOSTORE_CPU_SCOPE_DATA(ParsePackage, TCHAR_TO_ANSI(*Package.FileName));
-
-			UE_CLOG(++ParseCount % 1000 == 0, LogIoStore, Display, TEXT("Parsing %d/%d: '%s'"), ParseCount.Load(), Packages.Num(), *Package.FileName);
-
 			if (Summary.NameCount > 0)
 			{
 				Ar.Seek(Summary.NameOffset);
@@ -1880,12 +2439,36 @@ static void ParsePackageAssets(
 					Names.Emplace(NameEntry);
 				}
 
-				NameMapBuilder.MarkNamesAsReferenced(Names, Package.NameMap, Package.NameIndices); 
+				NameMapBuilder.MarkNamesAsReferenced(Names, Package.NameMap, Package.NameIndices);
 			}
 
+			Package.NameMapSize = Package.NameIndices.Num() * Package.NameIndices.GetTypeSize();
+		}
+
+		ParallelFor(TotalPackageCount,[
+				&ParseCount,
+				&PackageAssetBuffers,
+				&PackageFileSummaries,
+				&Packages,
+				&NameMapBuilder,
+				&PackageAssetData,
+				&GlobalPackageData,
+				&ExportGraph](int32 Index)
+		{
+			uint8* PackageBuffer = PackageAssetBuffers[Index];
+			const FPackageFileSummary& Summary = PackageFileSummaries[Index];
+			FPackage& Package = *Packages[Index];
+			TArrayView<const uint8> MemView(PackageBuffer, Package.UAssetSize);
+			FMemoryReaderView Ar(MemView);
+
+			IOSTORE_CPU_SCOPE_DATA(ParsePackage, TCHAR_TO_ANSI(*Package.FileName));
+
+			const int32 Count = ParseCount.IncrementExchange();
+			UE_CLOG(Count % 1000 == 0, LogIoStore, Display, TEXT("Parsing %d/%d: '%s'"), Count, Packages.Num(), *Package.FileName);
+
 			if (Summary.ImportCount > 0)
 			{
-				FNameProxyArchive ProxyAr(Ar, Package.NameMap);
+				FNameReaderProxyArchive ProxyAr(Ar, Package.NameMap);
 				ProxyAr.Seek(Summary.ImportOffset);
 
 				for (int32 I = 0; I < Summary.ImportCount; ++I)
@@ -1903,71 +2486,198 @@ static void ParsePackageAssets(
 
 			if (Summary.ExportCount > 0)
 			{
-				FNameProxyArchive ProxyAr(Ar, Package.NameMap);
+				FNameReaderProxyArchive ProxyAr(Ar, Package.NameMap);
 				ProxyAr.Seek(Summary.ExportOffset);
 
 				for (int32 I = 0; I < Summary.ExportCount; ++I)
 				{
 					FObjectExport& ObjectExport = PackageAssetData.ObjectExports[Package.ExportIndexOffset + I];
 					ProxyAr << ObjectExport;
+					Package.ExportsSerialSize += ObjectExport.SerialSize;
 				}
 			}
 		}, EParallelForFlags::Unbalanced);
 	}
 
+	FMemory::Free(UAssetMemory);
+}
+
+static void CreateGlobalImportsAndExports(
+	TArray<FPackage*>& Packages,
+	const FPackageMap& PackageMap,
+	FPackageAssetData& PackageAssetData,
+	FGlobalPackageData& GlobalPackageData,
+	FExportGraph& ExportGraph)
+{
+	IOSTORE_CPU_SCOPE(CreateGlobalImportsAndExports);
 	UE_LOG(LogIoStore, Display, TEXT("Creating global imports and exports..."));
+
+	FGlobalImports& GlobalImports = GlobalPackageData.Imports;
+	FImportObjectsByFullName& GlobalImportsByFullName = GlobalPackageData.Imports.ObjectsByFullName;
+
+	TArray<FString> TmpFullNames;
+
+	for (FPackage* Package : Packages)
 	{
-		IOSTORE_CPU_SCOPE(CreateGlobalImportsAndExports);
+		if (Package->ImportCount > 0)
+		{
+			FPackage* CurrentImportPackage = nullptr;
+			Package->ImportedFullNames.SetNum(Package->ImportCount);
+			for (int32 ImportIndex = 0; ImportIndex < Package->ImportCount; ++ImportIndex)
+			{
+				FindImport(
+					GlobalPackageData.Imports,
+					Package->ImportedFullNames,
+					PackageAssetData.ObjectImports.GetData() + Package->ImportIndexOffset,
+					ImportIndex,
+					PackageMap,
+					CurrentImportPackage);
+			}
+		}
 
-		TArray<FString> TmpFullNames;
-		for (FPackage* Package : Packages)
+		if (Package->ExportCount > 0)
 		{
-			if (Package->ImportCount > 0)
+			TmpFullNames.Reset();
+			TmpFullNames.SetNum(Package->ExportCount, false);
+
+			for (int32 ExportIndex = 0; ExportIndex < Package->ExportCount; ExportIndex++)
 			{
-				Package->ImportedFullNames.SetNum(Package->ImportCount);
-				for (int32 ImportIndex = 0; ImportIndex < Package->ImportCount; ++ImportIndex)
-				{
-					FindImport(
-						GlobalPackageData.Imports,
-						Package->ImportedFullNames,
-						PackageAssetData.ObjectImports.GetData() + Package->ImportIndexOffset,
-						ImportIndex);
-				}
+				FindExport(
+					GlobalPackageData.Exports,
+					TmpFullNames,
+					PackageAssetData.ObjectExports.GetData() + Package->ExportIndexOffset,
+					ExportIndex,
+					Package);
+				const int32* GlobalIndex = GlobalPackageData.Exports.ObjectsByFullName.Find(TmpFullNames[ExportIndex]);
+				FExportData& ExportData = GlobalPackageData.Exports.Objects[*GlobalIndex];
+				Package->Exports.Add(ExportData.GlobalIndex);
+				ExportData.CreateNode = ExportGraph.AddNode(Package, { uint32(ExportIndex), FExportBundleEntry::ExportCommandType_Create });
+				ExportData.SerializeNode = ExportGraph.AddNode(Package, { uint32(ExportIndex), FExportBundleEntry::ExportCommandType_Serialize });
+				Package->CreateExportNodes.Add(ExportData.CreateNode);
+				Package->SerializeExportNodes.Add(ExportData.SerializeNode);
+				ExportGraph.AddInternalDependency(ExportData.CreateNode, ExportData.SerializeNode);
 			}
+		}
+	}
 
-			if (Package->ExportCount > 0)
+	// lookup script CDO class indices for script imports
+	for (FScriptImportData& ImportData : GlobalImports.ScriptImports)
+	{
+		if (ImportData.bInitialized && ImportData.CDOClassFullName.Len() > 0)
+		{
+			FPackageObjectIndex* CDOClassIndex = GlobalImportsByFullName.Find(ImportData.CDOClassFullName);
+			if (CDOClassIndex)
+			{
+				ImportData.CDOClassIndex = *CDOClassIndex;
+			}
+			else
+			{
+				UE_LOG(LogIoStore, Warning, TEXT("Missing script import class '%s' for CDO '%s'"),
+					*ImportData.CDOClassFullName, *ImportData.FullName);
+			}
+		}
+	}
+
+	// lookup mappings between package imports and exports, and fill the package public exports
+	for (FPackageImportData& ImportData : GlobalImports.PackageImports)
+	{
+		if (ImportData.bInitialized)
+		{
+			int32* FindGlobalExport = GlobalPackageData.Exports.ObjectsByFullName.Find(ImportData.FullName);
+			if (FindGlobalExport)
+			{
+				check(ImportData.Package);
+				ImportData.Package->PublicExports.Add(ImportData.GlobalIndex);
+				ImportData.GlobalExportIndex = *FindGlobalExport;
+				GlobalPackageData.Exports.Objects[*FindGlobalExport].GlobalImportIndex = ImportData.GlobalIndex;
+			}
+			else
 			{
-				TmpFullNames.Reset();
-				TmpFullNames.SetNum(Package->ExportCount);
+				UE_LOG(LogIoStore, Warning, TEXT("Missing import '%s' due to missing export"), *ImportData.FullName);
+				ImportData.bIsMissingImport = true;
+			}
+		}
+	}
+}
 
-				for (int32 ExportIndex = 0; ExportIndex < Package->ExportCount; ExportIndex++)
+static void MapImportedPackages(
+	const FGlobalImports& GlobalImports,
+	TArray<FPackage*>& Packages)
+{
+	IOSTORE_CPU_SCOPE(MapImportedPackages);
+	UE_LOG(LogIoStore, Display, TEXT("Mapping import packages..."));
+
+	const FImportObjectsByFullName& ObjectsByFullName = GlobalImports.ObjectsByFullName;
+	TSet<FPackage*> ImportedPackages;
+	for (FPackage* Package : Packages)
+	{
+		bool bHasMissingImports = false;
+		Package->Imports.Reserve(Package->ImportCount);
+		Package->ImportedPackages.Reserve(Package->ImportCount / 2);
+		ImportedPackages.Reset();
+		for (int32 I = 0; I < Package->ImportCount; ++I)
+		{
+			FPackageObjectIndex GlobalImportIndex = ObjectsByFullName.FindRef(Package->ImportedFullNames[I]);
+			Package->Imports.Add(GlobalImportIndex);
+			if (GlobalImportIndex.IsPackageImport())
+			{
+				const FPackageImportData& ImportData = GlobalImports.PackageImports[GlobalImportIndex.GetIndex()];
+				bHasMissingImports |= ImportData.bIsMissingImport;
+				if (ImportData.Package)
 				{
-					FindExport(
-						GlobalPackageData.Exports,
-						TmpFullNames,
-						PackageAssetData.ObjectExports.GetData() + Package->ExportIndexOffset,
-						ExportIndex,
-						Package->Name);
-
-					const int32* GlobalIndex = GlobalPackageData.Exports.ObjectsByFullName.Find(TmpFullNames[ExportIndex]);
-					FExportData& ExportData = GlobalPackageData.Exports.Objects[*GlobalIndex];
-					Package->Exports.Add(ExportData.GlobalIndex);
-					ExportData.CreateNode = ExportGraph.AddNode(Package, { uint32(ExportIndex), FExportBundleEntry::ExportCommandType_Create });
-					ExportData.SerializeNode = ExportGraph.AddNode(Package, { uint32(ExportIndex), FExportBundleEntry::ExportCommandType_Serialize });
-					Package->CreateExportNodes.Add(ExportData.CreateNode);
-					Package->SerializeExportNodes.Add(ExportData.SerializeNode);
-
-					ExportGraph.AddInternalDependency(ExportData.CreateNode, ExportData.SerializeNode);
+					bool bAlreadyInSet;
+					ImportedPackages.Add(ImportData.Package, &bAlreadyInSet);
+					if (!bAlreadyInSet)
+					{
+						Package->ImportedPackages.Add(ImportData.Package);
+						ImportData.Package->ImportedByPackages.Add(Package);
+					}
 				}
 			}
 		}
+		UE_CLOG(bHasMissingImports, LogIoStore, Warning, TEXT("Missing imports for package '%s'"), *Package->Name.ToString());
 	}
-}
+};
+
+static void MapExportEntryIndices(
+	TArray<FObjectExport>& ObjectExports,
+	TArray<FExportData>& GlobalExports,
+	TArray<FPackage*>& Packages)
+{
+	IOSTORE_CPU_SCOPE(ExportData);
+	UE_LOG(LogIoStore, Display, TEXT("Converting export map import indices..."));
 
-void ProcessLocalizedPackages(
+	auto PackageObjectIndexFromPackageIndex =
+		[](const TArray<FPackageObjectIndex>& Imports, const FPackageIndex& PackageIndex) -> FPackageObjectIndex 
+		{ 
+			if (PackageIndex.IsImport())
+			{
+				return Imports[PackageIndex.ToImport()];
+			}
+			if (PackageIndex.IsExport())
+			{
+				return FPackageObjectIndex(FPackageObjectIndex::Export, PackageIndex.ToExport());
+			}
+			return FPackageObjectIndex();
+		};
+
+	for (FPackage* Package : Packages)
+	{
+		for (int32 I = 0; I < Package->ExportCount; ++I)
+		{
+			const FObjectExport& ObjectExport = ObjectExports[Package->ExportIndexOffset + I];
+			FExportData& ExportData = GlobalExports[Package->Exports[I]];
+			ExportData.OuterIndex = PackageObjectIndexFromPackageIndex(Package->Imports, ObjectExport.OuterIndex);
+			ExportData.ClassIndex = PackageObjectIndexFromPackageIndex(Package->Imports, ObjectExport.ClassIndex);
+			ExportData.SuperIndex = PackageObjectIndexFromPackageIndex(Package->Imports, ObjectExport.SuperIndex);
+			ExportData.TemplateIndex = PackageObjectIndexFromPackageIndex(Package->Imports, ObjectExport.TemplateIndex);
+		}
+	}
+};
+static void ProcessLocalizedPackages(
 	const TArray<FPackage*>& Packages,
 	const FPackageMap& PackageMap,
-	const TArray<FImportData>& GlobalImports,
+	const FGlobalImports& GlobalImports,
 	TArray<FExportData>& GlobalExports,
 	FCulturePackageMap& OutCulturePackageMap,
 	FSourceToLocalizedPackageMultimap& OutSourceToLocalizedPackageMap)
@@ -1990,7 +2700,7 @@ void ProcessLocalizedPackages(
 				TEXT("For culture '%s': Localized package '%s' (%d) should have a package name different from source name."),
 				*Package->Region,
 				*Package->Name.ToString(),
-				Package->GlobalPackageId)
+				Package->GlobalPackageId.ToIndexForDebugging())
 			continue;
 		}
 
@@ -2002,7 +2712,7 @@ void ProcessLocalizedPackages(
 				TEXT("For culture '%s': Localized package '%s' (%d) is unique and does not override a source package."),
 				*Package->Region,
 				*Package->Name.ToString(),
-				Package->GlobalPackageId);
+				Package->GlobalPackageId.ToIndexForDebugging());
 			continue;
 		}
 
@@ -2015,12 +2725,12 @@ void ProcessLocalizedPackages(
 		if (Package->bIsLocalizedAndConformed)
 		{
 			UE_LOG(LogIoStore, Verbose, TEXT("For culture '%s': Adding conformed localized package '%s' (%d) for '%s' (%d). ")
-				TEXT("When loading the source package it will be remapped to this localized package."),
+				TEXT("When loading the source package, it will be remapped to this localized package."),
 				*Package->Region,
 				*Package->Name.ToString(),
-				Package->GlobalPackageId,
+				Package->GlobalPackageId.ToIndexForDebugging(),
 				*Package->SourcePackageName.ToString(),
-				SourcePackage->GlobalPackageId);
+				SourcePackage->GlobalPackageId.ToIndexForDebugging());
 
 			OutSourceToLocalizedPackageMap.Add(SourcePackage, Package);
 			OutCulturePackageMap.FindOrAdd(Package->Region).Add(SourcePackage->GlobalPackageId, Package->GlobalPackageId);
@@ -2029,12 +2739,12 @@ void ProcessLocalizedPackages(
 		{
 			UE_LOG(LogIoStore, Warning,
 				TEXT("For culture '%s': Localized package '%s' (%d) does not conform to source package '%s' (%d) due to mismatching public exports. ")
-				TEXT("When loading the source package will never be remapped to this localized package."),
+				TEXT("When loading the source package, it will never be remapped to this localized package."),
 				*Package->Region,
 				*Package->Name.ToString(),
-				Package->GlobalPackageId,
+				Package->GlobalPackageId.ToIndexForDebugging(),
 				*Package->SourcePackageName.ToString(),
-				SourcePackage->GlobalPackageId);
+				SourcePackage->GlobalPackageId.ToIndexForDebugging());
 		}
 	}
 
@@ -2052,9 +2762,9 @@ void ProcessLocalizedPackages(
 			{
 				UE_LOG(LogIoStore, Verbose, TEXT("For package '%s' (%d): Adding localized imported package '%s' (%d)"),
 					*Package->Name.ToString(),
-					Package->GlobalPackageId,
+					Package->GlobalPackageId.ToIndexForDebugging(),
 					*LocalizedPackage->Name.ToString(),
-					LocalizedPackage->GlobalPackageId);
+					LocalizedPackage->GlobalPackageId.ToIndexForDebugging());
 			}
 		}
 		Package->ImportedPackagesSerializeCount = Package->ImportedPackages.Num();
@@ -2064,695 +2774,721 @@ void ProcessLocalizedPackages(
 	UE_LOG(LogIoStore, Display, TEXT("Conforming localized imports..."));
 	for (FPackage* Package : Packages)
 	{
-		for (int32& GlobalImportIndex : Package->Imports)
+		for (FPackageObjectIndex& GlobalImportIndex : Package->Imports)
 		{
-			const FImportData& ImportData = GlobalImports[GlobalImportIndex];
-			if (ImportData.bIsLocalized)
-			{
-				const int32* SourceGlobalImportIndex = LocalizedToSourceImportIndexMap.Find(GlobalImportIndex);
-				if (SourceGlobalImportIndex)
-				{
-					GlobalImportIndex = *SourceGlobalImportIndex;
-
-					const FImportData& SourceImportData = GlobalImports[*SourceGlobalImportIndex];
-					UE_LOG(LogIoStore, Verbose,
-						TEXT("For package '%s' (%d): Remap localized import %s to source import %s (in a conformed localized package)"),
-						*Package->Name.ToString(),
-						Package->GlobalPackageId,
-						*ImportData.FullName,
-						*SourceImportData.FullName);
-				}
-				else
-				{
-					UE_LOG(LogIoStore, Verbose,
-						TEXT("For package '%s' (%d): Skip remap for localized import %s")
-						TEXT(", either there is no source package or the localized package did not conform to it."),
-						*Package->Name.ToString(),
-						Package->GlobalPackageId,
-						*ImportData.FullName);
-				}
-
-				// FString SourceImport = RemapLocalizationPathIfNeeded(Package->ImportedFullNames[I], nullptr);
-				// if (int32* SourceImportIndex = GlobalImportsByFullName.Find(SourceImport))
-				// {
-				// 	UE_LOG(LogIoStore, Warning,
-				// 		TEXT("Remap import index from %d to %d with new name %s"),
-				// 		GlobalImportIndex, *SourceImportIndex, *SourceImport);
-
-				// 	ensure(GlobalImports[*SourceImportIndex].GlobalIndex == *SourceImportIndex);
-				// 	Package->Imports.Add(GlobalImports[*SourceImportIndex].GlobalIndex);
-				// }
-				// else
-				// {
-				// 	UE_LOG(LogIoStore, Warning,
-				// 		TEXT("Failed to remap import index from %d"),
-				// 		GlobalImportIndex);
-
-				// 	ensure(ImportData.GlobalIndex == GlobalImportIndex);
-				// 	Package->Imports.Add(ImportData.GlobalIndex);
-				// }
-			}
-			// else
-			// {
-			// 	ensure(ImportData.GlobalIndex == GlobalImportIndex);
-			// 	Package->Imports.Add(ImportData.GlobalIndex);
-			// }
-
-			// if (ImportData.Package)
-			// {
-			// 	Package->ImportedPackages.Add(ImportData.Package);
-			// 	ImportData.Package->ImportedByPackages.Add(Package);
-			// }
-		}
-}
-}
-
-int32 CreateTarget(
-	const TCHAR* OutputDir,
-	const TCHAR* CookedDir,
-	const TArray<FContainerSourceSpec>& Containers,
-	const FCookedFileStatMap& CookedFileStatMap,
-	const TMap<FString, uint64>& PackageOrderMap,
-	const TMap<FString, uint64>& CookerFileOpenOrderMap)
-{
-	TGuardValue<int32> GuardAllowUnversionedContentInEditor(GAllowUnversionedContentInEditor, 1);
-
-	FPackageStoreBulkDataManifest BulkDataManifest(FString(CookedDir) / FApp::GetProjectName());
-	const bool bWithBulkDataManifest = BulkDataManifest.Load();
-	if (bWithBulkDataManifest)
-	{
-		UE_LOG(LogIoStore, Display, TEXT("Loaded Bulk Data manifest '%s'"), *BulkDataManifest.GetFilename());
-	}
-
-#if OUTPUT_CHUNKID_DIRECTORY
-	ChunkIdCsv.CreateOutputFile(CookedDir);
-#endif
-
-	FNameMapBuilder NameMapBuilder;
-	FPackageAssetData PackageAssetData;
-	FGlobalPackageData GlobalPackageData;
-	FExportGraph ExportGraph;
-	TArray<FExportBundleMetaEntry> ExportBundleMetaEntries;
-
-	TArray<FPackage*> Packages;
-	FPackageMap PackageMap;
-
-	TArray<FIoStoreWriter*> IoStoreWriters;
-	FIoStoreWriter* GlobalIoStoreWriter;
-	FIoStoreEnvironment GlobalIoStoreEnv;
-	{
-		FString GlobalOutputPath = FString(OutputDir) / TEXT("global");
-		GlobalIoStoreEnv.InitializeFileEnvironment(*GlobalOutputPath);
-		GlobalIoStoreWriter = new FIoStoreWriter(GlobalIoStoreEnv);
-		IoStoreWriters.Add(GlobalIoStoreWriter);
-#if !SKIP_WRITE_CONTAINER
-		FIoStatus IoStatus = GlobalIoStoreWriter->Initialize();
-		check(IoStatus.IsOk());
-#if OUTPUT_CONTAINER_CSV
-		IoStatus = GlobalIoStoreWriter->EnableCsvOutput();
-		check(IoStatus.IsOk());
-#endif
-#endif
-	}
-
-	TArray<FContainerTargetSpec> ContainerTargets;
-
-	UE_LOG(LogIoStore, Display, TEXT("Creating container targets..."));
-	{
-		IOSTORE_CPU_SCOPE(CreateContainerTargets);
-
-		for (const FContainerSourceSpec& Container : Containers)
-		{
-			FContainerTargetSpec& ContainerTarget = ContainerTargets.AddDefaulted_GetRef();
-			if (Container.OutputPath.IsEmpty())
-			{
-				ContainerTarget.IoStoreWriter = GlobalIoStoreWriter;
-			}
-			else
-			{		
-				ContainerTarget.IoStoreEnv.Reset(new FIoStoreEnvironment());
-				ContainerTarget.IoStoreEnv->InitializeFileEnvironment(Container.OutputPath);
-				ContainerTarget.IoStoreWriter = new FIoStoreWriter(*ContainerTarget.IoStoreEnv);
-				IoStoreWriters.Add(ContainerTarget.IoStoreWriter);
-#if !SKIP_WRITE_CONTAINER
-				FIoStatus IoStatus = ContainerTarget.IoStoreWriter->Initialize();
-				check(IoStatus.IsOk());
-#if OUTPUT_CONTAINER_CSV
-				IoStatus = ContainerTarget.IoStoreWriter->EnableCsvOutput();
-				check(IoStatus.IsOk());
-#endif
-#endif
-			}
-
-			int32 CookedDirLen = FCString::Strlen(CookedDir) + 1;
-			for (const FContainerSourceFile& SourceFile : Container.SourceFiles)
+			if (GlobalImportIndex.IsPackageImport())
 			{
-				const FCookedFileStatData* CookedFileStatData = CookedFileStatMap.Find(SourceFile.NormalizedPath);
-				if (!CookedFileStatData)
+				const FPackageImportData& ImportData = GlobalImports.PackageImports[GlobalImportIndex.GetIndex()];
+				if (ImportData.bIsLocalized)
 				{
-					UE_LOG(LogIoStore, Warning, TEXT("File not found: '%s'"), *SourceFile.NormalizedPath);
-				}
-
-				FString RelativeFileName = TEXT("../../../");
-				RelativeFileName.AppendChars(*SourceFile.NormalizedPath + CookedDirLen, SourceFile.NormalizedPath.Len() - CookedDirLen);
-				
-				FPackage* Package = FindOrAddPackage(*RelativeFileName, Packages, PackageMap);
-				if (Package)
-				{
-					FContainerTargetFile& TargetFile = ContainerTarget.TargetFiles.AddDefaulted_GetRef();
-					TargetFile.Size = uint64(CookedFileStatData->FileSize);
-					TargetFile.NormalizedSourcePath = SourceFile.NormalizedPath;
-					TargetFile.TargetPath = MoveTemp(RelativeFileName);
-					TargetFile.Package = Package;
-
-					if (CookedFileStatData->FileType == FCookedFileStatData::PackageHeader)
+					const FPackageObjectIndex* SourceGlobalImportIndex = LocalizedToSourceImportIndexMap.Find(GlobalImportIndex);
+					if (SourceGlobalImportIndex)
 					{
-						TargetFile.bIsBulkData = false;
-						Package->FileName = SourceFile.NormalizedPath;
-						Package->RelativeFileName = TargetFile.TargetPath;
+						GlobalImportIndex = *SourceGlobalImportIndex;
+
+						const FPackageImportData& SourceImportData = GlobalImports.PackageImports[SourceGlobalImportIndex->GetIndex()];
+						UE_LOG(LogIoStore, Verbose,
+							TEXT("For package '%s' (%d): Remap localized import %s to source import %s (in a conformed localized package)"),
+							*Package->Name.ToString(),
+							Package->GlobalPackageId.ToIndexForDebugging(),
+							*ImportData.FullName,
+							*SourceImportData.FullName);
 					}
 					else
 					{
-						TargetFile.bIsBulkData = true;
-						TargetFile.bIsOptionalBulkData = CookedFileStatData->FileExt == FCookedFileStatData::UPtnl;
+						UE_LOG(LogIoStore, Verbose,
+							TEXT("For package '%s' (%d): Skip remap for localized import %s")
+							TEXT(", either there is no source package or the localized package did not conform to it."),
+							*Package->Name.ToString(),
+							Package->GlobalPackageId.ToIndexForDebugging(),
+							*ImportData.FullName);
 					}
 				}
 			}
 		}
 	}
+}
 
-	UE_LOG(LogIoStore, Display, TEXT("Parsing packages..."));
-	ParsePackageAssets(NameMapBuilder, Packages, PackageAssetData, GlobalPackageData, ExportGraph);
-
-	TArray<FImportData>& GlobalImports = GlobalPackageData.Imports.Objects;
-	TArray<FExportData>& GlobalExports = GlobalPackageData.Exports.Objects;
-	TMap<FString, int32>& GlobalImportsByFullName = GlobalPackageData.Imports.ObjectsByFullName;
-	TMap<FString, int32>& GlobalExportsByFullName = GlobalPackageData.Exports.ObjectsByFullName;
+static void SaveReleaseVersionMeta(const TCHAR* ReleaseVersionOutputDir, const FNameMapBuilder& NameMap, const FPackageGlobalIdMap& PackageGlobalIdMap, const FGlobalImports& GlobalImports, const TArray<FContainerTargetSpec>& ContainerTargets)
+{
+	IPlatformFile::GetPlatformPhysical().CreateDirectoryTree(ReleaseVersionOutputDir);
 
-	int32 NumScriptImports = 0;
 	{
-		IOSTORE_CPU_SCOPE(GlobalImports);
+		FString NameMapOutputPath = FPaths::Combine(ReleaseVersionOutputDir, TEXT("iodispatcher.unamemap"));
+		TUniquePtr<FArchive> NameMapArchive(IFileManager::Get().CreateFileWriter(*NameMapOutputPath));
+		(*NameMapArchive) << const_cast<FNameMapBuilder&>(NameMap);
+	}
 
-		// Sort imports by script objects first
-		GlobalImports.Sort();
+	{
+		FString PackageMapOutputPath = FPaths::Combine(ReleaseVersionOutputDir, TEXT("iodispatcher.upackagemap"));
+		TUniquePtr<FArchive> PackageMapArchive(IFileManager::Get().CreateFileWriter(*PackageMapOutputPath));
 
-		// build remap from old global import index to new sorted global import index
-		TMap<int32, int32> Remap;
-		Remap.Reserve(GlobalImports.Num() + 1);
-		Remap.Add(-1, -1);
-		for (int32 I = 0; I < GlobalImports.Num(); ++I)
+		TArray<FName> OrderedGlobalIdMap;
+		OrderedGlobalIdMap.SetNum(PackageGlobalIdMap.Num());
+		for (const auto& KV : PackageGlobalIdMap)
 		{
-			FImportData& Import = GlobalImports[I];
-			Remap.Add(Import.GlobalIndex, I);
+			FName PackageName = KV.Key;
+			FPackageId PackageGlobalId = KV.Value;
+			OrderedGlobalIdMap[PackageGlobalId.ToIndex()] = PackageName;
 		}
 
-		// remap all global import indices and lookup package pointers and export indices
-		FPackage* LastPackage = nullptr; 
-		for (int32 I = 0; I < GlobalImports.Num(); ++I)
+		int32 PackageCount = OrderedGlobalIdMap.Num();
+		(*PackageMapArchive) << PackageCount;
+		for (const FName& PackageName : OrderedGlobalIdMap)
 		{
-			FImportData& GlobalImport = GlobalImports[I];
-
-			GlobalImport.GlobalIndex = Remap[GlobalImport.GlobalIndex];
-			GlobalImport.OuterIndex = Remap[GlobalImport.OuterIndex];
-			GlobalImport.OutermostIndex = Remap[GlobalImport.OutermostIndex];
-
-			if (!GlobalImport.bIsScript)
-			{
-				if (NumScriptImports == 0)
-				{
-					NumScriptImports = I;
-				}
+			const FNameEntry* NameEntry = FName::GetEntry(PackageName.GetComparisonIndex());
+			NameEntry->Write(*PackageMapArchive);
+			int32 NameNumber = PackageName.GetNumber();
+			(*PackageMapArchive) << NameNumber;
+		}
+	}
 
-				if (GlobalImport.bIsPackage)
-				{
-					if (LastPackage)
-					{
-						LastPackage->GlobalImportCount = I - LastPackage->FirstGlobalImport;
-					}
-					FPackage* FindPackage = PackageMap.FindRef(GlobalImport.ObjectName);
-					if (!FindPackage)
-					{
-						UE_LOG(LogIoStore, Warning, TEXT("Missing import package '%s'"), *GlobalImport.ObjectName.ToString());
-						LastPackage = nullptr;
-						GlobalImport.bIsMissingImport = true;
-					}
-					else
-					{
-						FindPackage->FirstGlobalImport = I;
-						LastPackage = FindPackage;
-						GlobalImport.Package = FindPackage;
-					}
-				}
-				else
-				{
-					int32* FindGlobalExport = GlobalExportsByFullName.Find(GlobalImport.FullName);
-					if (!FindGlobalExport)
-					{
-						UE_LOG(LogIoStore, Warning, TEXT("Missing import object '%s'"), *GlobalImport.FullName);
-						GlobalImport.bIsMissingImport = true;
-					}
-					else
-					{
-						GlobalImport.GlobalExportIndex = *FindGlobalExport;
-					}
-				}
-			}
+	{
+		FString ImportMapOutputPath = FPaths::Combine(ReleaseVersionOutputDir, TEXT("iodispatcher.uimportmap"));
+		TUniquePtr<FArchive> ImportMapArchive(IFileManager::Get().CreateFileWriter(*ImportMapOutputPath));
 
-			GlobalImportsByFullName[GlobalImport.FullName] = I;
-		}
-		if (LastPackage)
+		int32 ImportCount = GlobalImports.ObjectsByFullName.Num();
+		(*ImportMapArchive) << ImportCount;
+		for (const auto& KV : GlobalImports.ObjectsByFullName)
 		{
-			LastPackage->GlobalImportCount = GlobalImports.Num() - LastPackage->FirstGlobalImport;
+			FString ObjectName = KV.Key;
+			FPackageObjectIndex ObjectIndex = KV.Value;
+			int32 Type = ObjectIndex.GetType();
+			int32 Index = ObjectIndex.GetIndex();
+
+			(*ImportMapArchive) << ObjectName;
+			(*ImportMapArchive) << Type;
+			(*ImportMapArchive) << Index;
 		}
 	}
 
+	for (const FContainerTargetSpec& ContainerTarget : ContainerTargets)
 	{
-		IOSTORE_CPU_SCOPE(GlobalExports);
-
-		for (FExportData& GlobalExport : GlobalPackageData.Exports.Objects)
+		TArray<FTocMetaEntry> TocMetaEntries;
+		TocMetaEntries.Reserve(ContainerTarget.TargetFiles.Num());
+		for (const FContainerTargetFile& TargetFile : ContainerTarget.TargetFiles)
 		{
-			int32* FindGlobalImport = GlobalPackageData.Imports.ObjectsByFullName.Find(GlobalExport.FullName);
-			GlobalExport.GlobalImportIndex = FindGlobalImport ? *FindGlobalImport : -1;
+			FTocMetaEntry& TocMetaEntry = TocMetaEntries.AddDefaulted_GetRef();
+			TocMetaEntry.Offset = TargetFile.Offset;
+			TocMetaEntry.Size = TargetFile.TargetSize;
+			TocMetaEntry.Hash = TargetFile.Hash;
 		}
+
+		FString TocMetaOutputPath = FPaths::Combine(ReleaseVersionOutputDir, ContainerTarget.Name + TEXT(".utocmeta"));
+		TUniquePtr<FArchive> TocMetaArchive(IFileManager::Get().CreateFileWriter(*TocMetaOutputPath));
+		(*TocMetaArchive) << TocMetaEntries;
 	}
+}
 
-#if OUTPUT_NAMEMAP_CSV
-	FString CsvFilePath = OutputDir / TEXT("AllImports.csv");
-	TUniquePtr<FArchive> CsvArchive(IFileManager::Get().CreateFileWriter(*CsvFilePath));
-	if (CsvArchive)
+static void LoadReleaseVersionMeta(const TCHAR* ReleaseVersionOutputDir, FNameMapBuilder& NameMap, FPackageGlobalIdMap& PackageGlobalIdMap, FGlobalImports& GlobalImports, TArray<FContainerTargetSpec>& ContainerTargets)
+{
 	{
-		ANSICHAR Line[MAX_SPRINTF + FName::StringBufferSize];
-		ANSICHAR Header[] = "Count\tOuter\tOutermost\tImportName\n";
-		CsvArchive->Serialize(Header, sizeof(Header) - 1);
-		for (const FImportData& ImportData : GlobalImports)
+		FString NameMapPath = FPaths::Combine(ReleaseVersionOutputDir, TEXT("iodispatcher.unamemap"));
+		TUniquePtr<FArchive> NameMapArchive(IFileManager::Get().CreateFileReader(*NameMapPath));
+		if (NameMapArchive)
 		{
-			FCStringAnsi::Sprintf(Line, "%d\t%d\t%d\t",
-				ImportData.RefCount, ImportData.OuterIndex, ImportData.OutermostIndex);
-			ANSICHAR* L = Line + FCStringAnsi::Strlen(Line);
-			const TCHAR* N = *ImportData.FullName;
-			while (*N)
-			{
-				*L++ = CharCast<ANSICHAR,TCHAR>(*N++);
-			}
-			*L++ = '\n';
-			CsvArchive.Get()->Serialize(Line, L - Line);
+			(*NameMapArchive) << NameMap;
 		}
 	}
-#endif
-
-
-	// Lookup global indices and package pointers for all imports before adding preload and postload arcs
-	UE_LOG(LogIoStore, Display, TEXT("Looking up import packages..."));
 	{
-		IOSTORE_CPU_SCOPE(ImportedPackages);
-
-		for (FPackage* Package : Packages)
+		FString PackageMapPath = FPaths::Combine(ReleaseVersionOutputDir, TEXT("iodispatcher.upackagemap"));
+		TUniquePtr<FArchive> PackageMapArchive(IFileManager::Get().CreateFileReader(*PackageMapPath));
+		if (PackageMapArchive)
 		{
-			bool bHasMissingImports = false;
-			Package->ImportedPackages.Reserve(Package->ImportCount);
-			for (int32 I = 0; I < Package->ImportCount; ++I)
+			int32 PackageCount;
+			(*PackageMapArchive) << PackageCount;
+			for (int32 PackageIndex = 0; PackageIndex < PackageCount; ++PackageIndex)
 			{
-				int32 GlobalImportIndex = *GlobalImportsByFullName.Find(Package->ImportedFullNames[I]);
-
-				FImportData& ImportData = GlobalImports[GlobalImportIndex];
-				bHasMissingImports |= ImportData.bIsMissingImport;
-				Package->Imports.Add(ImportData.GlobalIndex);
-				if (ImportData.Package)
-				{
-					Package->ImportedPackages.Add(ImportData.Package);
-					ImportData.Package->ImportedByPackages.Add(Package);
-				}
+				FNameEntrySerialized NameEntrySerialized(ENAME_LinkerConstructor);
+				(*PackageMapArchive) << NameEntrySerialized;
+				int32 NameNumber;
+				(*PackageMapArchive) << NameNumber;
+				FName Name(NameEntrySerialized, NameNumber);
+				check(!PackageGlobalIdMap.Contains(Name));
+				PackageGlobalIdMap.Add(Name, FPackageId::FromIndex(PackageIndex));
 			}
-			UE_CLOG(bHasMissingImports, LogIoStore, Warning, TEXT("Missing imports for package '%s'"), *Package->Name.ToString());
+			check(PackageGlobalIdMap.Num() == PackageCount);
 		}
 	}
-
-	UE_LOG(LogIoStore, Display, TEXT("Converting export map import indices..."));
 	{
-		IOSTORE_CPU_SCOPE(ExportData);
-
-		for (FPackage* Package : Packages)
+		FString ImportMapPath = FPaths::Combine(ReleaseVersionOutputDir, TEXT("iodispatcher.uimportmap"));
+		TUniquePtr<FArchive> ImportMapArchive(IFileManager::Get().CreateFileReader(*ImportMapPath));
+		if (ImportMapArchive)
 		{
-			for (int32 I = 0; I < Package->ExportCount; ++I)
+			int32 ImportCount = 0;
+			(*ImportMapArchive) << ImportCount;
+			for (int32 ImportIndex = 0; ImportIndex < ImportCount; ++ImportIndex)
 			{
-				FObjectExport& ObjectExport = PackageAssetData.ObjectExports[Package->ExportIndexOffset + I];
-				FExportData& ExportData = GlobalPackageData.Exports.Objects[Package->Exports[I]];
+				FString FullName;
+				int32 Type;
+				int32 Index;
 
-				check(!ObjectExport.OuterIndex.IsImport());
-				ExportData.OuterIndex = ObjectExport.OuterIndex;
-				ExportData.ClassIndex =
-					ObjectExport.ClassIndex.IsImport() ?
-					FPackageIndex::FromImport(Package->Imports[ObjectExport.ClassIndex.ToImport()]) :
-					ObjectExport.ClassIndex;
-				ExportData.SuperIndex =
-					ObjectExport.SuperIndex.IsImport() ?
-					FPackageIndex::FromImport(Package->Imports[ObjectExport.SuperIndex.ToImport()]) :
-					ObjectExport.SuperIndex;
-				ExportData.TemplateIndex =
-					ObjectExport.TemplateIndex.IsImport() ?
-					FPackageIndex::FromImport(Package->Imports[ObjectExport.TemplateIndex.ToImport()]) :
-					ObjectExport.TemplateIndex;
+				(*ImportMapArchive) << FullName;
+				(*ImportMapArchive) << Type;
+				(*ImportMapArchive) << Index;
+
+				FPackageObjectIndex ObjectIndex(static_cast<FPackageObjectIndex::Type>(Type), Index);
+				GlobalImports.ObjectsByFullName.Add(FullName, ObjectIndex);
+				if (Type == FPackageObjectIndex::ScriptImport)
+				{
+					TArray<FScriptImportData>& ScriptImports = GlobalImports.ScriptImports;
+					if (ScriptImports.Num() <= Index)
+					{
+						ScriptImports.AddDefaulted(Index - ScriptImports.Num() + 1);
+					}
+					FScriptImportData& ScriptImport = ScriptImports[Index];
+					ScriptImport.GlobalIndex = ObjectIndex;
+					ScriptImport.OutermostIndex = ObjectIndex;
+					ScriptImport.FullName = FullName;
+				}
+				else
+				{
+					TArray<FPackageImportData>& PackageImports = GlobalImports.PackageImports;
+					if (PackageImports.Num() <= Index)
+					{
+						PackageImports.AddDefaulted(Index - PackageImports.Num() + 1);
+					}
+					FPackageImportData& PackageImport = PackageImports[Index];
+					PackageImport.GlobalIndex = ObjectIndex;
+					PackageImport.FullName = FullName;
+				}
 			}
 		}
 	}
 
-	FSourceToLocalizedPackageMultimap SourceToLocalizedPackageMap;
-	FCulturePackageMap CulturePackageMap;
-
-	ProcessLocalizedPackages(
-		Packages, PackageMap, GlobalImports, GlobalExports,
-		CulturePackageMap, SourceToLocalizedPackageMap);
-
-	int32 CircularChainCount = 0;
-
-	UE_LOG(LogIoStore, Display, TEXT("Adding postload dependencies..."));
+	for (FContainerTargetSpec& ContainerTarget : ContainerTargets)
 	{
-		IOSTORE_CPU_SCOPE(PostLoadDependencies);
-
-		TSet<FPackage*> Visited;
-		TSet<FCircularImportChain> CircularChains;
-
-		for (FPackage* Package : Packages)
+		FString TocMetaOutputPath = FPaths::Combine(ReleaseVersionOutputDir, ContainerTarget.Name + TEXT(".utocmeta"));
+		TUniquePtr<FArchive> TocMetaArchive(IFileManager::Get().CreateFileReader(*TocMetaOutputPath));
+		if (TocMetaArchive)
 		{
-			Visited.Reset();
-			Visited.Add(Package);
-			AddPostLoadDependencies(*Package, Visited, CircularChains);
+			(*TocMetaArchive) << ContainerTarget.PreviousBuildToc;
 		}
-		CircularChainCount = CircularChains.Num();
 	}
-		
-	UE_LOG(LogIoStore, Display, TEXT("Adding preload dependencies..."));
-	{
-		IOSTORE_CPU_SCOPE(PreLoadDependencies);
+}
 
-		TArray<FPackage*> LocalizedPackages;
-		for (FPackage* Package : Packages)
+void InitializeContainerTargetsAndPackages(
+	const FIoStoreArguments& Arguments,
+	TArray<FPackage*>& Packages,
+	FPackageMap& PackageMap,
+	FPackageGlobalIdMap& PackageGlobalIdMap,
+	TArray<FContainerTargetSpec>& ContainerTargets)
+{
+	int32 CookedDirLen = Arguments.CookedDir.Len() + 1;
+	for (int32 ContainerIndex = 0; ContainerIndex < ContainerTargets.Num(); ++ContainerIndex)
+	{
+		const FContainerSourceSpec& Container = Arguments.Containers[ContainerIndex];
+		FContainerTargetSpec& ContainerTarget = ContainerTargets[ContainerIndex];
 		{
-			// Convert PreloadDependencies to arcs
-			for (int32 I = 0; I < Package->ExportCount; ++I)
+			IOSTORE_CPU_SCOPE(ProcessSourceFiles);
+			for (const FContainerSourceFile& SourceFile : Container.SourceFiles)
 			{
-				FObjectExport& ObjectExport = PackageAssetData.ObjectExports[Package->ExportIndexOffset + I];
-				FExportData& ExportData = GlobalPackageData.Exports.Objects[Package->Exports[I]];
-				int32 PreloadDependenciesBaseIndex = Package->PreloadIndexOffset;
+				const FCookedFileStatData* OriginalCookedFileStatData = Arguments.CookedFileStatMap.Find(SourceFile.NormalizedPath);
+				if (!OriginalCookedFileStatData)
+				{
+					UE_LOG(LogIoStore, Warning, TEXT("File not found: '%s'"), *SourceFile.NormalizedPath);
+					continue;
+				}
+				const FCookedFileStatData* CookedFileStatData = OriginalCookedFileStatData;
+				FString NormalizedSourcePath = SourceFile.NormalizedPath;
+				if (CookedFileStatData->FileType == FCookedFileStatData::PackageHeader)
+				{
+					NormalizedSourcePath = FPaths::ChangeExtension(SourceFile.NormalizedPath, TEXT(".uexp"));
+					CookedFileStatData = Arguments.CookedFileStatMap.Find(NormalizedSourcePath);
+					if (!CookedFileStatData)
+					{
+						UE_LOG(LogIoStore, Warning, TEXT("File not found: '%s'"), *NormalizedSourcePath);
+						continue;
+					}
+				}
 
-				FPackageIndex ExportPackageIndex = FPackageIndex::FromExport(I);
+				FString RelativeFileName = TEXT("../../../");
+				RelativeFileName.AppendChars(*NormalizedSourcePath + CookedDirLen, NormalizedSourcePath.Len() - CookedDirLen);
+				FPackage* Package = nullptr;
+				const bool bIsMemoryMappedBulkData = CookedFileStatData->FileExt == FCookedFileStatData::EFileExt::UMappedBulk;
 
-				auto AddPreloadArc = [&](FPackageIndex Dep, EPreloadDependencyType PhaseFrom, EPreloadDependencyType PhaseTo)
+				if (bIsMemoryMappedBulkData)
 				{
-					if (Dep.IsExport())
+					FString TmpFileName = FString(RelativeFileName.Len() - 8, GetData(RelativeFileName)) + TEXT(".ubulk");
+					Package = FindOrAddPackage(*TmpFileName, Packages, PackageMap, PackageGlobalIdMap);
+				}
+				else
+				{
+					Package = FindOrAddPackage(*RelativeFileName, Packages, PackageMap, PackageGlobalIdMap);
+				}
+
+				if (Package)
+				{
+					FContainerTargetFile& TargetFile = ContainerTarget.TargetFiles.AddDefaulted_GetRef();
+					TargetFile.SourceSize = uint64(CookedFileStatData->FileSize);
+					TargetFile.NormalizedSourcePath = NormalizedSourcePath;
+					TargetFile.TargetPath = MoveTemp(RelativeFileName);
+					TargetFile.Package = Package;
+					if (SourceFile.bNeedsCompression)
 					{
-						AddInternalExportArc(ExportGraph, *Package, Dep.ToExport(), PhaseFrom, I, PhaseTo);
+						ContainerTarget.bIsCompressed = true;
 					}
 					else
 					{
-						FImportData& Import = GlobalImports[Package->Imports[Dep.ToImport()]];
-						check(!Import.bIsPackage);
-						if (Import.bIsScript)
+						TargetFile.bForceUncompressed = true;
+					}
+
+					if (CookedFileStatData->FileType == FCookedFileStatData::BulkData)
+					{
+						TargetFile.bIsBulkData = true;
+						if (CookedFileStatData->FileExt == FCookedFileStatData::UPtnl)
+						{
+							TargetFile.bIsOptionalBulkData = true;
+							TargetFile.ChunkId = CreateChunkIdForBulkData(Package->GlobalPackageId, BulkdataTypeToChunkIdType(FPackageStoreBulkDataManifest::EBulkdataType::Optional), *TargetFile.TargetPath);
+						}
+						else if (CookedFileStatData->FileExt == FCookedFileStatData::UMappedBulk)
 						{
-							// Add script arc with null package and global import index as node index
-							AddScriptArc(*Package, Import.GlobalIndex, I, PhaseTo);
+							TargetFile.bIsMemoryMappedBulkData = true;
+							TargetFile.bForceUncompressed = true;
+							TargetFile.ChunkId = CreateChunkIdForBulkData(Package->GlobalPackageId, BulkdataTypeToChunkIdType(FPackageStoreBulkDataManifest::EBulkdataType::MemoryMapped), *TargetFile.TargetPath);
 						}
 						else
 						{
-							if (Import.GlobalExportIndex != -1)
-							{
-								FExportData& Export = GlobalExports[Import.GlobalExportIndex];
-
-								FPackage* SourcePackage = PackageMap.FindRef(Export.SourcePackageName);
-								check(SourcePackage);
-
-								AddExternalExportArc(ExportGraph, *SourcePackage, Export.SourceIndex, PhaseFrom, *Package, I, PhaseTo);
-								Package->ImportedPreloadPackages.Add(SourcePackage);
-
-								LocalizedPackages.Reset();
-								SourceToLocalizedPackageMap.MultiFind(SourcePackage, LocalizedPackages);
-								for (FPackage* LocalizedPackage : LocalizedPackages)
-								{
-									UE_LOG(LogIoStore, Verbose, TEXT("For package '%s' (%d): Adding localized preload dependency '%s' in '%s'"),
-										*Package->Name.ToString(),
-										Package->GlobalPackageId,
-										*Export.ObjectName.ToString(),
-										*LocalizedPackage->Name.ToString());
-
-									AddExternalExportArc(ExportGraph, *LocalizedPackage, Export.SourceIndex, PhaseFrom, *Package, I, PhaseTo);
-									Package->ImportedPreloadPackages.Add(LocalizedPackage);
-								}
-							}
-							else
-							{
-								UE_LOG(LogIoStore, Verbose, TEXT("Skipping export arc to '%s' due to missing import"), *Package->Name.ToString());
-							}
+							TargetFile.ChunkId = CreateChunkIdForBulkData(Package->GlobalPackageId, BulkdataTypeToChunkIdType(FPackageStoreBulkDataManifest::EBulkdataType::Normal), *TargetFile.TargetPath);
 						}
 					}
-				};
+					else
+					{
+						check(CookedFileStatData->FileType == FCookedFileStatData::PackageData);
+						Package->FileName = SourceFile.NormalizedPath; // .uasset path
+						Package->UAssetSize = OriginalCookedFileStatData->FileSize;
+						Package->UExpSize = CookedFileStatData->FileSize;
+						TargetFile.ChunkId = CreateChunkId(Package->GlobalPackageId, 0, EIoChunkType::ExportBundleData, *TargetFile.TargetPath);
+					}
+				}
+			}
+		}
+	}
+
+	Algo::Sort(Packages, [](const FPackage* A, const FPackage* B)
+	{
+		return A->GlobalPackageId < B->GlobalPackageId;
+	});
+};
+
+int32 CreateTarget(const FIoStoreArguments& Arguments, const FIoStoreWriterSettings& GeneralIoWriterSettings)
+{
+	TGuardValue<int32> GuardAllowUnversionedContentInEditor(GAllowUnversionedContentInEditor, 1);
+
+	FPackageStoreBulkDataManifest BulkDataManifest(FString(Arguments.CookedDir) / FApp::GetProjectName());
+	if (!BulkDataManifest.Load())
+	{
+		UE_LOG(LogIoStore, Warning, TEXT("Failed to load Bulk Data manifest %s"), *BulkDataManifest.GetFilename());
+	}
+	else
+	{
+		UE_LOG(LogIoStore, Display, TEXT("Loaded Bulk Data manifest '%s'"), *BulkDataManifest.GetFilename());
+	}
+	
+#if OUTPUT_CHUNKID_DIRECTORY
+	ChunkIdCsv.CreateOutputFile(CookedDir);
+#endif
+
+	FNameMapBuilder NameMapBuilder;
+	FPackageAssetData PackageAssetData;
+	FGlobalPackageData GlobalPackageData;
+	FExportGraph ExportGraph;
+
+	TArray<FPackage*> Packages;
+	FPackageMap PackageMap;
+	FPackageGlobalIdMap PackageGlobalIdMap;
 
-				if (PreloadDependenciesBaseIndex >= 0 && ObjectExport.FirstExportDependency >= 0)
-				{
-					int32 RunningIndex = PreloadDependenciesBaseIndex + ObjectExport.FirstExportDependency;
-					for (int32 Index = ObjectExport.SerializationBeforeSerializationDependencies; Index > 0; Index--)
-					{
-						FPackageIndex Dep = PackageAssetData.PreloadDependencies[RunningIndex++];
-						check(!Dep.IsNull());
-						AddPreloadArc(Dep, PreloadDependencyType_Serialize, PreloadDependencyType_Serialize);
-					}
+#if OUTPUT_DEBUG_PACKAGE_HASHES
+	TMap<FName, FPackageHashes> PreviousBuildPackageHashes;
+#endif
 
-					for (int32 Index = ObjectExport.CreateBeforeSerializationDependencies; Index > 0; Index--)
-					{
-						FPackageIndex Dep = PackageAssetData.PreloadDependencies[RunningIndex++];
-						check(!Dep.IsNull());
-						AddPreloadArc(Dep, PreloadDependencyType_Create, PreloadDependencyType_Serialize);
-					}
+	TArray<FContainerTargetSpec> ContainerTargets;
+	UE_LOG(LogIoStore, Display, TEXT("Creating container targets..."));
+	{
+		IOSTORE_CPU_SCOPE(CreateContainerTargets);
 
-					for (int32 Index = ObjectExport.SerializationBeforeCreateDependencies; Index > 0; Index--)
-					{
-						FPackageIndex Dep = PackageAssetData.PreloadDependencies[RunningIndex++];
-						check(!Dep.IsNull());
-						AddPreloadArc(Dep, PreloadDependencyType_Serialize, PreloadDependencyType_Create);
-					}
+		for (const FContainerSourceSpec& Container : Arguments.Containers)
+		{
+			FContainerTargetSpec& ContainerTarget = ContainerTargets.AddDefaulted_GetRef();
+			ContainerTarget.OutputPath = Container.OutputPath;
+			ContainerTarget.Name = FPaths::GetBaseFilename(ContainerTarget.OutputPath);
+		}
 
-					for (int32 Index = ObjectExport.CreateBeforeCreateDependencies; Index > 0; Index--)
-					{
-						FPackageIndex Dep = PackageAssetData.PreloadDependencies[RunningIndex++];
-						check(!Dep.IsNull());
-						// can't create this export until these things are created
-						AddPreloadArc(Dep, PreloadDependencyType_Create, PreloadDependencyType_Create);
-					}
+		if (!Arguments.BasedOnReleaseVersionDir.IsEmpty())
+		{
+			LoadReleaseVersionMeta(*Arguments.BasedOnReleaseVersionDir, NameMapBuilder, PackageGlobalIdMap, GlobalPackageData.Imports, ContainerTargets);
+#if OUTPUT_DEBUG_PACKAGE_HASHES
+			FString PackageHashesOutputPath = FPaths::Combine(*Arguments.BasedOnReleaseVersionDir, TEXT("iodispatcher.upackagehashes"));
+			TUniquePtr<FArchive> PackageHashesArchive(IFileManager::Get().CreateFileReader(*PackageHashesOutputPath));
+			if (PackageHashesArchive)
+			{
+				int32 HashesCount;
+				(*PackageHashesArchive) << HashesCount;
+				for (int32 Index = 0; Index < HashesCount; ++Index)
+				{
+					FNameEntrySerialized NameEntrySerialized(ENAME_LinkerConstructor);
+					(*PackageHashesArchive) << NameEntrySerialized;
+					int32 NameNumber;
+					(*PackageHashesArchive) << NameNumber;
+					FName Name(NameEntrySerialized, NameNumber);
+					FPackageHashes& Hashes = PreviousBuildPackageHashes.Add(Name);
+					(*PackageHashesArchive) << Hashes.UAssetHash;
+					(*PackageHashesArchive) << Hashes.UExpHash;
+					(*PackageHashesArchive) << Hashes.ExportBundleHash;
 				}
 			}
+#endif
 		}
+
+		InitializeContainerTargetsAndPackages(Arguments, Packages, PackageMap, PackageGlobalIdMap, ContainerTargets);
 	}
 
-	UE_LOG(LogIoStore, Display, TEXT("Building bundles..."));
-	BuildBundles(ExportGraph, Packages);
+	ParsePackageAssets(NameMapBuilder, Packages, PackageAssetData, GlobalPackageData, ExportGraph);
+	CreateGlobalImportsAndExports(Packages, PackageMap, PackageAssetData, GlobalPackageData, ExportGraph);
 
-	TUniquePtr<FLargeMemoryWriter> StoreTocArchive = MakeUnique<FLargeMemoryWriter>(0, true);
-	TUniquePtr<FLargeMemoryWriter> ImportedPackagesArchive = MakeUnique<FLargeMemoryWriter>(0, true);
-	TUniquePtr<FLargeMemoryWriter> GlobalImportNamesArchive = MakeUnique<FLargeMemoryWriter>(0, true);
-	TUniquePtr<FLargeMemoryWriter> InitialLoadArchive = MakeUnique<FLargeMemoryWriter>(0, true);
+	FGlobalPackageImports& PackageImports = GlobalPackageData.Imports.PackageImports;
+	FGlobalScriptImports& ScriptImports = GlobalPackageData.Imports.ScriptImports;
+	TArray<FExportData>& GlobalExports = GlobalPackageData.Exports.Objects;
 
-	UE_LOG(LogIoStore, Display, TEXT("Serializing global import names..."));
-	if (GlobalImportNamesArchive)
-	{
-		IOSTORE_CPU_SCOPE(SerializeGlobalImportNames);
+	// Mapped import and exports are required before processing localization, and preload/postload arcs
+	MapImportedPackages(GlobalPackageData.Imports, Packages);
+	MapExportEntryIndices(PackageAssetData.ObjectExports, GlobalExports, Packages);
 
-		for (const FImportData& ImportData : GlobalImports)
-		{
-			NameMapBuilder.SerializeName(*GlobalImportNamesArchive, ImportData.ObjectName);
-		}
-	}
+	FSourceToLocalizedPackageMultimap SourceToLocalizedPackageMap;
+	FCulturePackageMap CulturePackageMap;
+
+	ProcessLocalizedPackages(
+		Packages, PackageMap, GlobalPackageData.Imports, GlobalExports,
+		CulturePackageMap, SourceToLocalizedPackageMap);
+
+	const int32 CircularChainCount = AddPostLoadDependencies(Packages);
+	
+	AddPreloadDependencies(
+		PackageAssetData,
+		GlobalPackageData,
+		SourceToLocalizedPackageMap,
+		ExportGraph,
+		Packages);
+
+	BuildBundles(ExportGraph, Packages);
+
+	FinalizePackageHeaders(
+		Packages,
+		NameMapBuilder,
+		PackageAssetData.ObjectExports,
+		GlobalPackageData.Exports.Objects,
+		GlobalPackageData.Imports.ObjectsByFullName);
+
+	FLargeMemoryWriter GlobalMetaArchive(0, true);
+	FLargeMemoryWriter InitialLoadArchive(0, true);
+	int32 StoreTocByteCount = 0;
+	int32 StoreDataByteCount = 0;
+	int32 GlobalImportNamesByteCount = 0;
 
-	UE_LOG(LogIoStore, Display, TEXT("Serializing initial load..."));
-	// Separate file for script arcs that are only required during initial loading
-	if (InitialLoadArchive)
 	{
-		IOSTORE_CPU_SCOPE(SerializeInitialLoad);
+		IOSTORE_CPU_SCOPE(SerializeGlobalMetaData);
+		UE_LOG(LogIoStore, Display, TEXT("Serializing global meta data"));
 
-		int32 PackageCount = PackageMap.Num();
-		*InitialLoadArchive << PackageCount;
-		*InitialLoadArchive << NumScriptImports;
+		FLargeMemoryWriter StoreTocArchive(0, true);
+		FLargeMemoryWriter StoreDataArchive(0, true);
+		FLargeMemoryWriter GlobalImportNamesArchive(0, true);
 
-		FBufferWriter ScriptArcsArchive(nullptr, 0, EBufferWriterFlags::AllowResize | EBufferWriterFlags::TakeOwnership);
+		SerializeImportNames(NameMapBuilder, PackageImports, GlobalImportNamesArchive);
 
-		for (auto& PackageKV : PackageMap)
-		{
-			FPackage& Package = *PackageKV.Value;
+		SerializePackageStoreEntries(
+			Packages,
+			NameMapBuilder,
+			StoreTocArchive,
+			StoreDataArchive);
 
-			int32 ScriptArcsOffset = ScriptArcsArchive.Tell();
-			int32 ScriptArcsCount = Package.ScriptArcs.Num();
+		SerializeInitialLoad(
+			NameMapBuilder,
+			ScriptImports,
+			InitialLoadArchive);
 
-			*InitialLoadArchive << ScriptArcsOffset;
-			*InitialLoadArchive << ScriptArcsCount;
+		StoreTocByteCount = StoreTocArchive.TotalSize();
+		StoreDataByteCount = StoreDataArchive.TotalSize();
+		GlobalImportNamesByteCount = GlobalImportNamesArchive.TotalSize();
 
-			for (FArc& ScriptArc : Package.ScriptArcs)
-			{
-				ScriptArcsArchive << ScriptArc.FromNodeIndex;
-				ScriptArcsArchive << ScriptArc.ToNodeIndex;
-			}
-		}
+		int32 NumPackages = Packages.Num();
+		int32 StoreByteCount = StoreTocByteCount + StoreDataByteCount;
 
-		for (int32 I = 0; I < NumScriptImports; ++I)
-		{
-			FImportData& ImportData = GlobalImports[I];
-			FPackageIndex OuterIndex = 
-				ImportData.OuterIndex >= 0 ?
-				FPackageIndex::FromImport(ImportData.OuterIndex) :
-				FPackageIndex();
+		GlobalMetaArchive << NumPackages;
+		GlobalMetaArchive << StoreByteCount;
+		GlobalMetaArchive.Serialize(StoreTocArchive.GetData(), StoreTocByteCount);
+		GlobalMetaArchive.Serialize(StoreDataArchive.GetData(), StoreDataByteCount);
 
-			*InitialLoadArchive << OuterIndex;
-		}
+		GlobalMetaArchive.Serialize(GlobalImportNamesArchive.GetData(), GlobalImportNamesByteCount);
 
-		InitialLoadArchive->Serialize(ScriptArcsArchive.GetWriterData(), ScriptArcsArchive.Tell());
+		GlobalMetaArchive << CulturePackageMap;
 	}
 
+	UE_LOG(LogIoStore, Display, TEXT("Calculating hashes"));
 	{
-		IOSTORE_CPU_SCOPE(SerializeStoreToc);
+		IOSTORE_CPU_SCOPE(CalculateHashes);
 
-		for (FPackage* Package : Packages)
+		int32 TotalFileCount = 0;
+		for (FContainerTargetSpec& ContainerTarget : ContainerTargets)
 		{
-			Package->FirstExportBundleMetaEntry = ExportBundleMetaEntries.Num();
-			ExportBundleMetaEntries.AddDefaulted(Package->ExportBundles.Num());
-
-			NameMapBuilder.MarkNameAsReferenced(Package->Name);
-			NameMapBuilder.SerializeName(*StoreTocArchive, Package->Name);
-			*StoreTocArchive << Package->SourceGlobalPackageId;
-			*StoreTocArchive << Package->ExportCount;
-			int32 ExportBundleCount = Package->ExportBundles.Num();
-			*StoreTocArchive << ExportBundleCount;
-			*StoreTocArchive << Package->FirstExportBundleMetaEntry;
-			*StoreTocArchive << Package->FirstGlobalImport;
-			*StoreTocArchive << Package->GlobalImportCount;
-			int32 ImportedPackagesCount = Package->ImportedPackagesSerializeCount;
-			*StoreTocArchive << ImportedPackagesCount;
-			int32 ImportedPackagesOffset = ImportedPackagesArchive->Tell();
-			*StoreTocArchive << ImportedPackagesOffset;
-
-			for (int32 I = 0; I < Package->ImportedPackagesSerializeCount; ++I)
+			TotalFileCount += ContainerTarget.TargetFiles.Num();
+		}
+		TArray<FContainerTargetFile*> AllTargetFiles;
+		AllTargetFiles.Reserve(TotalFileCount);
+		for (FContainerTargetSpec& ContainerTarget : ContainerTargets)
+		{
+			for (FContainerTargetFile& TargetFile : ContainerTarget.TargetFiles)
 			{
-				FPackage* ImportedPackage = Package->ImportedPackages[I];
-				*ImportedPackagesArchive << ImportedPackage->GlobalPackageId;
+				AllTargetFiles.Add(&TargetFile);
 			}
 		}
-	}
-
-	UE_LOG(LogIoStore, Display, TEXT("Finalizing package headers..."));
-	{
-		IOSTORE_CPU_SCOPE(FinalizePackageHeaders);
+		TArray<FObjectExport>& ObjectExports = PackageAssetData.ObjectExports;
+		TAtomic<uint64> CurrentFileIndex{ 0 };
+		ParallelFor(TotalFileCount, [&AllTargetFiles, &ObjectExports, &CurrentFileIndex](int32 Index)
+		{
+			TRACE_CPUPROFILER_EVENT_SCOPE(HashFile);
+			FContainerTargetFile* TargetFile = AllTargetFiles[Index];
+			IFileHandle* FileHandle = FPlatformFileManager::Get().GetPlatformFile().OpenRead(*TargetFile->NormalizedSourcePath);
+			check(FileHandle);
+			FIoBuffer IoBuffer(TargetFile->SourceSize);
+			bool bSuccess = FileHandle->Read(IoBuffer.Data(), TargetFile->SourceSize);
+			check(bSuccess);
+			delete FileHandle;
+			if (!TargetFile->bIsBulkData)
+			{
+				IoBuffer = CreateExportBundleBuffer(TargetFile->Package, ObjectExports, IoBuffer);
+				TargetFile->TargetSize = IoBuffer.DataSize();
+			}
+			else
+			{
+				TargetFile->TargetSize = TargetFile->SourceSize;
+			}
+			FSHA1::HashBuffer(IoBuffer.Data(), IoBuffer.DataSize(), TargetFile->Hash.Hash);
+#if OUTPUT_DEBUG_PACKAGE_HASHES
+			if (!TargetFile->bIsBulkData)
+			{
+				TargetFile->Package->Hashes.ExportBundleHash = TargetFile->Hash;
+			}
+#endif
+			uint64 LocalFileIndex = CurrentFileIndex.IncrementExchange() + 1;
+			UE_CLOG(LocalFileIndex % 1000 == 0, LogIoStore, Display, TEXT("Hashing %d/%d: '%s'"), LocalFileIndex, AllTargetFiles.Num(), *TargetFile->NormalizedSourcePath);
+		}, EParallelForFlags::Unbalanced);
 
-		for (FPackage* Package : Packages)
+#if OUTPUT_DEBUG_PACKAGE_HASHES
+		ParallelFor(Packages.Num(), [&Packages](int32 Index)
 		{
-			FinalizePackageHeader(
-				Package,
-				NameMapBuilder,
-				PackageAssetData.ObjectExports,
-				GlobalPackageData.Exports.Objects,
-				GlobalPackageData.Imports.ObjectsByFullName,
-				ExportBundleMetaEntries.GetData());
-		}
+			FPackage* Package = Packages[Index];
+			
+			FString FileName = Package->FileName;
+			IFileHandle* FileHandle = FPlatformFileManager::Get().GetPlatformFile().OpenRead(*FileName);
+			check(FileHandle);
+			FIoBuffer IoBuffer(FileHandle->Size());
+			bool bSuccess = FileHandle->Read(IoBuffer.Data(), IoBuffer.DataSize());
+			check(bSuccess);
+			delete FileHandle;
+			FSHA1::HashBuffer(IoBuffer.Data(), IoBuffer.DataSize(), Package->Hashes.UAssetHash.Hash);
+			
+			FileName = FPaths::ChangeExtension(FileName, TEXT(".uexp"));
+			FileHandle = FPlatformFileManager::Get().GetPlatformFile().OpenRead(*FileName);
+			check(FileHandle);
+			IoBuffer = FIoBuffer(FileHandle->Size());
+			bSuccess = FileHandle->Read(IoBuffer.Data(), IoBuffer.DataSize());
+			check(bSuccess);
+			delete FileHandle;
+			FSHA1::HashBuffer(IoBuffer.Data(), IoBuffer.DataSize(), Package->Hashes.UExpHash.Hash);
+
+		}, EParallelForFlags::Unbalanced);
+#endif
 	}
 
 	UE_LOG(LogIoStore, Display, TEXT("Creating disk layout..."));
-	CreateDiskLayout(Packages, PackageOrderMap, CookerFileOpenOrderMap);
+	CreateDiskLayout(ContainerTargets, Packages, Arguments.GameOrderMap, Arguments.CookerOrderMap, GeneralIoWriterSettings.CompressionBlockSize, Arguments.MemoryMappingAlignment);
 
 	UE_LOG(LogIoStore, Display, TEXT("Serializing container(s)..."));
+	TUniquePtr<FIoStoreWriterContext> IoStoreWriterContext(new FIoStoreWriterContext());
+	TArray<FIoStoreWriter*> IoStoreWriters;
+	FIoStoreEnvironment GlobalIoStoreEnv;
+	FIoStoreWriter* GlobalIoStoreWriter;
+	{
+		IOSTORE_CPU_SCOPE(InitializeIoStoreWriters);
+		FString GlobalOutputPath = FString(Arguments.OutputDir) / TEXT("global");
+		GlobalIoStoreEnv.InitializeFileEnvironment(*GlobalOutputPath);
+		GlobalIoStoreWriter = new FIoStoreWriter(GlobalIoStoreEnv);
+		IoStoreWriters.Add(GlobalIoStoreWriter);
+		for (FContainerTargetSpec& ContainerTarget : ContainerTargets)
+		{
+			if (ContainerTarget.OutputPath.IsEmpty())
+			{
+				ContainerTarget.IoStoreWriter = GlobalIoStoreWriter;
+			}
+			else
+			{
+				ContainerTarget.IoStoreEnv.Reset(new FIoStoreEnvironment());
+				ContainerTarget.IoStoreEnv->InitializeFileEnvironment(ContainerTarget.OutputPath);
+				ContainerTarget.IoStoreWriter = new FIoStoreWriter(*ContainerTarget.IoStoreEnv);
+				IoStoreWriters.Add(ContainerTarget.IoStoreWriter);
+			}
+		}
+		FIoStatus IoStatus = IoStoreWriterContext->Initialize(GeneralIoWriterSettings);
+		check(IoStatus.IsOk());
+		IoStatus = GlobalIoStoreWriter->Initialize(*IoStoreWriterContext, /* bIsCompressed */ false);
+		check(IoStatus.IsOk());
+		for (FContainerTargetSpec& ContainerTarget : ContainerTargets)
+		{
+			if (ContainerTarget.IoStoreWriter != GlobalIoStoreWriter)
+			{
+				IoStatus = ContainerTarget.IoStoreWriter->Initialize(*IoStoreWriterContext, ContainerTarget.bIsCompressed);
+				check(IoStatus.IsOk());
+			}
+		}
+	}
 	{
 		IOSTORE_CPU_SCOPE(SerializeContainers);
+		
+		struct FReadFileTask
+		{
+			const FContainerTargetFile* ContainerTargetFile = nullptr;
+			FIoStoreWriter* IoStoreWriter = nullptr;
+			IAsyncReadFileHandle* FileHandle = nullptr;
+			IAsyncReadRequest* ReadRequest = nullptr;
+			FIoBuffer IoBuffer;
+			TAtomic<bool> bStarted{ false };
+		};
 
 		int32 TotalFileCount = 0;
 		for (FContainerTargetSpec& ContainerTarget : ContainerTargets)
 		{
 			TotalFileCount += ContainerTarget.TargetFiles.Num();
 		}
-		int32 CurrentFileIndex = 0;
 
-		ParallelFor(ContainerTargets.Num(),[
-			&CurrentFileIndex,
-			&TotalFileCount,
-			&ContainerTargets,
-			&PackageAssetData,
-			bWithBulkDataManifest,
-			&BulkDataManifest](int Index)
+		TArray<FReadFileTask> ReadFileTasks;
+		ReadFileTasks.SetNum(TotalFileCount);
+		int32 ReadFileTaskIndex = 0;
+		for (FContainerTargetSpec& ContainerTarget : ContainerTargets)
 		{
-			FContainerTargetSpec& ContainerTarget = ContainerTargets[Index];
+			for (FContainerTargetFile& TargetFile : ContainerTarget.TargetFiles)
+			{
+				FReadFileTask& ReadFileTask = ReadFileTasks[ReadFileTaskIndex++];
+				ReadFileTask.ContainerTargetFile = &TargetFile;
+				ReadFileTask.IoStoreWriter = ContainerTarget.IoStoreWriter;
+			}
+		}
 
-			Algo::Sort(ContainerTarget.TargetFiles, [](const FContainerTargetFile& A, const FContainerTargetFile& B)
+		FEvent* TaskStartedEvent = FPlatformProcess::GetSynchEventFromPool(false);
+		FEvent* TaskFinishedEvent = FPlatformProcess::GetSynchEventFromPool(false);
+		
+		TAtomic<uint64> BufferMemoryAllocated{ 0 };
+		TAtomic<uint64> OpenFileHandlesCount{ 0 };
+
+		TArray<FObjectExport>& ObjectExports = PackageAssetData.ObjectExports;
+		TFuture<void> ReaderTask = Async(EAsyncExecution::ThreadPool, [&ReadFileTasks, &BufferMemoryAllocated, &OpenFileHandlesCount, TaskStartedEvent, TaskFinishedEvent, &BulkDataManifest, &ObjectExports, &GeneralIoWriterSettings]()
+		{
+			IOSTORE_CPU_SCOPE(ReadContainerSourceFilesTask);
+			for (FReadFileTask& ReadFileTask : ReadFileTasks)
 			{
-				if (A.bIsBulkData != B.bIsBulkData)
+				while (!ReadFileTask.bStarted)
 				{
-					return B.bIsBulkData;
+					TRACE_CPUPROFILER_EVENT_SCOPE(WaitUntilStarted);
+					TaskStartedEvent->Wait();
 				}
+				ReadFileTask.ReadRequest->WaitCompletion();
+				delete ReadFileTask.ReadRequest;
+				delete ReadFileTask.FileHandle;
 
-				return A.Package->DiskLayoutOrder < B.Package->DiskLayoutOrder;
-			});
-			for (FContainerTargetFile& TargetFile : ContainerTarget.TargetFiles)
-			{
-				UE_CLOG(CurrentFileIndex % 1000 == 0, LogIoStore, Display, TEXT("Serializing %d/%d: '%s'"), CurrentFileIndex, TotalFileCount, *TargetFile.NormalizedSourcePath);
-				++CurrentFileIndex;
-				if (TargetFile.bIsBulkData)
-				{
-					if (bWithBulkDataManifest)
-					{
-						IOSTORE_CPU_SCOPE_DATA(SerializeFile, TCHAR_TO_ANSI(*TargetFile.NormalizedSourcePath));
+				--OpenFileHandlesCount;
 
-						FPackageStoreBulkDataManifest::EBulkdataType BulkDataType =
-							TargetFile.bIsOptionalBulkData ?
-							FPackageStoreBulkDataManifest::EBulkdataType::Optional :
-							FPackageStoreBulkDataManifest::EBulkdataType::Normal;
+				const FContainerTargetFile& TargetFile = *ReadFileTask.ContainerTargetFile;
 
-						WriteBulkData(TargetFile.NormalizedSourcePath, BulkDataType, *TargetFile.Package, BulkDataManifest, ContainerTarget.IoStoreWriter);
-					}
-				}
-				else
+				uint64 BufferSize = ReadFileTask.IoBuffer.DataSize();
+				if (!TargetFile.bIsBulkData)
 				{
-					IOSTORE_CPU_SCOPE_DATA(SerializeFile, TCHAR_TO_ANSI(*TargetFile.Package->FileName));
+					ReadFileTask.IoBuffer = CreateExportBundleBuffer(TargetFile.Package, ObjectExports, ReadFileTask.IoBuffer);
+				}
 
-					SerializePackageData(TargetFile.Package, PackageAssetData.ObjectExports, ContainerTarget.IoStoreWriter);
+				if (TargetFile.Padding)
+				{
+					ReadFileTask.IoStoreWriter->AppendPadding(TargetFile.Padding);
 				}
+				FIoWriteOptions WriteOptions;
+				WriteOptions.DebugName = *TargetFile.TargetPath;
+				WriteOptions.bForceUncompressed = TargetFile.bForceUncompressed;
+				WriteOptions.Alignment = TargetFile.Alignment;
+				FIoStatus Status = ReadFileTask.IoStoreWriter->Append(TargetFile.ChunkId, ReadFileTask.IoBuffer, WriteOptions);
+				UE_CLOG(!Status.IsOk(), LogIoStore, Fatal, TEXT("Failed to append chunk to container file due to '%s'"), *Status.ToString());
+
+				ReadFileTask.IoBuffer = FIoBuffer();
+				BufferMemoryAllocated -= BufferSize;
+				TaskFinishedEvent->Trigger();
 			}
 		});
+
+		uint64 CurrentFileIndex = 0;
+		const uint64 MaxReadFileBufferSize = 2ull << 30;
+		const uint64 MaxReadFileHandlesCount = 65536;
+		for (FReadFileTask& ReadFileTask : ReadFileTasks)
+		{
+			while (BufferMemoryAllocated > MaxReadFileBufferSize)
+			{
+				TRACE_CPUPROFILER_EVENT_SCOPE(WaitForBufferMemory);
+				TaskFinishedEvent->Wait();
+			}
+			while (OpenFileHandlesCount > MaxReadFileHandlesCount)
+			{
+				TRACE_CPUPROFILER_EVENT_SCOPE(WaitForFileHandle);
+				TaskFinishedEvent->Wait();
+			}
+			BufferMemoryAllocated += ReadFileTask.ContainerTargetFile->SourceSize;
+			ReadFileTask.IoBuffer = FIoBuffer(ReadFileTask.ContainerTargetFile->SourceSize);
+			++OpenFileHandlesCount;
+			ReadFileTask.FileHandle = FPlatformFileManager::Get().GetPlatformFile().OpenAsyncRead(*ReadFileTask.ContainerTargetFile->NormalizedSourcePath);
+			ReadFileTask.ReadRequest = ReadFileTask.FileHandle->ReadRequest(0, ReadFileTask.ContainerTargetFile->SourceSize, AIOP_Normal, nullptr, ReadFileTask.IoBuffer.Data());
+			ReadFileTask.bStarted = true;
+			TaskStartedEvent->Trigger();
+
+			++CurrentFileIndex;
+			UE_CLOG(CurrentFileIndex % 1000 == 0, LogIoStore, Display, TEXT("Serializing %d/%d: '%s'"), CurrentFileIndex, ReadFileTasks.Num(), *ReadFileTask.ContainerTargetFile->NormalizedSourcePath);
+		}
+		{
+			IOSTORE_CPU_SCOPE(CompleteReads);
+			ReaderTask.Wait();
+		}
+		FPlatformProcess::ReturnSynchEventToPool(TaskFinishedEvent);
+		FPlatformProcess::ReturnSynchEventToPool(TaskStartedEvent);
 	}
 
-	int32 StoreTocByteCount = StoreTocArchive->TotalSize();
-	int32 ImportedPackagesByteCount = ImportedPackagesArchive->TotalSize();
-	int32 GlobalImportNamesByteCount = GlobalImportNamesArchive->TotalSize();
-	int32 ExportBundleMetaByteCount = ExportBundleMetaEntries.Num() * sizeof(FExportBundleMetaEntry);
 	{
-		IOSTORE_CPU_SCOPE(SerializeGlobalMetaData);
+		IOSTORE_CPU_SCOPE(MapAdditionalBulkDataChunkIds);
+		for (FContainerTargetSpec& ContainerTarget : ContainerTargets)
+		{
+			for (FContainerTargetFile& TargetFile : ContainerTarget.TargetFiles)
+			{
+				if (TargetFile.bIsBulkData)
+				{
+					MapAdditionalBulkDataChunks(TargetFile, BulkDataManifest, ContainerTarget.IoStoreWriter);
+				}
+			}
+		}
+	}
 
+	{
 		UE_LOG(LogIoStore, Display, TEXT("Saving global meta data to container file"));
-		FLargeMemoryWriter GlobalMetaArchive(0, true);
-
-		GlobalMetaArchive << StoreTocByteCount;
-		GlobalMetaArchive.Serialize(StoreTocArchive->GetData(), StoreTocByteCount);
-
-		GlobalMetaArchive << ImportedPackagesByteCount;
-		GlobalMetaArchive.Serialize(ImportedPackagesArchive->GetData(), ImportedPackagesByteCount);
-
-		GlobalMetaArchive << GlobalImportNamesByteCount;
-		GlobalMetaArchive.Serialize(GlobalImportNamesArchive->GetData(), GlobalImportNamesByteCount);
-
-		GlobalMetaArchive << ExportBundleMetaByteCount;
-		GlobalMetaArchive.Serialize(ExportBundleMetaEntries.GetData(), ExportBundleMetaByteCount);
-
-		GlobalMetaArchive << CulturePackageMap;
-
-#if !SKIP_WRITE_CONTAINER
-		const FIoStatus Status = GlobalIoStoreWriter->Append(CreateIoChunkId(0, 0, EIoChunkType::LoaderGlobalMeta), FIoBuffer(FIoBuffer::Wrap, GlobalMetaArchive.GetData(), GlobalMetaArchive.TotalSize()), TEXT("LoaderGlobalMeta"));
+		FIoWriteOptions WriteOptions;
+		WriteOptions.DebugName = TEXT("LoaderGlobalMeta");
+		const FIoStatus Status = GlobalIoStoreWriter->Append(CreateIoChunkId(0, 0, EIoChunkType::LoaderGlobalMeta), FIoBuffer(FIoBuffer::Wrap, GlobalMetaArchive.GetData(), GlobalMetaArchive.TotalSize()), WriteOptions);
 		if (!Status.IsOk())
 		{
 			UE_LOG(LogIoStore, Error, TEXT("Failed to save global meta data to container file"));
 		}
-#endif
 	}
 
 	{
 		UE_LOG(LogIoStore, Display, TEXT("Saving initial load meta data to container file"));
-#if !SKIP_WRITE_CONTAINER
-		const FIoStatus Status = GlobalIoStoreWriter->Append(CreateIoChunkId(0, 0, EIoChunkType::LoaderInitialLoadMeta), FIoBuffer(FIoBuffer::Wrap, InitialLoadArchive->GetData(), InitialLoadArchive->TotalSize()), TEXT("LoaderInitialLoadMeta"));
+		FIoWriteOptions WriteOptions;
+		WriteOptions.DebugName = TEXT("LoaderInitialLoadMeta");
+		const FIoStatus Status = GlobalIoStoreWriter->Append(CreateIoChunkId(0, 0, EIoChunkType::LoaderInitialLoadMeta), FIoBuffer(FIoBuffer::Wrap, InitialLoadArchive.GetData(), InitialLoadArchive.TotalSize()), WriteOptions);
 		if (!Status.IsOk())
 		{
 			UE_LOG(LogIoStore, Error, TEXT("Failed to save initial load meta data to container file"));
 		}
-#endif
 	}
 
 	uint64 GlobalNamesMB = 0;
@@ -2769,12 +3505,13 @@ int32 CreateTarget(
 		GlobalNamesMB = Names.Num() >> 20;
 		GlobalNameHashesMB = Hashes.Num() >> 20;
 
-#if !SKIP_WRITE_CONTAINER
+		FIoWriteOptions WriteOptions;
+		WriteOptions.DebugName = TEXT("LoaderGlobalNames");
 		FIoStatus NameStatus = GlobalIoStoreWriter->Append(CreateIoChunkId(0, 0, EIoChunkType::LoaderGlobalNames), 
-													 FIoBuffer(FIoBuffer::Wrap, Names.GetData(), Names.Num()), TEXT("LoaderGlobalNames"));
+													 FIoBuffer(FIoBuffer::Wrap, Names.GetData(), Names.Num()), WriteOptions);
+		WriteOptions.DebugName = TEXT("LoaderGlobalNameHashes");
 		FIoStatus HashStatus = GlobalIoStoreWriter->Append(CreateIoChunkId(0, 0, EIoChunkType::LoaderGlobalNameHashes),
-													 FIoBuffer(FIoBuffer::Wrap, Hashes.GetData(), Hashes.Num()), TEXT("LoaderGlobalNameHashes"));
-#endif
+													 FIoBuffer(FIoBuffer::Wrap, Hashes.GetData(), Hashes.Num()), WriteOptions);
 		
 		if (!NameStatus.IsOk() || !HashStatus.IsOk())
 		{
@@ -2786,12 +3523,98 @@ int32 CreateTarget(
 #endif
 	}
 
+	TArray<FIoStoreWriterResult> IoStoreWriterResults;
+	IoStoreWriterResults.Reserve(IoStoreWriters.Num());
 	for (FIoStoreWriter* IoStoreWriter : IoStoreWriters)
 	{
+		IoStoreWriterResults.Emplace(IoStoreWriter->Flush().ConsumeValueOrDie());
 		delete IoStoreWriter;
 	}
 	IoStoreWriters.Empty();
 
+	if (!Arguments.OutputReleaseVersionDir.IsEmpty())
+	{
+		SaveReleaseVersionMeta(*Arguments.OutputReleaseVersionDir, NameMapBuilder, PackageGlobalIdMap, GlobalPackageData.Imports, ContainerTargets);
+#if OUTPUT_DEBUG_PACKAGE_HASHES
+		FString PackageHashesOutputPath = FPaths::Combine(*Arguments.OutputReleaseVersionDir, TEXT("iodispatcher.upackagehashes"));
+		TUniquePtr<FArchive> PackageHashesArchive(IFileManager::Get().CreateFileWriter(*PackageHashesOutputPath));
+
+		int32 PackageCount = Packages.Num();
+		(*PackageHashesArchive) << PackageCount;
+		for (FPackage* Package : Packages)
+		{
+			const FNameEntry* NameEntry = FName::GetEntry(Package->Name.GetComparisonIndex());
+			NameEntry->Write(*PackageHashesArchive);
+			int32 NameNumber = Package->Name.GetNumber();
+			(*PackageHashesArchive) << NameNumber;
+			(*PackageHashesArchive) << Package->Hashes.UAssetHash;
+			(*PackageHashesArchive) << Package->Hashes.UExpHash;
+			(*PackageHashesArchive) << Package->Hashes.ExportBundleHash;
+		}
+#endif
+	}
+#if OUTPUT_DEBUG_PACKAGE_HASHES
+	{
+		uint64 AddedCount = 0;
+		uint64 DeletedCount = 0;
+		uint64 ModifiedCountUAssetOrUExpCount = 0;
+		uint64 ModifiedUAssetCount = 0;
+		uint64 ModifiedUAssetSize = 0;
+		uint64 ModifiedUExpCount = 0;
+		uint64 ModifiedUExpSize = 0;
+		uint64 ModifiedExportBundlesCount = 0;
+		uint64 ModifiedExportBundlesSize = 0;
+		for (FPackage* Package : Packages)
+		{
+			const FPackageHashes* PreviousHashes = PreviousBuildPackageHashes.Find(Package->Name);
+			if (PreviousHashes)
+			{
+				bool bModified = false;
+				if (PreviousHashes->UAssetHash != Package->Hashes.UAssetHash)
+				{
+					++ModifiedUAssetCount;
+					bModified = true;
+					
+					ModifiedUAssetSize += Package->UAssetSize;
+				}
+				if (PreviousHashes->UExpHash != Package->Hashes.UExpHash)
+				{
+					++ModifiedUExpCount;
+					bModified = true;
+
+					ModifiedUExpSize += Package->UExpSize;
+				}
+				if (PreviousHashes->ExportBundleHash != Package->Hashes.ExportBundleHash)
+				{
+					++ModifiedExportBundlesCount;
+					ModifiedExportBundlesSize += Package->HeaderData.Num() + Package->ExportsSerialSize;
+				}
+				if (bModified)
+				{
+					++ModifiedCountUAssetOrUExpCount;
+				}
+			}
+			else
+			{
+				++AddedCount;
+			}
+		}
+		for (const auto& KV : PreviousBuildPackageHashes)
+		{
+			if (!PackageMap.Contains(KV.Key))
+			{
+				++DeletedCount;
+			}
+		}
+		UE_LOG(LogIoStore, Display, TEXT("Added packages: %d"), AddedCount);
+		UE_LOG(LogIoStore, Display, TEXT("Deleted packages: %d"), DeletedCount);
+		UE_LOG(LogIoStore, Display, TEXT("Modified packages (export bundle): %d, %fMB"), ModifiedExportBundlesCount, ModifiedExportBundlesSize / 1024.0 / 1024.0);
+		UE_LOG(LogIoStore, Display, TEXT("Modified packages (uasset|uexp): %d, %fMB"), ModifiedCountUAssetOrUExpCount, (ModifiedUAssetSize + ModifiedUExpSize) / 1024.0 / 1024.0);
+		UE_LOG(LogIoStore, Display, TEXT("Modified packages (uasset): %d, %fMB"), ModifiedUAssetCount, ModifiedUAssetSize / 1024.0 / 1024.0);
+		UE_LOG(LogIoStore, Display, TEXT("Modified packages (uexp): %d, %fMB"), ModifiedUExpCount, ModifiedUExpSize / 1024.0 / 1024.0);
+	}
+#endif
+
 	IOSTORE_CPU_SCOPE(CalculateStats);
 
 	UE_LOG(LogIoStore, Display, TEXT("Calculating stats..."));
@@ -2805,8 +3628,9 @@ int32 CreateTarget(
 	uint64 NameMapCount = 0;
 	uint64 PackageSummarySize = Packages.Num() * sizeof(FPackageSummary);
 	uint64 ImportedPackagesCount = 0;
-	uint64 InitialLoadSize = InitialLoadArchive->Tell();
-	uint64 ScriptArcsCount = 0;
+	uint64 PublicExportsCount = 0;
+	uint64 ExportBundlesMetaCount = 0;
+	uint64 InitialLoadSize = InitialLoadArchive.Tell();
 	uint64 CircularPackagesCount = 0;
 	uint64 TotalInternalArcCount = 0;
 	uint64 TotalExternalArcCount = 0;
@@ -2819,12 +3643,6 @@ int32 CreateTarget(
 
 	uint64 PackageHeaderSize = 0;
 
-	uint64 UniqueImportPackages = 0;
-	for (const FImportData& ImportData : GlobalImports)
-	{
-		UniqueImportPackages += (ImportData.OuterIndex == 0);
-	}
-
 	for (auto& PackageKV : PackageMap)
 	{
 		FPackage& Package = *PackageKV.Value;
@@ -2837,10 +3655,11 @@ int32 CreateTarget(
 		ExportMapSize += Package.ExportMapSize;
 		NameMapSize += Package.NameMapSize;
 		NameMapCount += Package.NameIndices.Num();
-		ScriptArcsCount += Package.ScriptArcs.Num();
 		CircularPackagesCount += Package.bHasCircularImportDependencies;
 		TotalInternalArcCount += Package.InternalArcs.Num();
-		ImportedPackagesCount += Package.ImportedPackages.Num();
+		ImportedPackagesCount += Package.ImportedPackagesSerializeCount;
+		PublicExportsCount += Package.PublicExports.Num();
+		ExportBundlesMetaCount += Package.ExportBundles.Num();
 		NameCount += Package.NameMap.Num();
 		PackagesWithoutPreloadDependenciesCount += Package.ImportedPreloadPackages.Num() == 0;
 		PackagesWithoutImportDependenciesCount += Package.ImportedPackages.Num() == 0;
@@ -2860,19 +3679,46 @@ int32 CreateTarget(
 
 	PackageHeaderSize = PackageSummarySize + NameMapSize + ImportMapSize + ExportMapSize + UGraphSize;
 
-	UE_LOG(LogIoStore, Display, TEXT("-------------------- IoStore Summary --------------------"));
+	UE_LOG(LogIoStore, Display, TEXT("--------------------------------------------- IoStore Summary -----------------------------------------------"));
+
+	UE_LOG(LogIoStore, Display, TEXT(""));
+	UE_LOG(LogIoStore, Display, TEXT("%-35s %15s %15s %15s %25s"), TEXT("Container"), TEXT("TOC Size (KB)"), TEXT("TOC Entries"), TEXT("Size (MB)"), TEXT("Compressed (MB)"));
+	UE_LOG(LogIoStore, Display, TEXT("-------------------------------------------------------------------------------------------------------------"));
+	int64 TotalPaddingSize = 0;
+	for (const FIoStoreWriterResult& Result : IoStoreWriterResults)
+	{
+		const double Compression = Result.CompressionMethod == NAME_None
+			? 0.0
+			: (double(Result.UncompressedContainerSize - Result.CompressedContainerSize) / double(Result.UncompressedContainerSize)) * 100.0;
+
+		UE_LOG(LogIoStore, Display, TEXT("%-35s %15.2lf %15d %15.2lf %25s"),
+			*Result.ContainerName,
+			(double)Result.TocSize / 1024.0,
+			Result.TocEntryCount,
+			(double)Result.UncompressedContainerSize / 1024.0 / 1024.0,
+			*FString::Printf(TEXT("%.2lf (%.2lf%% %s)"),
+				(double)Result.CompressedContainerSize / 1024.0 / 1024.0,
+				Compression,
+				*Result.CompressionMethod.ToString()));
+		TotalPaddingSize += Result.PaddingSize;
+	}
+	UE_LOG(LogIoStore, Display, TEXT("Padding: %8.2f MB"), TotalPaddingSize / 1024.0 / 1024.0);
+	UE_LOG(LogIoStore, Display, TEXT(""));
 	UE_LOG(LogIoStore, Display, TEXT("Packages: %8d total, %d circular dependencies, %d no preload dependencies, %d no import dependencies"),
 		PackageMap.Num(), CircularPackagesCount, PackagesWithoutPreloadDependenciesCount, PackagesWithoutImportDependenciesCount);
 	UE_LOG(LogIoStore, Display, TEXT("Bundles:  %8d total, %d entries, %d export objects"), BundleCount, BundleEntryCount, GlobalPackageData.Exports.Objects.Num());
 
 	UE_LOG(LogIoStore, Display, TEXT("IoStore: %8.2f MB GlobalNames, %d unique names"), (double)GlobalNamesMB, NameMapBuilder.GetNameMap().Num());
 	UE_LOG(LogIoStore, Display, TEXT("IoStore: %8.2f MB GlobalNameHashes"), (double)GlobalNameHashesMB);
-	UE_LOG(LogIoStore, Display, TEXT("IoStore: %8.2f MB GlobalPackageData"), (double)StoreTocByteCount / 1024.0 / 1024.0);
-	UE_LOG(LogIoStore, Display, TEXT("IoStore: %8.2f MB GlobalImportedPackages, %d imported packages"), (double)ImportedPackagesByteCount / 1024.0 / 1024.0, ImportedPackagesCount);
-	UE_LOG(LogIoStore, Display, TEXT("IoStore: %8.2f MB GlobalBundleMeta, %d bundles"), (double)ExportBundleMetaByteCount / 1024.0 / 1024.0, ExportBundleMetaEntries.Num());
-	UE_LOG(LogIoStore, Display, TEXT("IoStore: %8.2f MB GlobalImportNames, %d total imports, %d script imports, %d UPackage imports"),
-		(double)GlobalImportNamesByteCount / 1024.0 / 1024.0, GlobalImportsByFullName.Num(), NumScriptImports, UniqueImportPackages);
-	UE_LOG(LogIoStore, Display, TEXT("IoStore: %8.2f MB InitialLoadData, %d script arcs, %d script outers, %d packages"), (double)InitialLoadSize / 1024.0 / 1024.0, ScriptArcsCount, NumScriptImports, Packages.Num());
+	UE_LOG(LogIoStore, Display, TEXT("IoStore: %8.2f MB GlobalPackageStoreToc"), (double)StoreTocByteCount / 1024.0 / 1024.0);
+	UE_LOG(LogIoStore, Display, TEXT("IoStore: %8.2f MB GlobalPackageStoreData, %d imported packages, %d public exports, %d export bundles"),
+		(double)StoreDataByteCount / 1024.0 / 1024.0,
+		ImportedPackagesCount,
+		PublicExportsCount,
+		ExportBundlesMetaCount);
+	UE_LOG(LogIoStore, Display, TEXT("IoStore: %8.2f MB PackageImportNames, %d imports"),
+		(double)GlobalImportNamesByteCount / 1024.0 / 1024.0, PackageImports.Num());
+	UE_LOG(LogIoStore, Display, TEXT("IoStore: %8.2f MB InitialLoadData, %d script imports"), (double)InitialLoadSize / 1024.0 / 1024.0, ScriptImports.Num());
 	UE_LOG(LogIoStore, Display, TEXT("IoStore: %8.2f MB PackageHeader, %d packages"), (double)PackageHeaderSize / 1024.0 / 1024.0, Packages.Num());
 	UE_LOG(LogIoStore, Display, TEXT("IoStore: %8.2f MB PackageSummary"), (double)PackageSummarySize / 1024.0 / 1024.0);
 	UE_LOG(LogIoStore, Display, TEXT("IoStore: %8.2f MB PackageNameMap, %d indices"), (double)NameMapSize / 1024.0 / 1024.0, NameMapCount);
@@ -2927,11 +3773,19 @@ static bool ParsePakResponseFile(const TCHAR* FilePath, TArray<FContainerSourceF
 
 		FContainerSourceFile& FileEntry = OutFiles.AddDefaulted_GetRef();
 		FileEntry.NormalizedPath = MoveTemp(SourceAndDest[0]);
+
+		for (int32 Index = 0; Index < Switches.Num(); ++Index)
+		{
+			if (Switches[Index] == TEXT("compress"))
+			{
+				FileEntry.bNeedsCompression = true;
+			}
+		}
 	}
 	return true;
 }
 
-static bool ParsePakOrderFile(const TCHAR* FilePath, TMap<FString, uint64>& OutMap)
+static bool ParsePakOrderFile(const TCHAR* FilePath, TMap<FName, uint64>& OutMap)
 {
 	TArray<FString> OrderFileContents;
 	if (!FFileHelper::LoadFileToStringArray(OrderFileContents, FilePath))
@@ -2950,7 +3804,12 @@ static bool ParsePakOrderFile(const TCHAR* FilePath, TMap<FString, uint64>& OutM
 			UE_LOG(LogIoStore, Error, TEXT("Invalid line in order file '%s'."), *OrderLine);
 			return false;
 		}
-		FPaths::NormalizeFilename(Path);
+		FString PackageName;
+		if (!FPackageName::TryConvertFilenameToLongPackageName(Path, PackageName, nullptr))
+		{
+			continue;;
+		}
+
 		uint64 Order = LineNumber;
 		FString OrderStr;
 		if (FParse::Token(OrderLinePtr, OrderStr, false))
@@ -2963,9 +3822,10 @@ static bool ParsePakOrderFile(const TCHAR* FilePath, TMap<FString, uint64>& OutM
 			Order = FCString::Atoi64(*OrderStr);
 		}
 
-		if (!OutMap.Contains(Path))
+		FName PackageFName(MoveTemp(PackageName));
+		if (!OutMap.Contains(PackageFName))
 		{
-			OutMap.Emplace(MoveTemp(Path), Order);
+			OutMap.Emplace(PackageFName, Order);
 		}
 
 		++LineNumber;
@@ -2991,8 +3851,9 @@ public:
 	virtual bool Visit(const TCHAR* FilenameOrDirectory, const FFileStatData& StatData)
 	{
 		// Should match FCookedFileStatData::EFileExt
-		static const TCHAR* Extensions[] = { TEXT("umap"), TEXT("uasset"), TEXT("ubulk"), TEXT("uptnl") };
+		static const TCHAR* Extensions[] = { TEXT("umap"), TEXT("uasset"), TEXT("uexp"), TEXT("ubulk"), TEXT("uptnl"), TEXT("m.ubulk") };
 		static const int32 NumPackageExtensions = 2;
+		static const int32 UExpExtensionIndex = 2;
 
 		if (StatData.bIsDirectory)
 		{
@@ -3006,10 +3867,21 @@ public:
 		}
 
 		int32 ExtIndex = 0;
-		for (ExtIndex = 0; ExtIndex < UE_ARRAY_COUNT(Extensions); ++ExtIndex)
+		if (0 == FCString::Stricmp(Extension, Extensions[3]))
 		{
-			if (0 == FCString::Stricmp(Extension, Extensions[ExtIndex]))
-				break;
+			ExtIndex = 3;
+			if (0 == FCString::Stricmp(Extension - 3, TEXT(".m.ubulk")))
+			{
+				ExtIndex = 5;
+			}
+		}
+		else
+		{
+			for (ExtIndex = 0; ExtIndex < UE_ARRAY_COUNT(Extensions); ++ExtIndex)
+			{
+				if (0 == FCString::Stricmp(Extension, Extensions[ExtIndex]))
+					break;
+			}
 		}
 
 		if (ExtIndex >= UE_ARRAY_COUNT(Extensions))
@@ -3020,7 +3892,7 @@ public:
 		FString Path = FilenameOrDirectory;
 		FPaths::NormalizeFilename(Path);
 
-		if (ContainerSpec)
+		if (ContainerSpec && ExtIndex != UExpExtensionIndex)
 		{
 			FContainerSourceFile& FileEntry = ContainerSpec->SourceFiles.AddDefaulted_GetRef();
 			FileEntry.NormalizedPath = Path;
@@ -3029,11 +3901,18 @@ public:
 		FCookedFileStatData& CookedFileStatData = CookedFileStatMap.Add(MoveTemp(Path));
 		CookedFileStatData.FileSize = StatData.FileSize;
 		CookedFileStatData.FileExt = FCookedFileStatData::EFileExt(ExtIndex);
-		CookedFileStatData.FileType = 
-			ExtIndex < NumPackageExtensions ?
-			FCookedFileStatData::PackageHeader : 
-			FCookedFileStatData::BulkData;
-
+		if (ExtIndex < NumPackageExtensions)
+		{
+			CookedFileStatData.FileType = FCookedFileStatData::PackageHeader;
+		}
+		else if (ExtIndex == UExpExtensionIndex)
+		{
+			CookedFileStatData.FileType = FCookedFileStatData::PackageData;
+		}
+		else
+		{
+			CookedFileStatData.FileType = FCookedFileStatData::BulkData;
+		}
 		return true;
 	}
 };
@@ -3044,40 +3923,98 @@ int32 CreateIoStoreContainerFiles(const TCHAR* CmdLine)
 
 	UE_LOG(LogIoStore, Display, TEXT("==================== IoStore Utils ===================="));
 
-	TMap<FString, uint64> PackageOrderMap;
-	FString PackageOrderFilePath;
-	if (FParse::Value(FCommandLine::Get(), TEXT("PackageOrder="), PackageOrderFilePath))
+	FIoStoreArguments Arguments;
+	FString GameOrderFilePath;
+	if (FParse::Value(FCommandLine::Get(), TEXT("GameOrder="), GameOrderFilePath))
 	{
-		if (!ParsePakOrderFile(*PackageOrderFilePath, PackageOrderMap))
+		if (!ParsePakOrderFile(*GameOrderFilePath, Arguments.GameOrderMap))
 		{
 			return -1;
 		}
 	}
 
-	TMap<FString, uint64> CookerFileOpenOrderMap;
-	FString CookerFileOpenOrderFilePath;
-	if (FParse::Value(FCommandLine::Get(), TEXT("CookerOrder="), CookerFileOpenOrderFilePath))
+	FString CookerOrderFilePath;
+	if (FParse::Value(FCommandLine::Get(), TEXT("CookerOrder="), CookerOrderFilePath))
 	{
-		if (!ParsePakOrderFile(*CookerFileOpenOrderFilePath, CookerFileOpenOrderMap))
+		if (!ParsePakOrderFile(*CookerOrderFilePath, Arguments.CookerOrderMap))
 		{
 			return -1;
 		}
 	}
 
+	FIoStoreWriterSettings GeneralIoWriterSettings { DefaultCompressionMethod, DefaultCompressionBlockSize, false };
+	GeneralIoWriterSettings.bEnableCsvOutput = FParse::Param(CmdLine, TEXT("-csvoutput"));
+
+	if (!FParse::Value(CmdLine, TEXT("-AlignForMemoryMapping="), Arguments.MemoryMappingAlignment))
+	{
+		Arguments.MemoryMappingAlignment = DefaultMemoryMappingAlignment;
+	}
+	UE_LOG(LogIoStore, Display, TEXT("Using memory mapping alignment '%ld'"), Arguments.MemoryMappingAlignment);
+	
+	TArray<FName> CompressionFormats;
+	FString DesiredCompressionFormats;
+	if (FParse::Value(CmdLine, TEXT("-compressionformats="), DesiredCompressionFormats) ||
+		FParse::Value(CmdLine, TEXT("-compressionformat="), DesiredCompressionFormats))
+	{
+		TArray<FString> Formats;
+		DesiredCompressionFormats.ParseIntoArray(Formats, TEXT(","));
+		for (FString& Format : Formats)
+		{
+			// look until we have a valid format
+			FName FormatName = *Format;
+
+			if (FCompression::IsFormatValid(FormatName))
+			{
+				GeneralIoWriterSettings.CompressionMethod = FormatName;
+				break;
+			}
+		}
+
+		if (GeneralIoWriterSettings.CompressionMethod == NAME_None)
+		{
+			UE_LOG(LogIoStore, Warning, TEXT("Failed to find desired compression format(s) '%s'. Using falling back to '%s'"),
+				*DesiredCompressionFormats, *DefaultCompressionMethod.ToString());
+		}
+		else
+		{
+			UE_LOG(LogIoStore, Display, TEXT("Using compression format '%s'"), *GeneralIoWriterSettings.CompressionMethod.ToString());
+		}
+	}
+
+	FString CompressionBlockSizeString;
+	if (FParse::Value(CmdLine, TEXT("-compressionblocksize="), CompressionBlockSizeString))
+	{
+		FParse::Value(CmdLine, TEXT("-compressionblocksize="), GeneralIoWriterSettings.CompressionBlockSize);
+
+		if (CompressionBlockSizeString.EndsWith(TEXT("MB")))
+		{
+			GeneralIoWriterSettings.CompressionBlockSize *= 1024 * 1024;
+		}
+		else if (CompressionBlockSizeString.EndsWith(TEXT("KB")))
+		{
+			GeneralIoWriterSettings.CompressionBlockSize *= 1024;
+		}
+	}
+	UE_LOG(LogIoStore, Display, TEXT("Using compression block size '%ld'"), GeneralIoWriterSettings.CompressionBlockSize);
+
+	FParse::Value(CmdLine, TEXT("-compressionblockalignment="), GeneralIoWriterSettings.CompressionBlockAlignment);
+	UE_LOG(LogIoStore, Display, TEXT("Using compression block alignment '%ld'"), GeneralIoWriterSettings.CompressionBlockAlignment);
+
+	FParse::Value(CmdLine, TEXT("-CreateReleaseVersionDirectory="), Arguments.OutputReleaseVersionDir);
+	FParse::Value(CmdLine, TEXT("-BasedOnReleaseVersionDirectory="), Arguments.BasedOnReleaseVersionDir);
+
 	FString CommandListFile;
 	if (FParse::Value(FCommandLine::Get(), TEXT("Commands="), CommandListFile))
 	{
 		UE_LOG(LogIoStore, Display, TEXT("Using command list file: '%s'"), *CommandListFile);
 
-		FString OutputDirectory;
-		if (!FParse::Value(FCommandLine::Get(), TEXT("OutputDirectory="), OutputDirectory))
+		if (!FParse::Value(FCommandLine::Get(), TEXT("OutputDirectory="), Arguments.OutputDir))
 		{
 			UE_LOG(LogIoStore, Error, TEXT("OutputDirectory must be specified"));
 			return 1;
 		}
 
-		FString CookedDirectory;
-		if (!FParse::Value(FCommandLine::Get(), TEXT("CookedDirectory="), CookedDirectory))
+		if (!FParse::Value(FCommandLine::Get(), TEXT("CookedDirectory="), Arguments.CookedDir))
 		{
 			UE_LOG(LogIoStore, Error, TEXT("CookedDirectory must be specified"));
 			return 1;
@@ -3090,16 +4027,17 @@ int32 CreateIoStoreContainerFiles(const TCHAR* CmdLine)
 			return -1;
 		}
 
-		TArray<FContainerSourceSpec> Containers;
+		Arguments.Containers.Reserve(Commands.Num());
 		for (const FString& Command : Commands)
 		{
-			FContainerSourceSpec& ContainerSpec = Containers.AddDefaulted_GetRef();
+			FContainerSourceSpec& ContainerSpec = Arguments.Containers.AddDefaulted_GetRef();
 
 			if (!FParse::Value(*Command, TEXT("Output="), ContainerSpec.OutputPath))
 			{
 				UE_LOG(LogIoStore, Error, TEXT("Output argument missing from command '%s'"), *Command);
 				return -1;
 			}
+			ContainerSpec.OutputPath = FPaths::Combine(Arguments.OutputDir, ContainerSpec.OutputPath);
 
 			FString ResponseFilePath;
 			if (!FParse::Value(*Command, TEXT("ResponseFile="), ResponseFilePath))
@@ -3115,19 +4053,13 @@ int32 CreateIoStoreContainerFiles(const TCHAR* CmdLine)
 			}
 		}
 
-		UE_LOG(LogIoStore, Display, TEXT("Searching for cooked assets in folder '%s'"), *CookedDirectory);
+		UE_LOG(LogIoStore, Display, TEXT("Searching for cooked assets in folder '%s'"), *Arguments.CookedDir);
 		FCookedFileStatMap CookedFileStatMap;
-		FCookedFileVisitor CookedFileVistor(CookedFileStatMap);
-		IFileManager::Get().IterateDirectoryStatRecursively(*CookedDirectory, CookedFileVistor);
-		UE_LOG(LogIoStore, Display, TEXT("Found '%d' assets"), CookedFileStatMap.Num());
-
-		if (CookedFileStatMap.Num() == 0)
-		{
-			UE_LOG(LogIoStore, Error, TEXT("Failed find cooked assets in folder '%s'"), *CookedDirectory);
-			return -1;
-		}
+		FCookedFileVisitor CookedFileVistor(Arguments.CookedFileStatMap);
+		IFileManager::Get().IterateDirectoryStatRecursively(*Arguments.CookedDir, CookedFileVistor);
+		UE_LOG(LogIoStore, Display, TEXT("Found '%d' files"), Arguments.CookedFileStatMap.Num());
 
-		int32 ReturnValue = CreateTarget(*OutputDirectory, *CookedDirectory, Containers, CookedFileStatMap, PackageOrderMap, CookerFileOpenOrderMap);
+		int32 ReturnValue = CreateTarget(Arguments, GeneralIoWriterSettings);
 		if (ReturnValue != 0)
 		{
 			return ReturnValue;
@@ -3139,22 +4071,21 @@ int32 CreateIoStoreContainerFiles(const TCHAR* CmdLine)
 
 		for (const ITargetPlatform* TargetPlatform : Platforms)
 		{
-			const FString TargetCookedDirectory = FPaths::ProjectSavedDir() / TEXT("Cooked") / TargetPlatform->PlatformName();
-			const FString TargetCookedProjectDirectory = TargetCookedDirectory / FApp::GetProjectName();
+			Arguments.CookedDir = FPaths::ProjectSavedDir() / TEXT("Cooked") / TargetPlatform->PlatformName();
+			Arguments.OutputDir = Arguments.CookedDir / FApp::GetProjectName();
 
-			FCookedFileStatMap CookedFileStatMap;
 			TArray<FContainerSourceSpec> Containers;
-			FContainerSourceSpec& ContainerSpec = Containers.AddDefaulted_GetRef();
-			FCookedFileVisitor CookedFileVistor(CookedFileStatMap, &ContainerSpec);
+			FContainerSourceSpec& ContainerSpec = Arguments.Containers.AddDefaulted_GetRef();
+			FCookedFileVisitor CookedFileVistor(Arguments.CookedFileStatMap, &ContainerSpec);
 
 			UE_LOG(LogIoStore, Display, TEXT("Searching for cooked files..."));
-			IFileManager::Get().IterateDirectoryStatRecursively(*TargetCookedDirectory, CookedFileVistor);
+			IFileManager::Get().IterateDirectoryStatRecursively(*Arguments.CookedDir, CookedFileVistor);
 
 			UE_LOG(LogIoStore, Display, TEXT("Found '%d' files"), ContainerSpec.SourceFiles.Num());
 
-			UE_LOG(LogIoStore, Display, TEXT("Creating target: '%s' using output path: '%s'"), *TargetPlatform->PlatformName(), *TargetCookedProjectDirectory);
+			UE_LOG(LogIoStore, Display, TEXT("Creating target: '%s' using output path: '%s'"), *TargetPlatform->PlatformName(), *Arguments.OutputDir);
 
-			int32 ReturnValue = CreateTarget(*TargetCookedProjectDirectory, *TargetCookedDirectory, Containers, CookedFileStatMap, PackageOrderMap, CookerFileOpenOrderMap);
+			int32 ReturnValue = CreateTarget(Arguments, GeneralIoWriterSettings);
 			if (ReturnValue != 0)
 			{
 				return ReturnValue;
