Copying //UE4/Release-4.25Plus-Dev to Release-4.25Plus (//UE4/Release-4.25Plus)
#rnx

diff --git a/Engine/Source/Developer/IoStoreUtilities/Private/IoStoreUtilities.cpp b/Engine/Source/Developer/IoStoreUtilities/Private/IoStoreUtilities.cpp
--- a/Engine/Source/Developer/IoStoreUtilities/Private/IoStoreUtilities.cpp
+++ b/Engine/Source/Developer/IoStoreUtilities/Private/IoStoreUtilities.cpp
@@ -52,6 +52,150 @@ static const FName DefaultCompressionMethod = NAME_Zlib;
 static const int64 DefaultCompressionBlockSize = 64 << 10;
 const int64 DefaultMemoryMappingAlignment = 16 << 10;
 
+class FNameMapBuilder
+{
+public:
+	void SetNameMapType(FMappedName::EType InNameMapType)
+	{
+		NameMapType = InNameMapType;
+	}
+
+	void MarkNamesAsReferenced(const TArray<FName>& Names, TArray<int32>& OutNameIndices)
+	{
+		for (const FName& Name : Names)
+		{
+			const FNameEntryId Id = Name.GetComparisonIndex();
+			int32& Index = NameIndices.FindOrAdd(Id);
+			if (Index == 0)
+			{
+				Index = NameIndices.Num();
+				NameMap.Add(Id);
+			}
+
+			OutNameIndices.Add(Index - 1);
+		}
+	}
+
+	void MarkNameAsReferenced(const FName& Name)
+	{
+		const FNameEntryId Id = Name.GetComparisonIndex();
+		int32& Index = NameIndices.FindOrAdd(Id);
+		if (Index == 0)
+		{
+			Index = NameIndices.Num();
+			NameMap.Add(Id);
+		}
+#if OUTPUT_NAMEMAP_CSV
+		// debug counts
+		{
+			const int32 Number = Name.GetNumber();
+			TTuple<int32,int32,int32>& Counts = DebugNameCounts.FindOrAdd(Id);
+
+			if (Number == 0)
+			{
+				++Counts.Get<0>();
+			}
+			else
+			{
+				++Counts.Get<1>();
+				if (Number > Counts.Get<2>())
+				{
+					Counts.Get<2>() = Number;
+				}
+			}
+		}
+#endif
+	}
+
+	FMappedName MapName(const FName& Name) const
+	{
+		const FNameEntryId Id = Name.GetComparisonIndex();
+		const int32* Index = NameIndices.Find(Id);
+		check(Index);
+		return FMappedName::Create(*Index - 1, Name.GetNumber(), NameMapType);
+	}
+
+	const TArray<FNameEntryId>& GetNameMap() const
+	{
+		return NameMap;
+	}
+
+	friend FArchive& operator<<(FArchive& Ar, FNameMapBuilder& NameMapBuilder)
+	{
+		if (Ar.IsSaving())
+		{
+			int32 NameCount = NameMapBuilder.NameMap.Num();
+			Ar << NameCount;
+			for (FNameEntryId NameEntryId : NameMapBuilder.NameMap)
+			{
+				const FNameEntry* NameEntry = FName::GetEntry(NameEntryId);
+				NameEntry->Write(Ar);
+			}
+		}
+		else
+		{
+			int32 NameCount;
+			Ar << NameCount;
+			for (int32 NameIndex = 0; NameIndex < NameCount; ++NameIndex)
+			{
+				FNameEntrySerialized NameEntrySerialized(ENAME_LinkerConstructor);
+				Ar << NameEntrySerialized;
+				FName Name(NameEntrySerialized);
+				FNameEntryId NameId = Name.GetComparisonIndex();
+				NameMapBuilder.NameMap.Add(NameId);
+				NameMapBuilder.NameIndices.Add(NameId, NameIndex + 1);
+			}
+		}
+		return Ar;
+	}
+
+#if OUTPUT_NAMEMAP_CSV
+	void SaveCsv(const FString& CsvFilePath)
+	{
+		{
+			TUniquePtr<FArchive> CsvArchive(IFileManager::Get().CreateFileWriter(*CsvFilePath));
+			if (CsvArchive)
+			{
+				TCHAR Name[FName::StringBufferSize];
+				ANSICHAR Line[MAX_SPRINTF + FName::StringBufferSize];
+				ANSICHAR Header[] = "Length\tMaxNumber\tNumberCount\tBaseCount\tTotalCount\tFName\n";
+				CsvArchive->Serialize(Header, sizeof(Header) - 1);
+				for (auto& Counts : DebugNameCounts)
+				{
+					const int32 NameLen = FName::CreateFromDisplayId(Counts.Key, 0).ToString(Name);
+					FCStringAnsi::Sprintf(Line, "%d\t%d\t%d\t%d\t%d\t",
+						NameLen, Counts.Value.Get<2>(), Counts.Value.Get<1>(), Counts.Value.Get<0>(), Counts.Value.Get<0>() + Counts.Value.Get<1>());
+					ANSICHAR* L = Line + FCStringAnsi::Strlen(Line);
+					const TCHAR* N = Name;
+					while (*N)
+					{
+						*L++ = CharCast<ANSICHAR,TCHAR>(*N++);
+					}
+					*L++ = '\n';
+					CsvArchive.Get()->Serialize(Line, L - Line);
+				}
+			}
+		}
+	}
+#endif
+
+	void Empty()
+	{
+		NameIndices.Empty();
+		NameMap.Empty();
+#if OUTPUT_NAMEMAP_CSV
+		DebugNameCounts.Empty();
+#endif
+	}
+private:
+	TMap<FNameEntryId, int32> NameIndices;
+	TArray<FNameEntryId> NameMap;
+#if OUTPUT_NAMEMAP_CSV
+	TMap<FNameEntryId, TTuple<int32,int32,int32>> DebugNameCounts; // <Number0Count,OtherNumberCount,MaxNumber>
+#endif
+	FMappedName::EType NameMapType = FMappedName::EType::Global;
+};
+
 class FNameReaderProxyArchive
 	: public FArchiveProxy
 {
@@ -62,7 +206,12 @@ public:
 
 	FNameReaderProxyArchive(FArchive& InAr, const TArray<FNameEntryId>& InNameMap)
 		: FArchiveProxy(InAr)
-		, NameMap(InNameMap) { }
+		, NameMap(InNameMap)
+	{ 
+		// Replicate the filter editor only state of the InnerArchive as FArchiveProxy will
+		// not intercept it.
+		FArchive::SetFilterEditorOnly(InAr.IsFilterEditorOnly());
+	}
 
 	FArchive& operator<<(FName& Name)
 	{
@@ -89,8 +238,11 @@ struct FContainerSourceFile
 
 struct FContainerSourceSpec
 {
+	FName Name;
 	FString OutputPath;
 	TArray<FContainerSourceFile> SourceFiles;
+	TArray<FString> PatchSourceContainerFiles;
+	bool bGenerateDiffPatch;
 };
 
 struct FCookedFileStatData
@@ -106,9 +258,18 @@ struct FCookedFileStatData
 using FCookedFileStatMap = TMap<FString, FCookedFileStatData>;
 
 struct FPackage;
+struct FContainerTargetSpec;
+
+struct FContainerTargetFilePartialMapping
+{
+	FIoChunkId PartialChunkId;
+	uint64 Offset;
+	uint64 Length;
+};
 
 struct FContainerTargetFile
 {
+	FContainerTargetSpec* ContainerTarget = nullptr;
 	FPackage* Package = nullptr;
 	FString NormalizedSourcePath;
 	FString TargetPath;
@@ -118,7 +279,10 @@ struct FContainerTargetFile
 	uint64 Padding = 0;
 	uint64 Alignment = 0;
 	FIoChunkId ChunkId;
-	FSHAHash Hash;
+	FIoChunkHash ChunkHash;
+	TArray<uint8> PackageHeaderData;
+	TArray<int32> NameIndices;
+	TArray<FContainerTargetFilePartialMapping> PartialMappings;
 	bool bIsBulkData = false;
 	bool bIsOptionalBulkData = false;
 	bool bIsMemoryMappedBulkData = false;
@@ -127,7 +291,7 @@ struct FContainerTargetFile
 
 struct FIoStoreArguments
 {
-	FString OutputDir;
+	FString GlobalContainerPath;
 	FString CookedDir;
 	FString OutputReleaseVersionDir;
 	FString BasedOnReleaseVersionDir;
@@ -138,30 +302,35 @@ struct FIoStoreArguments
 	int64 MemoryMappingAlignment = 0;
 };
 
-struct FTocMetaEntry
+struct FContainerMeta
 {
-	uint64 Offset = 0;
-	uint64 Size = 0;
-	FSHAHash Hash;
+	FString ContainerName;
+	FIoContainerId ContainerId;
+	FNameMapBuilder NameMapBuilder;
 
-	friend FArchive& operator<<(FArchive& Ar, FTocMetaEntry& TocMetaEntry)
+	friend FArchive& operator<<(FArchive& Ar, FContainerMeta& ContainerMeta)
 	{
-		Ar << TocMetaEntry.Offset;
-		Ar << TocMetaEntry.Size;
-		Ar << TocMetaEntry.Hash;
+		Ar << ContainerMeta.ContainerName;
+		Ar << ContainerMeta.ContainerId;
+		Ar << ContainerMeta.NameMapBuilder;
 		return Ar;
 	}
 };
 
 struct FContainerTargetSpec
 {
-	FString Name;
+	FContainerHeader Header;
+	FName Name;
 	FString OutputPath;
 	FIoStoreWriter* IoStoreWriter;
 	TArray<FContainerTargetFile> TargetFiles;
 	TUniquePtr<FIoStoreEnvironment> IoStoreEnv;
-	TArray<FTocMetaEntry> PreviousBuildToc;
+	TArray<TUniquePtr<FIoStoreReader>> PatchSourceReaders;
+	FNameMapBuilder LocalNameMapBuilder;
+	FNameMapBuilder* NameMapBuilder = nullptr;
 	bool bIsCompressed = false;
+	bool bUseLocalNameMap = false;
+	bool bGenerateDiffPatch = false;
 };
 
 struct FPackageAssetData
@@ -209,149 +378,6 @@ static FString RemapLocalizationPathIfNeeded(const FString& Path, FString* OutRe
 	return Path;
 }
 
-class FNameMapBuilder
-{
-public:
-	void MarkNamesAsReferenced(const TArray<FName>& Names, TArray<FNameEntryId>& OutNameMap, TArray<int32>& OutNameIndices)
-	{
-		for (const FName& Name : Names)
-		{
-			const FNameEntryId Id = Name.GetComparisonIndex();
-			int32& Index = NameIndices.FindOrAdd(Id);
-			if (Index == 0)
-			{
-				Index = NameIndices.Num();
-				NameMap.Add(Id);
-			}
-
-			OutNameMap.Emplace(Name.GetDisplayIndex());
-			OutNameIndices.Add(Index - 1);
-		}
-	}
-
-	void MarkNameAsReferenced(const FName& Name)
-	{
-		const FNameEntryId Id = Name.GetComparisonIndex();
-		int32& Index = NameIndices.FindOrAdd(Id);
-		if (Index == 0)
-		{
-			Index = NameIndices.Num();
-			NameMap.Add(Id);
-		}
-#if OUTPUT_NAMEMAP_CSV
-		// debug counts
-		{
-			const int32 Number = Name.GetNumber();
-			TTuple<int32,int32,int32>& Counts = DebugNameCounts.FindOrAdd(Id);
-
-			if (Number == 0)
-			{
-				++Counts.Get<0>();
-			}
-			else
-			{
-				++Counts.Get<1>();
-				if (Number > Counts.Get<2>())
-				{
-					Counts.Get<2>() = Number;
-				}
-			}
-		}
-#endif
-	}
-
-	int32 MapName(const FName& Name) const
-	{
-		const FNameEntryId Id = Name.GetComparisonIndex();
-		const int32* Index = NameIndices.Find(Id);
-		check(Index);
-		return Index ? *Index - 1 : INDEX_NONE;
-	}
-
-	void SerializeName(FArchive& A, const FName& N) const
-	{
-		int32 NameIndex = MapName(N);
-		int32 NameNumber = N.GetNumber();
-		A << NameIndex << NameNumber;
-	}
-
-	const TMap<FNameEntryId, int32>& GetNameIndices() const
-	{
-		return NameIndices;
-	}
-
-	const TArray<FNameEntryId>& GetNameMap() const
-	{
-		return NameMap;
-	}
-
-	friend FArchive& operator<<(FArchive& Ar, FNameMapBuilder& NameMapBuilder)
-	{
-		if (Ar.IsSaving())
-		{
-			int32 NameCount = NameMapBuilder.NameMap.Num();
-			Ar << NameCount;
-			for (FNameEntryId NameEntryId : NameMapBuilder.NameMap)
-			{
-				const FNameEntry* NameEntry = FName::GetEntry(NameEntryId);
-				NameEntry->Write(Ar);
-			}
-		}
-		else
-		{
-			int32 NameCount;
-			Ar << NameCount;
-			for (int32 NameIndex = 0; NameIndex < NameCount; ++NameIndex)
-			{
-				FNameEntrySerialized NameEntrySerialized(ENAME_LinkerConstructor);
-				Ar << NameEntrySerialized;
-				FName Name(NameEntrySerialized);
-				FNameEntryId NameId = Name.GetComparisonIndex();
-				NameMapBuilder.NameMap.Add(NameId);
-				NameMapBuilder.NameIndices.Add(NameId, NameIndex + 1);
-			}
-		}
-		return Ar;
-	}
-
-#if OUTPUT_NAMEMAP_CSV
-	void SaveCsv(const FString& CsvFilePath)
-	{
-		{
-			TUniquePtr<FArchive> CsvArchive(IFileManager::Get().CreateFileWriter(*CsvFilePath));
-			if (CsvArchive)
-			{
-				TCHAR Name[FName::StringBufferSize];
-				ANSICHAR Line[MAX_SPRINTF + FName::StringBufferSize];
-				ANSICHAR Header[] = "Length\tMaxNumber\tNumberCount\tBaseCount\tTotalCount\tFName\n";
-				CsvArchive->Serialize(Header, sizeof(Header) - 1);
-				for (auto& Counts : DebugNameCounts)
-				{
-					const int32 NameLen = FName::CreateFromDisplayId(Counts.Key, 0).ToString(Name);
-					FCStringAnsi::Sprintf(Line, "%d\t%d\t%d\t%d\t%d\t",
-						NameLen, Counts.Value.Get<2>(), Counts.Value.Get<1>(), Counts.Value.Get<0>(), Counts.Value.Get<0>() + Counts.Value.Get<1>());
-					ANSICHAR* L = Line + FCStringAnsi::Strlen(Line);
-					const TCHAR* N = Name;
-					while (*N)
-					{
-						*L++ = CharCast<ANSICHAR,TCHAR>(*N++);
-					}
-					*L++ = '\n';
-					CsvArchive.Get()->Serialize(Line, L - Line);
-				}
-			}
-		}
-	}
-#endif
-
-private:
-	TMap<FNameEntryId, int32> NameIndices;
-	TArray<FNameEntryId> NameMap;
-#if OUTPUT_NAMEMAP_CSV
-	TMap<FNameEntryId, TTuple<int32,int32,int32>> DebugNameCounts; // <Number0Count,OtherNumberCount,MaxNumber>
-#endif
-};
-
 #if OUTPUT_CHUNKID_DIRECTORY
 class FChunkIdCsv
 {
@@ -549,7 +575,7 @@ struct FPackageHashes
 {
 	FSHAHash UAssetHash;
 	FSHAHash UExpHash;
-	FSHAHash ExportBundleHash;
+	FIoChunkHash ExportBundleHash;
 };
 #endif
 
@@ -590,8 +616,8 @@ struct FPackage
 	TSet<FPackage*> AllReachablePackages;
 	TSet<FPackage*> ImportedPreloadPackages;
 
+	TArray<FName> Names;
 	TArray<FNameEntryId> NameMap;
-	TArray<int32> NameIndices;
 
 	TArray<FPackageObjectIndex> Imports;
 	TArray<FPackageObjectIndex> PublicExports;
@@ -608,7 +634,7 @@ struct FPackage
 	TArray<FExportGraphNode*> NodesWithNoIncomingEdges;
 	FPackageGraphNode* Node = nullptr;
 
-	TArray<uint8> HeaderData;
+	uint64 HeaderSerialSize = 0;
 	uint64 ExportsSerialSize = 0;
 
 	uint64 DiskLayoutOrder = MAX_uint64;
@@ -1019,26 +1045,11 @@ static void BuildBundles(FExportGraph& ExportGraph, const TArray<FPackage*>& Pac
 	}
 }
 
-struct FFileLayoutNode
-{
-	FFileLayoutNode* Prev = nullptr;
-	FFileLayoutNode* Next = nullptr;
-	FContainerTargetFile* TargetFile = nullptr;
-	uint64 IdealOrder = 0;
-	uint64 PreviousOffset = 0;
-	bool bHasPreviousOffset = false;
-};
-
-static void CreateDiskLayout(
-	TArray<FContainerTargetSpec>& ContainerTargets,
+static void AssignPackagesDiskOrder(
 	const TArray<FPackage*>& Packages,
 	const TMap<FName, uint64> GameOrderMap,
-	const TMap<FName, uint64>& CookerOrderMap,
-	uint64 CompressionBlockSize,
-	uint64 MemoryMappingAlignment)
+	const TMap<FName, uint64>& CookerOrderMap)
 {
-	IOSTORE_CPU_SCOPE(CreateDiskLayout);
-
 	struct FCluster
 	{
 		TArray<FPackage*> Packages;
@@ -1144,224 +1155,545 @@ static void CreateDiskLayout(
 		}
 		delete Cluster;
 	}
+}
 
-	for (FContainerTargetSpec& ContainerTarget : ContainerTargets)
+static void CreateDiskLayout(
+	const TArray<FContainerTargetSpec*>& ContainerTargets,
+	const TArray<FPackage*>& Packages,
+	const TMap<FName, uint64> PackageOrderMap,
+	const TMap<FName, uint64>& CookerOrderMap,
+	uint64 CompressionBlockSize,
+	uint64 MemoryMappingAlignment)
+{
+	IOSTORE_CPU_SCOPE(CreateDiskLayout);
+
+	struct FLayoutEntry
+	{
+		enum ELayoutEntryType
+		{
+			Invalid,
+			File,
+			FreeSpace,
+			BlockBoundary,
+		};
+		FLayoutEntry* Prev = nullptr;
+		FLayoutEntry* Next = nullptr;
+		int32 BeginBlockIndex = -1;
+		int32 EndBlockIndex = -1;
+		ELayoutEntryType Type = Invalid;
+		uint64 Size = 0;
+		uint64 PreviousBuildOffset = 0;
+		uint64 IdealOrder = 0;
+		FContainerTargetFile* TargetFile = nullptr;
+		bool bHasPreviousBuildOffset = false;
+		bool bModified = false;
+		bool bLocked = false;
+	};
+
+	struct FLayoutBlock
+	{
+		FLayoutEntry* FirstEntry = nullptr;
+		FLayoutEntry* LastEntry = nullptr;
+		bool bModified = false;
+	};
+
+	AssignPackagesDiskOrder(Packages, PackageOrderMap, CookerOrderMap);
+
+	for (FContainerTargetSpec* ContainerTarget : ContainerTargets)
 	{
-		Algo::Sort(ContainerTarget.TargetFiles, [](const FContainerTargetFile& A, const FContainerTargetFile& B)
+		TArray<FLayoutEntry*> Entries;
+
+		FLayoutEntry* EntriesHead = new FLayoutEntry();
+		Entries.Add(EntriesHead);
+
+		TMap<int64, FLayoutEntry*> EntriesByOrderMap;
+		FLayoutEntry* LastAddedEntry = EntriesHead;
+
+		TMap<FIoChunkHash, TArray<FLayoutEntry*>> PreviousBuildFileByHash;
+		TMap<FIoChunkId, FIoChunkHash> PreviousBuildHashByChunkId;
+		uint64 CurrentOffset = 0;
+		FLayoutEntry* PrevEntryLink = EntriesHead;
+		
+		if (ContainerTarget->bGenerateDiffPatch)
+		{
+			for (const TUniquePtr<FIoStoreReader>& PatchSourceReader : ContainerTarget->PatchSourceReaders)
+			{
+				PatchSourceReader->EnumerateChunks([&PreviousBuildHashByChunkId](const FIoStoreTocChunkInfo& ChunkInfo)
+				{
+					PreviousBuildHashByChunkId.Add(ChunkInfo.Id, ChunkInfo.Hash);
+					return true;
+				});
+			}
+		}
+		else
+		{
+			if (ContainerTarget->PatchSourceReaders.Num())
+			{
+				ContainerTarget->PatchSourceReaders[0]->EnumerateChunks([&CurrentOffset, &PrevEntryLink, &Entries, &PreviousBuildFileByHash, &PreviousBuildHashByChunkId](const FIoStoreTocChunkInfo& ChunkInfo)
+				{
+					if (CurrentOffset < ChunkInfo.Offset)
+					{
+						FLayoutEntry* FreeSpaceEntry = new FLayoutEntry();
+						FreeSpaceEntry->Type = FLayoutEntry::FreeSpace;
+						FreeSpaceEntry->Size = ChunkInfo.Offset - CurrentOffset;
+						FreeSpaceEntry->PreviousBuildOffset = CurrentOffset;
+						FreeSpaceEntry->bHasPreviousBuildOffset = true;
+						CurrentOffset = ChunkInfo.Offset;
+						PrevEntryLink->Next = FreeSpaceEntry;
+						FreeSpaceEntry->Prev = PrevEntryLink;
+						PrevEntryLink = FreeSpaceEntry;
+						Entries.Add(FreeSpaceEntry);
+					}
+					FLayoutEntry* FileEntry = new FLayoutEntry();
+					FileEntry->Type = FLayoutEntry::File;
+					FileEntry->Size = ChunkInfo.Size;
+					FileEntry->PreviousBuildOffset = ChunkInfo.Offset;
+					FileEntry->bHasPreviousBuildOffset = true;
+
+					PrevEntryLink->Next = FileEntry;
+					FileEntry->Prev = PrevEntryLink;
+					PrevEntryLink = FileEntry;
+
+					CurrentOffset += ChunkInfo.Size;
+
+					Entries.Add(FileEntry);
+
+					TArray<FLayoutEntry*>& PreviousBuildFileHashesArray = PreviousBuildFileByHash.FindOrAdd(ChunkInfo.Hash);
+					PreviousBuildFileHashesArray.Push(FileEntry);
+					PreviousBuildHashByChunkId.Add(ChunkInfo.Id, ChunkInfo.Hash);
+
+					return true;
+				});
+			}
+		}
+		if (CurrentOffset % CompressionBlockSize != 0)
+		{
+			FLayoutEntry* LastFreeSpace = new FLayoutEntry();
+			LastFreeSpace->Type = FLayoutEntry::FreeSpace;
+			LastFreeSpace->Size = Align(CurrentOffset, CompressionBlockSize) - CurrentOffset;
+			LastFreeSpace->Prev = PrevEntryLink;
+			PrevEntryLink->Next = LastFreeSpace;
+			PrevEntryLink = LastFreeSpace;
+			Entries.Add(LastFreeSpace);
+			CurrentOffset += LastFreeSpace->Size;
+		}
+
+		FLayoutEntry* EntriesTail = new FLayoutEntry();
+		Entries.Add(EntriesTail);
+		PrevEntryLink->Next = EntriesTail;
+		EntriesTail->Prev = PrevEntryLink;
+
+		Algo::Sort(ContainerTarget->TargetFiles, [](const FContainerTargetFile& A, const FContainerTargetFile& B)
 		{
 			if (A.bIsMemoryMappedBulkData != B.bIsMemoryMappedBulkData)
 			{
 				return B.bIsMemoryMappedBulkData;
 			}
-			if (A.bIsBulkData != B.bIsBulkData)
+			if (A.bIsBulkData != B.bIsBulkData)
+			{
+				return B.bIsBulkData;
+			}
+
+			return A.Package->DiskLayoutOrder < B.Package->DiskLayoutOrder;
+		});
+
+		uint64 DiffFileCount = 0;
+		uint64 DiffFileSize = 0;
+		uint64 AddedFileCount = 0;
+		uint64 AddedFileSize = 0;
+		int64 IdealOrder = 0;
+		TArray<FLayoutEntry*> UnassignedEntries;
+		for (FContainerTargetFile& ContainerTargetFile : ContainerTarget->TargetFiles)
+		{
+			bool bIsAddedOrModified = false;
+			const FIoChunkHash* FindPreviousHashByChunkId = PreviousBuildHashByChunkId.Find(ContainerTargetFile.ChunkId);
+			if (FindPreviousHashByChunkId)
+			{
+				if (*FindPreviousHashByChunkId != ContainerTargetFile.ChunkHash)
+				{
+					//UE_LOG(LogIoStore, Display, TEXT("Diffing: %s"), *ContainerTargetFile.TargetPath);
+					++DiffFileCount;
+					DiffFileSize += ContainerTargetFile.TargetSize;
+					bIsAddedOrModified = true;
+				}
+			}
+			else
+			{
+				//UE_LOG(LogIoStore, Display, TEXT("Added: %s"), *ContainerTargetFile.TargetPath);
+				++AddedFileCount;
+				AddedFileSize += ContainerTargetFile.TargetSize;
+				bIsAddedOrModified = true;
+			}
+
+			bool bAddToUnassignedPool = true;
+			if (ContainerTarget->bGenerateDiffPatch)
+			{
+				if (!bIsAddedOrModified)
+				{
+					bAddToUnassignedPool = false;
+				}
+			}
+			else
+			{
+				FLayoutEntry* FindPreviousFileEntry = nullptr;
+				TArray<FLayoutEntry*>* FindPreviousBuildFileHashesArray = PreviousBuildFileByHash.Find(ContainerTargetFile.ChunkHash);
+				if (FindPreviousBuildFileHashesArray && FindPreviousBuildFileHashesArray->Num())
+				{
+					FindPreviousFileEntry = FindPreviousBuildFileHashesArray->Pop(false);
+				}
+
+				if (FindPreviousFileEntry && !FindPreviousFileEntry->TargetFile && ContainerTargetFile.TargetSize == FindPreviousFileEntry->Size)
+				{
+					FindPreviousFileEntry->TargetFile = &ContainerTargetFile;
+					FindPreviousFileEntry->IdealOrder = IdealOrder;
+					EntriesByOrderMap.Add(IdealOrder, FindPreviousFileEntry);
+					bAddToUnassignedPool = false;
+				}
+			}
+			if (bAddToUnassignedPool)
 			{
-				return B.bIsBulkData;
+				FLayoutEntry* NewEntry = new FLayoutEntry();
+				NewEntry->Type = FLayoutEntry::File;
+				NewEntry->Size = ContainerTargetFile.TargetSize;
+				NewEntry->TargetFile = &ContainerTargetFile;
+				NewEntry->IdealOrder = IdealOrder;
+				NewEntry->bModified = true;
+				Entries.Add(NewEntry);
+				UnassignedEntries.Add(NewEntry);
 			}
+			++IdealOrder;
+		}
+		UE_LOG(LogIoStore, Display, TEXT("%s: %d/%d modified entries (%fMB)"), *ContainerTarget->Name.ToString(), DiffFileCount, ContainerTarget->TargetFiles.Num(), DiffFileSize / 1024.0 / 1024.0);
+		UE_LOG(LogIoStore, Display, TEXT("%s: %d/%d added entries (%fMB)"), *ContainerTarget->Name.ToString(), AddedFileCount, ContainerTarget->TargetFiles.Num(), AddedFileSize / 1024.0 / 1024.0);
 
-			return A.Package->DiskLayoutOrder < B.Package->DiskLayoutOrder;
-		});
-
-		TMap<FSHAHash, FTocMetaEntry> PreviousBuildTocMap;
-		for (const FTocMetaEntry& TocMetaEntry : ContainerTarget.PreviousBuildToc)
+		for (FLayoutEntry* EntryIt = EntriesHead->Next; EntryIt != EntriesTail; EntryIt = EntryIt->Next)
 		{
-			if (PreviousBuildTocMap.Find(TocMetaEntry.Hash))
+			if (EntryIt->Type == FLayoutEntry::File && !EntryIt->TargetFile)
 			{
-				UE_DEBUG_BREAK();
+				EntryIt->Type = FLayoutEntry::FreeSpace;
+				EntryIt->bModified = true;
 			}
-			PreviousBuildTocMap.Add(TocMetaEntry.Hash, TocMetaEntry);
 		}
 
-		uint64 DiffCount = 0;
-		uint64 DiffSize = 0;
-		TArray<FFileLayoutNode*> LayoutNodes;
-		LayoutNodes.Reserve(ContainerTarget.TargetFiles.Num());
-		uint64 IdealOrder = 0;
-		TMap<uint64, FFileLayoutNode*> UnassignedNodes;
-		TMap<uint64, FFileLayoutNode*> AssignedNodes;
-		for (FContainerTargetFile& ContainerTargetFile : ContainerTarget.TargetFiles)
-		{
-			FFileLayoutNode* LayoutNode = new FFileLayoutNode();
-			LayoutNodes.Add(LayoutNode);
-			LayoutNode->TargetFile = &ContainerTargetFile;
-			LayoutNode->IdealOrder = IdealOrder++;
-			UnassignedNodes.Add(LayoutNode->IdealOrder, LayoutNode);
-
-			FTocMetaEntry* FindPreviousBuildEntry = PreviousBuildTocMap.Find(ContainerTargetFile.Hash);
-			if (!FindPreviousBuildEntry)
-			{
-				++DiffCount;
-				DiffSize += ContainerTargetFile.TargetSize;
-				//UE_LOG(LogIoStore, Display, TEXT("Diffing: %s"), *ContainerTargetFile.TargetPath);
-			}
-			else
+		// Assign entries to blocks
+		TArray<FLayoutBlock> Blocks;
+		Blocks.SetNum(CurrentOffset / CompressionBlockSize);
+		CurrentOffset = 0;
+		for (FLayoutEntry* EntryIt = EntriesHead->Next; EntryIt != EntriesTail; EntryIt = EntryIt->Next)
+		{
+			EntryIt->BeginBlockIndex = int32(CurrentOffset / CompressionBlockSize);
+			EntryIt->EndBlockIndex = int32(Align(CurrentOffset + EntryIt->Size, CompressionBlockSize) / CompressionBlockSize);
+			CurrentOffset += EntryIt->Size;
+			for (int32 BlockIndex = EntryIt->BeginBlockIndex; BlockIndex < EntryIt->EndBlockIndex; ++BlockIndex)
 			{
-				LayoutNode->PreviousOffset = FindPreviousBuildEntry->Offset;
-				LayoutNode->bHasPreviousOffset = true;
+				FLayoutBlock& Block = Blocks[BlockIndex];
+				Block.bModified |= EntryIt->bModified;
+				if (!Block.FirstEntry)
+				{
+					Block.FirstEntry = EntryIt;
+				}
+				Block.LastEntry = EntryIt;
 			}
 		}
-		UE_LOG(LogIoStore, Display, TEXT("%s: %d/%d new/modified entries (%fMB)"), *ContainerTarget.Name, DiffCount, ContainerTarget.TargetFiles.Num(), DiffSize / 1024.0 / 1024.0);
-
-		Algo::Sort(LayoutNodes, [](const FFileLayoutNode* A, const FFileLayoutNode* B)
+		// Put all file entries that only touch already modified blocks back to the unassigned pool
+		for (FLayoutEntry* EntryIt = EntriesHead->Next; EntryIt != EntriesTail; EntryIt = EntryIt->Next)
 		{
-			if (A->bHasPreviousOffset != B->bHasPreviousOffset)
+			if (EntryIt->Type == FLayoutEntry::File)
 			{
-				return A->bHasPreviousOffset;
+				bool bAllBlocksTouchedByEntryModified = true;
+				for (int32 BlockIndex = EntryIt->BeginBlockIndex; BlockIndex < EntryIt->EndBlockIndex; ++BlockIndex)
+				{
+					FLayoutBlock& Block = Blocks[BlockIndex];
+					if (!Block.bModified)
+					{
+						bAllBlocksTouchedByEntryModified = false;
+						break;
+					}
+				}
+				if (bAllBlocksTouchedByEntryModified)
+				{
+					FLayoutEntry* ReleasedEntry = new FLayoutEntry();
+					ReleasedEntry->Type = FLayoutEntry::File;
+					ReleasedEntry->Size = EntryIt->Size;
+					ReleasedEntry->TargetFile = EntryIt->TargetFile;
+					ReleasedEntry->IdealOrder = EntryIt->IdealOrder;
+					ReleasedEntry->bModified = true;
+					Entries.Add(ReleasedEntry);
+					UnassignedEntries.Add(ReleasedEntry);
+
+					EntryIt->Type = FLayoutEntry::FreeSpace;
+					EntryIt->TargetFile = nullptr;
+					EntryIt->bModified = true;
+				}
 			}
-			else if (A->PreviousOffset != B->PreviousOffset)
+		}
+		uint64 UnmodifiedBlocksCount = 0;
+		for (const FLayoutBlock& Block : Blocks)
+		{
+			if (!Block.bModified)
 			{
-				return A->PreviousOffset < B->PreviousOffset;
+				++UnmodifiedBlocksCount;
 			}
-			return A->IdealOrder < B->IdealOrder;
-		});
+		}
 
-		FFileLayoutNode Head;
-		FFileLayoutNode Tail;
-		Head.Next = &Tail;
-		Tail.Prev = &Head;
-		FFileLayoutNode* PreviousNode = &Head;
-		for (FFileLayoutNode* LayoutNode : LayoutNodes)
+		// Split all free space entries so that they don't cross any block boundaries
+		CurrentOffset = 0;
+		for (FLayoutEntry* EntryIt = EntriesHead->Next; EntryIt != EntriesTail; EntryIt = EntryIt->Next)
 		{
-			check(UnassignedNodes.Contains(LayoutNode->IdealOrder));
-			check(!AssignedNodes.Contains(LayoutNode->IdealOrder));
-			
-			FFileLayoutNode* OrderAfterNode;
-			if (LayoutNode->bHasPreviousOffset)
-			{
-				OrderAfterNode = PreviousNode;
-			}
-			else
+			uint64 NextOffset = CurrentOffset + EntryIt->Size;
+			if (EntryIt->Type == FLayoutEntry::FreeSpace)
 			{
-				if (LayoutNode->IdealOrder == 0)
+				if (EntryIt->BeginBlockIndex != EntryIt->EndBlockIndex - 1)
 				{
-					OrderAfterNode = &Head;
+					uint64 SizeInFirstBlock = Align(CurrentOffset, CompressionBlockSize) - CurrentOffset;
+					if (SizeInFirstBlock)
+					{
+						FLayoutEntry* FreeSpaceEntry = new FLayoutEntry();
+						FreeSpaceEntry->Type = FLayoutEntry::FreeSpace;
+						FreeSpaceEntry->Size = SizeInFirstBlock;
+						FreeSpaceEntry->bModified = EntryIt->bModified;
+						EntryIt->Size -= SizeInFirstBlock;
+						FreeSpaceEntry->Prev = EntryIt->Prev;
+						FreeSpaceEntry->Next = EntryIt;
+						EntryIt->Prev->Next = FreeSpaceEntry;
+						EntryIt->Prev = FreeSpaceEntry;
+						Entries.Add(FreeSpaceEntry);
+					}
+					uint64 SizeInLastBlock = EntryIt->Size % CompressionBlockSize;
+					if (SizeInLastBlock != EntryIt->Size)
+					{
+						uint64 SizeInMiddleBlocks = EntryIt->Size - SizeInLastBlock;
+						FLayoutEntry* FreeSpaceEntry = new FLayoutEntry();
+						FreeSpaceEntry->Type = FLayoutEntry::FreeSpace;
+						FreeSpaceEntry->Size = SizeInMiddleBlocks;
+						FreeSpaceEntry->bModified = EntryIt->bModified;
+						EntryIt->Size -= SizeInMiddleBlocks;
+						FreeSpaceEntry->Prev = EntryIt->Prev;
+						FreeSpaceEntry->Next = EntryIt;
+						EntryIt->Prev->Next = FreeSpaceEntry;
+						EntryIt->Prev = FreeSpaceEntry;
+						Entries.Add(FreeSpaceEntry);
+					}
 				}
-				else
+			}
+			CurrentOffset = NextOffset;
+		}
+		// Update entry block assignment
+		// Lock all file entries
+		CurrentOffset = 0;
+		for (FLayoutEntry* EntryIt = EntriesHead->Next; EntryIt != EntriesTail; EntryIt = EntryIt->Next)
+		{
+			EntryIt->BeginBlockIndex = int32(CurrentOffset / CompressionBlockSize);
+			EntryIt->EndBlockIndex = int32(Align(CurrentOffset + EntryIt->Size, CompressionBlockSize) / CompressionBlockSize);
+			CurrentOffset += EntryIt->Size;
+			for (int32 BlockIndex = EntryIt->BeginBlockIndex; BlockIndex < EntryIt->EndBlockIndex; ++BlockIndex)
+			{
+				FLayoutBlock& Block = Blocks[BlockIndex];
+				check(!EntryIt->bModified || Block.bModified);
+				if (!Block.FirstEntry)
 				{
-					OrderAfterNode = AssignedNodes.FindRef(LayoutNode->IdealOrder - 1);
-					check(OrderAfterNode);
+					Block.FirstEntry = EntryIt;
 				}
+				Block.LastEntry = EntryIt;
+			}
+			if (EntryIt->Type == FLayoutEntry::File)
+			{
+				check(!EntryIt->bModified);
+				EntryIt->bLocked = true;
 			}
-
-			OrderAfterNode->Next->Prev = LayoutNode;
-			LayoutNode->Next = OrderAfterNode->Next;
-			LayoutNode->Prev = OrderAfterNode;
-			OrderAfterNode->Next = LayoutNode;
-			
-			AssignedNodes.Add(LayoutNode->IdealOrder, LayoutNode);
-			UnassignedNodes.Remove(LayoutNode->IdealOrder);
-
-			PreviousNode = LayoutNode;
 		}
-
-		TArray<bool> BlocksModificationState;
-		uint64 CurrentOffset = 0;
-		FFileLayoutNode* CurrentNode = Head.Next;
-		bool bIsAlignedForMemoryMappedBulkData = false;
-		while (CurrentNode != &Tail)
+		// Lock all free space entries in unmodified blocks
+		for (FLayoutBlock& Block : Blocks)
 		{
-			uint64 ModificationOffset = 0;
-			uint64 ModificationSize = 0;
-
-			if (CurrentNode->TargetFile->bIsMemoryMappedBulkData)
+			if (!Block.bModified)
 			{
-				if (!bIsAlignedForMemoryMappedBulkData)
+				for (FLayoutEntry* EntryIt = Block.FirstEntry; EntryIt != Block.LastEntry->Next; EntryIt = EntryIt->Next)
 				{
-					if (!IsAligned(CurrentOffset, CompressionBlockSize))
+					if (EntryIt->Type == FLayoutEntry::FreeSpace)
 					{
-						uint64 RemainingInBlock = Align(CurrentOffset, CompressionBlockSize) - CurrentOffset;
-						CurrentNode->TargetFile->Padding = RemainingInBlock;
-						CurrentOffset += RemainingInBlock;
+						EntryIt->bLocked = true;
 					}
-					CurrentNode->TargetFile->Alignment = MemoryMappingAlignment;
-					bIsAlignedForMemoryMappedBulkData = true;
 				}
-				else
+			}
+		}
+		// Merge and shrink all unlocked free space
+		for (FLayoutEntry* EntryIt = EntriesHead->Next; EntryIt != EntriesTail; EntryIt = EntryIt->Next)
+		{
+			if (EntryIt->Type == FLayoutEntry::FreeSpace && !EntryIt->bLocked)
+			{
+				if (EntryIt->Prev->Type == FLayoutEntry::FreeSpace && !EntryIt->Prev->bLocked)
 				{
-					if (!IsAligned(CurrentOffset, MemoryMappingAlignment))
-					{
-						uint64 RequiredPadding = Align(CurrentOffset, MemoryMappingAlignment) - CurrentOffset;
-						CurrentNode->TargetFile->Padding = RequiredPadding;
-						CurrentOffset += RequiredPadding;
-					}
+					FLayoutEntry* MergeWithFreeSpace = EntryIt->Prev;
+					EntryIt->Size += MergeWithFreeSpace->Size;
+					MergeWithFreeSpace->Prev->Next = EntryIt;
+					EntryIt->Prev = MergeWithFreeSpace->Prev;
 				}
+				EntryIt->Size %= CompressionBlockSize;
+			}
+		}
+		// Insert block boundaries
+		CurrentOffset = 0;
+		for (FLayoutEntry* EntryIt = EntriesHead->Next; EntryIt != EntriesTail; EntryIt = EntryIt->Next)
+		{
+			if (CurrentOffset % CompressionBlockSize == 0)
+			{
+				FLayoutEntry* BlockBoundaryEntry = new FLayoutEntry();
+				BlockBoundaryEntry->Type = FLayoutEntry::BlockBoundary;
+				BlockBoundaryEntry->Prev = EntryIt->Prev;
+				BlockBoundaryEntry->Next = EntryIt;
+				EntryIt->Prev->Next = BlockBoundaryEntry;
+				EntryIt->Prev = BlockBoundaryEntry;
+				Entries.Add(BlockBoundaryEntry);
 			}
-			else if (bIsAlignedForMemoryMappedBulkData)
+			
+			CurrentOffset += EntryIt->Size;
+		}
+		check(CurrentOffset % CompressionBlockSize == 0);
+		FLayoutEntry* LastBlockBoundaryEntry = new FLayoutEntry();
+		LastBlockBoundaryEntry->Type = FLayoutEntry::BlockBoundary;
+		LastBlockBoundaryEntry->Prev = EntriesTail->Prev;
+		LastBlockBoundaryEntry->Next = EntriesTail;
+		EntriesTail->Prev->Next = LastBlockBoundaryEntry;
+		EntriesTail->Prev = LastBlockBoundaryEntry;
+		Entries.Add(LastBlockBoundaryEntry);
+
+		FLayoutEntry* MemoryMappedFilesTarget = new FLayoutEntry();
+		MemoryMappedFilesTarget->Type = FLayoutEntry::BlockBoundary;
+		MemoryMappedFilesTarget->Prev = LastBlockBoundaryEntry;
+		LastBlockBoundaryEntry->Next = MemoryMappedFilesTarget;
+		EntriesTail->Prev = MemoryMappedFilesTarget;
+		MemoryMappedFilesTarget->Next = EntriesTail;
+		Entries.Add(MemoryMappedFilesTarget);
+
+		// Insert new/modified entries according to their ideal order
+		Algo::Sort(UnassignedEntries, [](const FLayoutEntry* A, const FLayoutEntry* B)
+		{
+			return A->IdealOrder < B->IdealOrder;
+		});
+		for (FLayoutEntry* UnassignedEntry : UnassignedEntries)
+		{
+			check(UnassignedEntry->TargetFile);
+			if (UnassignedEntry->TargetFile->bIsMemoryMappedBulkData)
 			{
-				if (!IsAligned(CurrentOffset, CompressionBlockSize))
+				UnassignedEntry->Prev = MemoryMappedFilesTarget->Prev;
+				UnassignedEntry->Next = MemoryMappedFilesTarget;
+				MemoryMappedFilesTarget->Prev->Next = UnassignedEntry;
+				MemoryMappedFilesTarget->Prev = UnassignedEntry;
+
+				uint64 AlignedSize = Align(UnassignedEntry->Size, MemoryMappingAlignment);
+				uint64 MemoryMapPadding = AlignedSize - UnassignedEntry->Size;
+				if (MemoryMapPadding)
 				{
-					uint64 RemainingInBlock = Align(CurrentOffset, CompressionBlockSize) - CurrentOffset;
-					CurrentNode->TargetFile->Padding = RemainingInBlock;
-					CurrentOffset += RemainingInBlock;
+					FLayoutEntry* PaddingEntry = new FLayoutEntry();
+					PaddingEntry->Type = FLayoutEntry::FreeSpace;
+					PaddingEntry->Size = MemoryMapPadding;
+
+					PaddingEntry->Prev = UnassignedEntry;
+					PaddingEntry->Next = UnassignedEntry->Next;
+					UnassignedEntry->Next->Prev = PaddingEntry;
+					UnassignedEntry->Next = PaddingEntry;
+					Entries.Add(PaddingEntry);
 				}
-				bIsAlignedForMemoryMappedBulkData = false;
 			}
-
-			if (CurrentNode->bHasPreviousOffset)
+			else
 			{
-				uint64 PreviousOffsetInBlock = CurrentNode->PreviousOffset % CompressionBlockSize;
-				uint64 CurrentOffsetInBlock = CurrentOffset % CompressionBlockSize;
-				if (CurrentOffsetInBlock != PreviousOffsetInBlock)
+				FLayoutEntry* PutAfterEntry = EntriesByOrderMap.FindRef(UnassignedEntry->IdealOrder - 1);
+				if (!PutAfterEntry)
 				{
-					uint64 PadBytes = 0;
-					if (CurrentOffsetInBlock < PreviousOffsetInBlock)
+					PutAfterEntry = LastAddedEntry;
+				}
+				
+				FLayoutEntry* TargetFreeSpace = nullptr;
+				FLayoutEntry* CandidateTarget = PutAfterEntry->Next;
+				if (CandidateTarget->Type == FLayoutEntry::FreeSpace && !CandidateTarget->bLocked)
+				{
+					TargetFreeSpace = CandidateTarget;
+					if (TargetFreeSpace->Size < UnassignedEntry->Size)
 					{
-						PadBytes = PreviousOffsetInBlock - CurrentOffsetInBlock;
+						uint64 SizeExtension = Align(UnassignedEntry->Size - TargetFreeSpace->Size, CompressionBlockSize);
+						TargetFreeSpace->Size += SizeExtension;
 					}
-					else if (PreviousOffsetInBlock < CurrentOffsetInBlock)
+				}
+
+				if (!TargetFreeSpace)
+				{
+					FLayoutEntry* NextBlockBoundary = PutAfterEntry->Next;
+					check(NextBlockBoundary);
+					while (NextBlockBoundary->Type != FLayoutEntry::BlockBoundary)
 					{
-						PadBytes = CompressionBlockSize - CurrentOffsetInBlock + PreviousOffsetInBlock;
+						NextBlockBoundary = NextBlockBoundary->Next;
 					}
-					ModificationOffset = CurrentOffset;
-					ModificationSize += PadBytes;
-					CurrentOffset += PadBytes;
-					CurrentNode->TargetFile->Padding = PadBytes;
+					FLayoutEntry* NewFreeSpaceEntry = new FLayoutEntry();
+					NewFreeSpaceEntry->Type = FLayoutEntry::FreeSpace;
+					NewFreeSpaceEntry->Size = Align(UnassignedEntry->Size, CompressionBlockSize);
+					Entries.Add(NewFreeSpaceEntry);
+					FLayoutEntry* NewBlockBoundaryEntry = new FLayoutEntry();
+					NewBlockBoundaryEntry->Type = FLayoutEntry::BlockBoundary;
+					Entries.Add(NewBlockBoundaryEntry);
+
+					NewFreeSpaceEntry->Prev = NextBlockBoundary;
+					NewFreeSpaceEntry->Next = NewBlockBoundaryEntry;
+					NewBlockBoundaryEntry->Prev = NewFreeSpaceEntry;
+					NewBlockBoundaryEntry->Next = NextBlockBoundary->Next;
+					NextBlockBoundary->Next->Prev = NewBlockBoundaryEntry;
+					NextBlockBoundary->Next = NewFreeSpaceEntry;
+					TargetFreeSpace = NewFreeSpaceEntry;
 				}
-				check(CurrentOffset % CompressionBlockSize == CurrentNode->PreviousOffset % CompressionBlockSize);
+
+				check(TargetFreeSpace->Type == FLayoutEntry::FreeSpace);
+				check(TargetFreeSpace->Size >= UnassignedEntry->Size);
+				check(!TargetFreeSpace->bLocked);
+				UnassignedEntry->Prev = TargetFreeSpace->Prev;
+				UnassignedEntry->Next = TargetFreeSpace;
+				TargetFreeSpace->Prev->Next = UnassignedEntry;
+				TargetFreeSpace->Prev = UnassignedEntry;
+				TargetFreeSpace->Size -= UnassignedEntry->Size;
+
+				EntriesByOrderMap.Add(UnassignedEntry->IdealOrder, UnassignedEntry);
+				LastAddedEntry = UnassignedEntry;
 			}
-			else
+		}
+		TArray<FContainerTargetFile> IncludedContainerTargetFiles;
+		CurrentOffset = 0;
+		uint64 PaddingBytes = 0;
+		uint64 TotalChunkPaddingSize = 0;
+		for (FLayoutEntry* EntryIt = EntriesHead->Next; EntryIt != EntriesTail; EntryIt = EntryIt->Next)
+		{
+			if (EntryIt->Type == FLayoutEntry::FreeSpace)
 			{
-				ModificationOffset = CurrentOffset;
-				ModificationSize = CurrentNode->TargetFile->TargetSize;
+				PaddingBytes += EntryIt->Size;
 			}
-
-			if (ModificationSize)
+			else if (EntryIt->Type == FLayoutEntry::File)
 			{
-				int32 BeginBlockIndex = ModificationOffset / CompressionBlockSize;
-				int32 EndBlockIndex = Align(ModificationOffset + ModificationSize, CompressionBlockSize) / CompressionBlockSize;
-				if (BlocksModificationState.Num() < EndBlockIndex)
+				check(EntryIt->TargetFile);
+				EntryIt->TargetFile->Padding = PaddingBytes;
+				TotalChunkPaddingSize += PaddingBytes;
+				PaddingBytes = 0;
+				EntryIt->TargetFile->Offset = CurrentOffset;
+
+				if (EntryIt->bHasPreviousBuildOffset && EntryIt->bLocked)
 				{
-					BlocksModificationState.AddZeroed(EndBlockIndex - BlocksModificationState.Num());
+					check(EntryIt->PreviousBuildOffset % CompressionBlockSize == EntryIt->TargetFile->Offset % CompressionBlockSize);
 				}
-				for (int32 BlockIndex = BeginBlockIndex; BlockIndex < EndBlockIndex; ++BlockIndex)
+				if (EntryIt->TargetFile->bIsMemoryMappedBulkData)
 				{
-					BlocksModificationState[BlockIndex] = true;
+					check(IsAligned(CurrentOffset, MemoryMappingAlignment));
 				}
+				IncludedContainerTargetFiles.Add(*EntryIt->TargetFile);
 			}
-
-			CurrentNode->TargetFile->Offset = CurrentOffset;
-			CurrentOffset += CurrentNode->TargetFile->TargetSize;
-			CurrentNode = CurrentNode->Next;
-		}
-
-		uint64 ModifiedBlocksCount = 0;
-		for (bool BlockModified : BlocksModificationState)
-		{
-			if (BlockModified)
-			{
-				++ModifiedBlocksCount;
-			}
+			CurrentOffset += EntryIt->Size;
 		}
-		UE_LOG(LogIoStore, Display, TEXT("%s: %d new/modified blocks (%fMB)"), *ContainerTarget.Name, ModifiedBlocksCount, ModifiedBlocksCount * CompressionBlockSize / 1024.0f / 1024.0f);
+		uint64 TotalBlockCount = Align(CurrentOffset, CompressionBlockSize) / CompressionBlockSize;
+		uint64 ModifiedBlocksCount = TotalBlockCount - UnmodifiedBlocksCount;
+		UE_LOG(LogIoStore, Display, TEXT("%s: %d/%d modified blocks (%fMB)"), *ContainerTarget->Name.ToString(), ModifiedBlocksCount, TotalBlockCount, (ModifiedBlocksCount * CompressionBlockSize) / 1024.0 / 1024.0);
+		UE_LOG(LogIoStore, Display, TEXT("%s: Total chunk padding %fMB"), *ContainerTarget->Name.ToString(), TotalChunkPaddingSize / 1024.0 / 1024.0);
 
-		for (FFileLayoutNode* LayoutNode : LayoutNodes)
+		for (FLayoutEntry* Entry : Entries)
 		{
-			delete LayoutNode;
+			delete Entry;
 		}
 
-		Algo::Sort(ContainerTarget.TargetFiles, [](const FContainerTargetFile& A, const FContainerTargetFile& B)
-		{
-			return A.Offset < B.Offset;
-		});
+		Swap(ContainerTarget->TargetFiles, IncludedContainerTargetFiles);
 	}
 }
 
@@ -1381,7 +1713,7 @@ static EIoChunkType BulkdataTypeToChunkIdType(FPackageStoreBulkDataManifest::EBu
 	}
 }
 
-static bool MapAdditionalBulkDataChunks(const FContainerTargetFile& TargetFile, const FPackageStoreBulkDataManifest& BulkDataManifest, FIoStoreWriter* IoStoreWriter)
+static bool MapAdditionalBulkDataChunks(FContainerTargetFile& TargetFile, const FPackageStoreBulkDataManifest& BulkDataManifest)
 {
 	const FPackage& Package = *TargetFile.Package;
 	const FPackageStoreBulkDataManifest::FPackageDesc* PackageDesc = BulkDataManifest.Find(Package.FileName);
@@ -1400,12 +1732,10 @@ static bool MapAdditionalBulkDataChunks(const FContainerTargetFile& TargetFile,
 		{
 			if (BulkDataDesc.Type == BulkDataType)
 			{
-				const FIoChunkId AccessChunkId = CreateChunkIdForBulkData(Package.GlobalPackageId, BulkDataDesc.ChunkId, ChunkIdType, *Package.FileName);
-				const FIoStatus PartialResult = IoStoreWriter->MapPartialRange(TargetFile.ChunkId, BulkDataDesc.Offset, BulkDataDesc.Size, AccessChunkId);
-				if (!PartialResult.IsOk())
-				{
-					UE_LOG(LogIoStore, Warning, TEXT("Failed to map partial range for '%s' due to: %s"), *Package.FileName, *PartialResult.ToString());
-				}
+				FContainerTargetFilePartialMapping& PartialMapping = TargetFile.PartialMappings.AddDefaulted_GetRef();
+				PartialMapping.PartialChunkId = CreateChunkIdForBulkData(Package.GlobalPackageId, BulkDataDesc.ChunkId, ChunkIdType, *Package.FileName);
+				PartialMapping.Offset = BulkDataDesc.Offset;
+				PartialMapping.Length = BulkDataDesc.Size;
 			}
 		}
 	}
@@ -1673,6 +2003,48 @@ static void FindExport(
 	}
 }
 
+FContainerTargetSpec* AddContainer(
+	FName Name,
+	FIoContainerId Id,
+	TArray<FContainerTargetSpec*>& Containers,
+	TMap<FName, FContainerTargetSpec*>& ContainerMap)
+{
+	for (const FContainerTargetSpec* ExistingContainer : Containers)
+	{
+		check(ExistingContainer->Header.ContainerId != Id);
+	}
+	check(!ContainerMap.Contains(Name));
+	FContainerTargetSpec* ContainerTargetSpec = new FContainerTargetSpec();
+	ContainerTargetSpec->Name = Name;
+	ContainerTargetSpec->Header.ContainerId = Id;
+	ContainerMap.Add(Name, ContainerTargetSpec);
+	Containers.Add(ContainerTargetSpec);
+	return ContainerTargetSpec;
+}
+
+FContainerTargetSpec* FindOrAddContainer(
+	FName Name,
+	TArray<FContainerTargetSpec*>& Containers,
+	TMap<FName, FContainerTargetSpec*>& ContainerMap)
+{
+	FContainerTargetSpec* ContainerTargetSpec = ContainerMap.FindRef(Name);
+	if (!ContainerTargetSpec)
+	{
+		uint16 NextContainerTargetId = 1;
+		for (const FContainerTargetSpec* ExistingContainer : Containers)
+		{
+			check(ExistingContainer->Header.ContainerId.IsValid());
+			NextContainerTargetId = FMath::Max<uint16>(ExistingContainer->Header.ContainerId.ToIndex() + 1, NextContainerTargetId);
+		}
+		ContainerTargetSpec = new FContainerTargetSpec();
+		ContainerTargetSpec->Name = Name;
+		ContainerTargetSpec->Header.ContainerId = FIoContainerId::FromIndex(NextContainerTargetId);
+		ContainerMap.Add(Name, ContainerTargetSpec);
+		Containers.Add(ContainerTargetSpec);
+	}
+	return ContainerTargetSpec;
+}
+
 FPackage* FindOrAddPackage(
 	const TCHAR* RelativeFileName,
 	TArray<FPackage*>& Packages,
@@ -2026,18 +2398,31 @@ static void AddPreloadDependencies(
 	}
 };
 
+void BuildContainerNameMap(FContainerTargetSpec& ContainerTarget)
+{
+	FNameMapBuilder& NameMapBuilder = *ContainerTarget.NameMapBuilder;
+
+	for (FContainerTargetFile& TargetFile : ContainerTarget.TargetFiles)
+	{
+		NameMapBuilder.MarkNameAsReferenced(TargetFile.Package->Name);
+		NameMapBuilder.MarkNamesAsReferenced(TargetFile.Package->Names, TargetFile.NameIndices);
+	}
+}
+
 void FinalizePackageHeaders(
-	const TArray<FPackage*>& Packages,
-	const FNameMapBuilder& NameMapBuilder,
+	FContainerTargetSpec& ContainerTarget,
 	const TArray<FObjectExport>& ObjectExports,
 	const TArray<FExportData>& GlobalExports,
 	const FImportObjectsByFullName& GlobalImportsByFullName)
 {
-	IOSTORE_CPU_SCOPE(FinalizePackageHeaders);
-	UE_LOG(LogIoStore, Display, TEXT("Finalizing package headers..."));
+	const uint16 NameMapIndex = ContainerTarget.bUseLocalNameMap ? ContainerTarget.Header.ContainerId.ToIndex() : 0;
 
-	for (FPackage* Package : Packages)
+	for (FContainerTargetFile& TargetFile : ContainerTarget.TargetFiles)
 	{
+		FPackage* Package = TargetFile.Package;
+		check(TargetFile.ContainerTarget->NameMapBuilder);
+		const FNameMapBuilder& NameMapBuilder = *TargetFile.ContainerTarget->NameMapBuilder;
+
 		// Temporary Archive for serializing ImportMap
 		FBufferWriter ImportMapArchive(nullptr, 0, EBufferWriterFlags::AllowResize | EBufferWriterFlags::TakeOwnership);
 		for (FPackageObjectIndex& GlobalImportIndex : Package->Imports)
@@ -2077,8 +2462,7 @@ void FinalizePackageHeaders(
 			FExportMapEntry ExportMapEntry;
 			ExportMapEntry.CookedSerialOffset = ObjectExport.SerialOffset;
 			ExportMapEntry.CookedSerialSize = ObjectExport.SerialSize;
-			ExportMapEntry.ObjectName.NameEntryIndex = NameMapBuilder.MapName(ObjectExport.ObjectName);
-			ExportMapEntry.ObjectName.Number = ObjectExport.ObjectName.GetNumber();
+			ExportMapEntry.ObjectName = NameMapBuilder.MapName(ObjectExport.ObjectName);
 			ExportMapEntry.OuterIndex = ExportData.OuterIndex;
 			ExportMapEntry.ClassIndex = ExportData.ClassIndex;
 			ExportMapEntry.SuperIndex = ExportData.SuperIndex;
@@ -2119,7 +2503,11 @@ void FinalizePackageHeaders(
 		}
 		Package->ExportBundlesHeaderSize = ExportBundlesArchive.Tell();
 
-		uint64 PackageHeaderSize =
+		const uint64 NameMapSize = TargetFile.NameIndices.Num() * TargetFile.NameIndices.GetTypeSize();
+		check(Package->NameMapSize == 0 || Package->NameMapSize == NameMapSize);
+		Package->NameMapSize = NameMapSize;
+
+		const uint64 PackageHeaderSize =
 			sizeof(FPackageSummary)
 			+ Package->NameMapSize
 			+ Package->ImportMapSize
@@ -2127,12 +2515,17 @@ void FinalizePackageHeaders(
 			+ Package->ExportBundlesHeaderSize
 			+ Package->UGraphSize;
 
-		Package->HeaderData.AddZeroed(PackageHeaderSize);
-		uint8* PackageHeaderBuffer = Package->HeaderData.GetData();
+		check(Package->HeaderSerialSize == 0 || Package->HeaderSerialSize == PackageHeaderSize);
+		Package->HeaderSerialSize = PackageHeaderSize;
+
+		TargetFile.PackageHeaderData.AddZeroed(PackageHeaderSize);
+		uint8* PackageHeaderBuffer = TargetFile.PackageHeaderData.GetData();
 		FPackageSummary* PackageSummary = reinterpret_cast<FPackageSummary*>(PackageHeaderBuffer);
 
 		PackageSummary->PackageFlags = Package->PackageFlags;
 		PackageSummary->CookedHeaderSize = Package->CookedHeaderSize;
+		PackageSummary->NameMapIndex = NameMapIndex;
+		PackageSummary->Pad = 0;
 		PackageSummary->GraphDataSize = Package->UGraphSize;
 
 		FBufferWriter SummaryArchive(PackageHeaderBuffer, PackageHeaderSize);
@@ -2141,7 +2534,8 @@ void FinalizePackageHeaders(
 		// NameMap data
 		{
 			PackageSummary->NameMapOffset = SummaryArchive.Tell();
-			SummaryArchive.Serialize(Package->NameIndices.GetData(), Package->NameMapSize);
+			SummaryArchive.Serialize(TargetFile.NameIndices.GetData(),
+				TargetFile.NameIndices.Num() * TargetFile.NameIndices.GetTypeSize());
 		}
 
 		// ImportMap data
@@ -2176,7 +2570,6 @@ void FinalizePackageHeaders(
 
 void SerializePackageStoreEntries(
 	TArray<FPackage*> Packages,
-	FNameMapBuilder& NameMapBuilder,
 	FArchive& StoreTocArchive,
 	FArchive& StoreDataArchive)
 {
@@ -2196,11 +2589,10 @@ void SerializePackageStoreEntries(
 		StoreTocArchive << OffsetToDataFromThis;
 	};
 
-
 	for (FPackage* Package : Packages)
 	{
-		NameMapBuilder.MarkNameAsReferenced(Package->Name);
-		NameMapBuilder.SerializeName(StoreTocArchive, Package->Name);
+		FMappedName PackageName;
+		StoreTocArchive << PackageName;
 		StoreTocArchive << Package->SourceGlobalPackageId;
 		StoreTocArchive << Package->ExportCount;
 
@@ -2229,7 +2621,7 @@ void SerializePackageStoreEntries(
 			//       or just store the data for the first entry directly on the package entry itself
 			SerializePackageEntryCArrayHeader(Package->ExportBundles.Num());
 			FExportBundleMetaEntry ExportBundleMetaEntry;
-			ExportBundleMetaEntry.PayloadSize = Package->HeaderData.Num() + Package->ExportsSerialSize;
+			ExportBundleMetaEntry.PayloadSize = Package->HeaderSerialSize + Package->ExportsSerialSize;
 			for (const FExportBundle& ExportBundle : Package->ExportBundles)
 			{
 				ExportBundleMetaEntry.LoadOrder = ExportBundle.LoadOrder;
@@ -2241,22 +2633,8 @@ void SerializePackageStoreEntries(
 	}
 }
 
-template<class T>
-static void SerializeImportNames(
-	const FNameMapBuilder& NameMapBuilder,
-	const TArray<T>& Imports,
-	FArchive& Ar)
-{
-	int32 Count = Imports.Num();
-	Ar << Count;
-	for (const T& ImportData : Imports)
-	{
-		NameMapBuilder.SerializeName(Ar, ImportData.ObjectName);
-	}
-};
-
 static void SerializeInitialLoad(
-	const FNameMapBuilder& NameMapBuilder,
+	FNameMapBuilder& GlobalNameMapBuilder,
 	const FGlobalScriptImports& GlobalScriptImports,
 	FArchive& InitialLoadArchive)
 {
@@ -2268,9 +2646,9 @@ static void SerializeInitialLoad(
 
 	for (const FScriptImportData& ImportData : GlobalScriptImports)
 	{
+		GlobalNameMapBuilder.MarkNameAsReferenced(ImportData.ObjectName);
 		FScriptObjectEntry Entry;
-		Entry.ObjectName.Index = NameMapBuilder.MapName(ImportData.ObjectName);
-		Entry.ObjectName.Number = ImportData.ObjectName.GetNumber();
+		Entry.ObjectName = GlobalNameMapBuilder.MapName(ImportData.ObjectName).ToUnresolvedMinimalName();
 		Entry.OuterIndex = ImportData.OuterIndex;
 		Entry.CDOClassIndex = ImportData.CDOClassIndex;
 
@@ -2278,13 +2656,15 @@ static void SerializeInitialLoad(
 	}
 };
 
-static FIoBuffer CreateExportBundleBuffer(const FPackage* Package, const TArray<FObjectExport>& ObjectExports, const FIoBuffer UExpBuffer)
+static FIoBuffer CreateExportBundleBuffer(const FContainerTargetFile& TargetFile, const TArray<FObjectExport>& ObjectExports, const FIoBuffer UExpBuffer)
 {
-	uint64 BundleBufferSize = Package->HeaderData.Num() + Package->ExportsSerialSize;
+	const FPackage* Package = TargetFile.Package;
+	check(TargetFile.PackageHeaderData.Num() > 0);
+	const uint64 BundleBufferSize = TargetFile.PackageHeaderData.Num() + TargetFile.Package->ExportsSerialSize;
 	FIoBuffer BundleBuffer(BundleBufferSize);
-	FMemory::Memcpy(BundleBuffer.Data(), Package->HeaderData.GetData(), Package->HeaderData.Num());
-	uint64 BundleBufferOffset = Package->HeaderData.Num();
-	for (const FExportBundle& ExportBundle : Package->ExportBundles)
+	FMemory::Memcpy(BundleBuffer.Data(), TargetFile.PackageHeaderData.GetData(), TargetFile.PackageHeaderData.Num());
+	uint64 BundleBufferOffset = TargetFile.PackageHeaderData.Num();
+	for (const FExportBundle& ExportBundle : TargetFile.Package->ExportBundles)
 	{
 		for (const FExportGraphNode* Node : ExportBundle.Nodes)
 		{
@@ -2303,10 +2683,8 @@ static FIoBuffer CreateExportBundleBuffer(const FPackage* Package, const TArray<
 }
 
 static void ParsePackageAssets(
-	FNameMapBuilder& NameMapBuilder,
 	TArray<FPackage*>& Packages,
 	FPackageAssetData& PackageAssetData,
-	FGlobalPackageData& GlobalPackageData,
 	FExportGraph& ExportGraph)
 {
 	IOSTORE_CPU_SCOPE(ParsePackageAssets);
@@ -2427,22 +2805,17 @@ static void ParsePackageAssets(
 			{
 				Ar.Seek(Summary.NameOffset);
 
-				TArray<FName> Names;
-				Names.Reserve(Summary.NameCount);
+				Package.Names.Reserve(Summary.NameCount);
 				Package.NameMap.Reserve(Summary.NameCount);
-				Package.NameIndices.Reserve(Summary.NameCount);
 				FNameEntrySerialized NameEntry(ENAME_LinkerConstructor);
 
 				for (int32 I = 0; I < Summary.NameCount; ++I)
 				{
 					Ar << NameEntry;
-					Names.Emplace(NameEntry);
+					FName& Name = Package.Names.Emplace_GetRef(NameEntry);
+					Package.NameMap.Emplace(Name.GetDisplayIndex());
 				}
-
-				NameMapBuilder.MarkNamesAsReferenced(Names, Package.NameMap, Package.NameIndices);
 			}
-
-			Package.NameMapSize = Package.NameIndices.Num() * Package.NameIndices.GetTypeSize();
 		}
 
 		ParallelFor(TotalPackageCount,[
@@ -2450,9 +2823,7 @@ static void ParsePackageAssets(
 				&PackageAssetBuffers,
 				&PackageFileSummaries,
 				&Packages,
-				&NameMapBuilder,
 				&PackageAssetData,
-				&GlobalPackageData,
 				&ExportGraph](int32 Index)
 		{
 			uint8* PackageBuffer = PackageAssetBuffers[Index];
@@ -2460,6 +2831,7 @@ static void ParsePackageAssets(
 			FPackage& Package = *Packages[Index];
 			TArrayView<const uint8> MemView(PackageBuffer, Package.UAssetSize);
 			FMemoryReaderView Ar(MemView);
+			Ar.SetFilterEditorOnly((Package.PackageFlags & EPackageFlags::PKG_FilterEditorOnly) != 0);
 
 			IOSTORE_CPU_SCOPE_DATA(ParsePackage, TCHAR_TO_ANSI(*Package.FileName));
 
@@ -2809,14 +3181,16 @@ static void ProcessLocalizedPackages(
 	}
 }
 
-static void SaveReleaseVersionMeta(const TCHAR* ReleaseVersionOutputDir, const FNameMapBuilder& NameMap, const FPackageGlobalIdMap& PackageGlobalIdMap, const FGlobalImports& GlobalImports, const TArray<FContainerTargetSpec>& ContainerTargets)
+static void SaveReleaseVersionMeta(const TCHAR* ReleaseVersionOutputDir, const FNameMapBuilder& GlobalNameMap, const FPackageGlobalIdMap& PackageGlobalIdMap, const FGlobalImports& GlobalImports, TArray<FContainerTargetSpec*>& ContainerTargets)
 {
+	UE_LOG(LogIoStore, Display, TEXT("Saving release meta data to '%s'"), ReleaseVersionOutputDir);
+
 	IPlatformFile::GetPlatformPhysical().CreateDirectoryTree(ReleaseVersionOutputDir);
 
 	{
 		FString NameMapOutputPath = FPaths::Combine(ReleaseVersionOutputDir, TEXT("iodispatcher.unamemap"));
 		TUniquePtr<FArchive> NameMapArchive(IFileManager::Get().CreateFileWriter(*NameMapOutputPath));
-		(*NameMapArchive) << const_cast<FNameMapBuilder&>(NameMap);
+		(*NameMapArchive) << const_cast<FNameMapBuilder&>(GlobalNameMap);
 	}
 
 	{
@@ -2862,32 +3236,43 @@ static void SaveReleaseVersionMeta(const TCHAR* ReleaseVersionOutputDir, const F
 		}
 	}
 
-	for (const FContainerTargetSpec& ContainerTarget : ContainerTargets)
+	for (FContainerTargetSpec* ContainerTarget : ContainerTargets)
 	{
-		TArray<FTocMetaEntry> TocMetaEntries;
-		TocMetaEntries.Reserve(ContainerTarget.TargetFiles.Num());
-		for (const FContainerTargetFile& TargetFile : ContainerTarget.TargetFiles)
-		{
-			FTocMetaEntry& TocMetaEntry = TocMetaEntries.AddDefaulted_GetRef();
-			TocMetaEntry.Offset = TargetFile.Offset;
-			TocMetaEntry.Size = TargetFile.TargetSize;
-			TocMetaEntry.Hash = TargetFile.Hash;
-		}
-
-		FString TocMetaOutputPath = FPaths::Combine(ReleaseVersionOutputDir, ContainerTarget.Name + TEXT(".utocmeta"));
-		TUniquePtr<FArchive> TocMetaArchive(IFileManager::Get().CreateFileWriter(*TocMetaOutputPath));
-		(*TocMetaArchive) << TocMetaEntries;
-	}
+		FContainerMeta ContainerMeta;
+		ContainerMeta.ContainerId = ContainerTarget->Header.ContainerId;
+		ContainerMeta.ContainerName = ContainerTarget->Name.ToString();
+
+		// Should be safe now as all container(s) has been saved.
+		ContainerMeta.NameMapBuilder = MoveTemp(ContainerTarget->LocalNameMapBuilder);
+
+		FString MetaOutputPath = FPaths::Combine(ReleaseVersionOutputDir, ContainerTarget->Name.ToString() + TEXT(".ucontainermeta"));
+		TUniquePtr<FArchive> Ar(IFileManager::Get().CreateFileWriter(*MetaOutputPath));
+		*Ar << ContainerMeta;
+
+		UE_LOG(LogIoStore, Display,
+			TEXT("Saved container release meta '%s' with container ID '%d' and '%d' names"),
+			*MetaOutputPath,
+			ContainerMeta.ContainerId.ToIndex(),
+			ContainerMeta.NameMapBuilder.GetNameMap().Num());
+	}	
 }
 
-static void LoadReleaseVersionMeta(const TCHAR* ReleaseVersionOutputDir, FNameMapBuilder& NameMap, FPackageGlobalIdMap& PackageGlobalIdMap, FGlobalImports& GlobalImports, TArray<FContainerTargetSpec>& ContainerTargets)
+static void LoadReleaseVersionMeta(
+	const TCHAR* ReleaseVersionOutputDir,
+	FNameMapBuilder& GlobalNameMap,
+	FPackageGlobalIdMap& PackageGlobalIdMap,
+	FGlobalImports& GlobalImports,
+	TArray<FContainerTargetSpec*>& ContainerTargets,
+	TMap<FName, FContainerTargetSpec*>& ContainerTargetMap)
 {
+	UE_LOG(LogIoStore, Display, TEXT("Loading release meta data from '%s'"), ReleaseVersionOutputDir);
+
 	{
 		FString NameMapPath = FPaths::Combine(ReleaseVersionOutputDir, TEXT("iodispatcher.unamemap"));
 		TUniquePtr<FArchive> NameMapArchive(IFileManager::Get().CreateFileReader(*NameMapPath));
 		if (NameMapArchive)
 		{
-			(*NameMapArchive) << NameMap;
+			(*NameMapArchive) << GlobalNameMap;
 		}
 	}
 	{
@@ -2956,13 +3341,26 @@ static void LoadReleaseVersionMeta(const TCHAR* ReleaseVersionOutputDir, FNameMa
 		}
 	}
 
-	for (FContainerTargetSpec& ContainerTarget : ContainerTargets)
+	FString ContainerMetaWildcard = FPaths::Combine(ReleaseVersionOutputDir, TEXT("*.ucontainermeta"));
+	TArray<FString> ContainerMetaFileNames;
+	IFileManager::Get().FindFiles(ContainerMetaFileNames, *ContainerMetaWildcard, true, false);
+	for (FString& ContainerMetaFileName : ContainerMetaFileNames)
 	{
-		FString TocMetaOutputPath = FPaths::Combine(ReleaseVersionOutputDir, ContainerTarget.Name + TEXT(".utocmeta"));
-		TUniquePtr<FArchive> TocMetaArchive(IFileManager::Get().CreateFileReader(*TocMetaOutputPath));
-		if (TocMetaArchive)
+		ContainerMetaFileName = ReleaseVersionOutputDir / ContainerMetaFileName;
+		TUniquePtr<FArchive> Ar(IFileManager::Get().CreateFileReader(*ContainerMetaFileName));
+		if (Ar)
 		{
-			(*TocMetaArchive) << ContainerTarget.PreviousBuildToc;
+			FContainerMeta ContainerMeta;
+			*Ar << ContainerMeta;
+
+			FContainerTargetSpec* ContainerTarget = AddContainer(FName(ContainerMeta.ContainerName), ContainerMeta.ContainerId, ContainerTargets, ContainerTargetMap);
+			ContainerTarget->LocalNameMapBuilder = MoveTemp(ContainerMeta.NameMapBuilder);
+
+			UE_LOG(LogIoStore, Display,
+				TEXT("Loaded container release meta '%s' with container ID '%d' and '%d' names"),
+				*ContainerMetaFileName,
+				ContainerTarget->Header.ContainerId.ToIndex(),
+				ContainerTarget->LocalNameMapBuilder.GetNameMap().Num());
 		}
 	}
 }
@@ -2972,16 +3370,71 @@ void InitializeContainerTargetsAndPackages(
 	TArray<FPackage*>& Packages,
 	FPackageMap& PackageMap,
 	FPackageGlobalIdMap& PackageGlobalIdMap,
-	TArray<FContainerTargetSpec>& ContainerTargets)
+	TArray<FContainerTargetSpec*>& ContainerTargets,
+	TMap<FName, FContainerTargetSpec*>& ContainerTargetMap,
+	FNameMapBuilder& GlobalNameMapBuilder)
 {
-	int32 CookedDirLen = Arguments.CookedDir.Len() + 1;
-	for (int32 ContainerIndex = 0; ContainerIndex < ContainerTargets.Num(); ++ContainerIndex)
+	FString ProjectName = FApp::GetProjectName();
+	FString RelativeEnginePath = FPaths::GetRelativePathToRoot();
+	FString RelativeProjectPath = FPaths::ProjectDir();
+	int32 CookedEngineDirLen = Arguments.CookedDir.Len() + 1;;
+	int32 CookedProjectDirLen = CookedEngineDirLen + ProjectName.Len() + 1;
+
+	auto ConvertCookedPathToRelativePath = [
+		&ProjectName,
+		&CookedEngineDirLen,
+		&CookedProjectDirLen,
+		&RelativeEnginePath,
+		&RelativeProjectPath]
+		(const FString& CookedFile) -> FString
 	{
-		const FContainerSourceSpec& Container = Arguments.Containers[ContainerIndex];
-		FContainerTargetSpec& ContainerTarget = ContainerTargets[ContainerIndex];
+		FString RelativeFileName;
+		const TCHAR* FileName = *CookedFile + CookedEngineDirLen;
+		if (FCString::Strncmp(FileName, *ProjectName, ProjectName.Len()))
+		{
+			int32 FileNameLen = CookedFile.Len() - CookedEngineDirLen;
+			RelativeFileName.Reserve(RelativeEnginePath.Len() + FileNameLen);
+			RelativeFileName = RelativeEnginePath;
+			RelativeFileName.AppendChars(*CookedFile + CookedEngineDirLen, FileNameLen);
+		}
+		else
+		{
+			FileName = *CookedFile + CookedProjectDirLen;
+			int32 FileNameLen = CookedFile.Len() - CookedProjectDirLen;
+			RelativeFileName.Reserve(RelativeProjectPath.Len() + FileNameLen);
+			RelativeFileName = RelativeProjectPath;
+			RelativeFileName.AppendChars(*CookedFile + CookedProjectDirLen, FileNameLen);
+		}
+		return RelativeFileName;
+	};
+
+	for (const FContainerSourceSpec& ContainerSource : Arguments.Containers)
+	{
+		FContainerTargetSpec* ContainerTarget = FindOrAddContainer(ContainerSource.Name, ContainerTargets, ContainerTargetMap);
+		ContainerTarget->OutputPath = ContainerSource.OutputPath;
+		ContainerTarget->bGenerateDiffPatch = ContainerSource.bGenerateDiffPatch;
+
+		for (const FString& PatchSourceContainerFile : ContainerSource.PatchSourceContainerFiles)
+		{
+			FIoStoreEnvironment PatchSourceEnvironment;
+			PatchSourceEnvironment.InitializeFileEnvironment(FPaths::ChangeExtension(PatchSourceContainerFile, TEXT("")));
+			TUniquePtr<FIoStoreReader> PatchSourceReader(new FIoStoreReader());
+			FIoStatus Status = PatchSourceReader->Initialize(PatchSourceEnvironment);
+			if (Status.IsOk())
+			{
+				UE_LOG(LogIoStore, Display, TEXT("Loaded patch source container '%s'"), *PatchSourceContainerFile);
+				ContainerTarget->PatchSourceReaders.Add(MoveTemp(PatchSourceReader));
+			}
+			else
+			{
+				UE_LOG(LogIoStore, Warning, TEXT("Failed loading patch source container '%s' [%s]"), *PatchSourceContainerFile, *Status.ToString())
+			}
+		}
+		
+		ContainerTarget->LocalNameMapBuilder.SetNameMapType(FMappedName::EType::Container);
 		{
 			IOSTORE_CPU_SCOPE(ProcessSourceFiles);
-			for (const FContainerSourceFile& SourceFile : Container.SourceFiles)
+			for (const FContainerSourceFile& SourceFile : ContainerSource.SourceFiles)
 			{
 				const FCookedFileStatData* OriginalCookedFileStatData = Arguments.CookedFileStatMap.Find(SourceFile.NormalizedPath);
 				if (!OriginalCookedFileStatData)
@@ -3002,8 +3455,7 @@ void InitializeContainerTargetsAndPackages(
 					}
 				}
 
-				FString RelativeFileName = TEXT("../../../");
-				RelativeFileName.AppendChars(*NormalizedSourcePath + CookedDirLen, NormalizedSourcePath.Len() - CookedDirLen);
+				FString RelativeFileName = ConvertCookedPathToRelativePath(SourceFile.NormalizedPath);
 				FPackage* Package = nullptr;
 				const bool bIsMemoryMappedBulkData = CookedFileStatData->FileExt == FCookedFileStatData::EFileExt::UMappedBulk;
 
@@ -3019,14 +3471,15 @@ void InitializeContainerTargetsAndPackages(
 
 				if (Package)
 				{
-					FContainerTargetFile& TargetFile = ContainerTarget.TargetFiles.AddDefaulted_GetRef();
+					FContainerTargetFile& TargetFile = ContainerTarget->TargetFiles.AddDefaulted_GetRef();
+					TargetFile.ContainerTarget = ContainerTarget;
 					TargetFile.SourceSize = uint64(CookedFileStatData->FileSize);
 					TargetFile.NormalizedSourcePath = NormalizedSourcePath;
 					TargetFile.TargetPath = MoveTemp(RelativeFileName);
 					TargetFile.Package = Package;
 					if (SourceFile.bNeedsCompression)
 					{
-						ContainerTarget.bIsCompressed = true;
+						ContainerTarget->bIsCompressed = true;
 					}
 					else
 					{
@@ -3062,6 +3515,17 @@ void InitializeContainerTargetsAndPackages(
 					}
 				}
 			}
+
+			if (ContainerTarget->bUseLocalNameMap)
+			{
+				ContainerTarget->NameMapBuilder = &ContainerTarget->LocalNameMapBuilder;
+			}
+			else
+			{
+				// Clear the local container map in case this was used in a previous release
+				ContainerTarget->LocalNameMapBuilder.Empty();
+				ContainerTarget->NameMapBuilder = &GlobalNameMapBuilder;
+			}
 		}
 	}
 
@@ -3089,7 +3553,7 @@ int32 CreateTarget(const FIoStoreArguments& Arguments, const FIoStoreWriterSetti
 	ChunkIdCsv.CreateOutputFile(CookedDir);
 #endif
 
-	FNameMapBuilder NameMapBuilder;
+	FNameMapBuilder GlobalNameMapBuilder;
 	FPackageAssetData PackageAssetData;
 	FGlobalPackageData GlobalPackageData;
 	FExportGraph ExportGraph;
@@ -3102,21 +3566,15 @@ int32 CreateTarget(const FIoStoreArguments& Arguments, const FIoStoreWriterSetti
 	TMap<FName, FPackageHashes> PreviousBuildPackageHashes;
 #endif
 
-	TArray<FContainerTargetSpec> ContainerTargets;
+	TArray<FContainerTargetSpec*> ContainerTargets;
+	TMap<FName, FContainerTargetSpec*> ContainerTargetMap;
 	UE_LOG(LogIoStore, Display, TEXT("Creating container targets..."));
 	{
 		IOSTORE_CPU_SCOPE(CreateContainerTargets);
 
-		for (const FContainerSourceSpec& Container : Arguments.Containers)
-		{
-			FContainerTargetSpec& ContainerTarget = ContainerTargets.AddDefaulted_GetRef();
-			ContainerTarget.OutputPath = Container.OutputPath;
-			ContainerTarget.Name = FPaths::GetBaseFilename(ContainerTarget.OutputPath);
-		}
-
 		if (!Arguments.BasedOnReleaseVersionDir.IsEmpty())
 		{
-			LoadReleaseVersionMeta(*Arguments.BasedOnReleaseVersionDir, NameMapBuilder, PackageGlobalIdMap, GlobalPackageData.Imports, ContainerTargets);
+			LoadReleaseVersionMeta(*Arguments.BasedOnReleaseVersionDir, GlobalNameMapBuilder, PackageGlobalIdMap, GlobalPackageData.Imports, ContainerTargets, ContainerTargetMap);
 #if OUTPUT_DEBUG_PACKAGE_HASHES
 			FString PackageHashesOutputPath = FPaths::Combine(*Arguments.BasedOnReleaseVersionDir, TEXT("iodispatcher.upackagehashes"));
 			TUniquePtr<FArchive> PackageHashesArchive(IFileManager::Get().CreateFileReader(*PackageHashesOutputPath));
@@ -3140,10 +3598,10 @@ int32 CreateTarget(const FIoStoreArguments& Arguments, const FIoStoreWriterSetti
 #endif
 		}
 
-		InitializeContainerTargetsAndPackages(Arguments, Packages, PackageMap, PackageGlobalIdMap, ContainerTargets);
+		InitializeContainerTargetsAndPackages(Arguments, Packages, PackageMap, PackageGlobalIdMap, ContainerTargets, ContainerTargetMap, GlobalNameMapBuilder);
 	}
 
-	ParsePackageAssets(NameMapBuilder, Packages, PackageAssetData, GlobalPackageData, ExportGraph);
+	ParsePackageAssets(Packages, PackageAssetData, ExportGraph);
 	CreateGlobalImportsAndExports(Packages, PackageMap, PackageAssetData, GlobalPackageData, ExportGraph);
 
 	FGlobalPackageImports& PackageImports = GlobalPackageData.Imports.PackageImports;
@@ -3172,12 +3630,42 @@ int32 CreateTarget(const FIoStoreArguments& Arguments, const FIoStoreWriterSetti
 
 	BuildBundles(ExportGraph, Packages);
 
-	FinalizePackageHeaders(
-		Packages,
-		NameMapBuilder,
-		PackageAssetData.ObjectExports,
-		GlobalPackageData.Exports.Objects,
-		GlobalPackageData.Imports.ObjectsByFullName);
+	{
+		IOSTORE_CPU_SCOPE(BuildContainerNameMaps);
+		UE_LOG(LogIoStore, Display, TEXT("Creating container local name map(s)..."));
+
+		for (FContainerTargetSpec* ContainerTarget : ContainerTargets)
+		{
+			BuildContainerNameMap(*ContainerTarget);
+		}
+	}
+
+	{
+		IOSTORE_CPU_SCOPE(FinalizePackageHeaders);
+		UE_LOG(LogIoStore, Display, TEXT("Finalizing package headers..."));
+
+		for (FContainerTargetSpec* ContainerTarget : ContainerTargets)
+		{
+			FinalizePackageHeaders(
+				*ContainerTarget,
+				PackageAssetData.ObjectExports,
+				GlobalPackageData.Exports.Objects,
+				GlobalPackageData.Imports.ObjectsByFullName);
+
+			const FNameMapBuilder& NameMapBuilder = *ContainerTarget->NameMapBuilder;
+			SaveNameBatch(ContainerTarget->LocalNameMapBuilder.GetNameMap(), ContainerTarget->Header.Names, ContainerTarget->Header.NameHashes);
+			ContainerTarget->Header.PackageIds.Reserve(ContainerTarget->TargetFiles.Num());
+			ContainerTarget->Header.PackageNames.Reserve(ContainerTarget->TargetFiles.Num());
+			for (FContainerTargetFile& TargetFile : ContainerTarget->TargetFiles)
+			{
+				if (!TargetFile.bIsBulkData)
+				{
+					ContainerTarget->Header.PackageIds.Add(TargetFile.Package->GlobalPackageId);
+					ContainerTarget->Header.PackageNames.Add(NameMapBuilder.MapName(TargetFile.Package->Name));
+				}
+			}
+		}
+	}
 
 	FLargeMemoryWriter GlobalMetaArchive(0, true);
 	FLargeMemoryWriter InitialLoadArchive(0, true);
@@ -3191,35 +3679,30 @@ int32 CreateTarget(const FIoStoreArguments& Arguments, const FIoStoreWriterSetti
 
 		FLargeMemoryWriter StoreTocArchive(0, true);
 		FLargeMemoryWriter StoreDataArchive(0, true);
-		FLargeMemoryWriter GlobalImportNamesArchive(0, true);
-
-		SerializeImportNames(NameMapBuilder, PackageImports, GlobalImportNamesArchive);
 
 		SerializePackageStoreEntries(
 			Packages,
-			NameMapBuilder,
 			StoreTocArchive,
 			StoreDataArchive);
 
 		SerializeInitialLoad(
-			NameMapBuilder,
+			GlobalNameMapBuilder,
 			ScriptImports,
 			InitialLoadArchive);
 
 		StoreTocByteCount = StoreTocArchive.TotalSize();
 		StoreDataByteCount = StoreDataArchive.TotalSize();
-		GlobalImportNamesByteCount = GlobalImportNamesArchive.TotalSize();
 
 		int32 NumPackages = Packages.Num();
+		int32 PackageImportCount = PackageImports.Num();
 		int32 StoreByteCount = StoreTocByteCount + StoreDataByteCount;
 
 		GlobalMetaArchive << NumPackages;
+		GlobalMetaArchive << PackageImportCount;
 		GlobalMetaArchive << StoreByteCount;
 		GlobalMetaArchive.Serialize(StoreTocArchive.GetData(), StoreTocByteCount);
 		GlobalMetaArchive.Serialize(StoreDataArchive.GetData(), StoreDataByteCount);
 
-		GlobalMetaArchive.Serialize(GlobalImportNamesArchive.GetData(), GlobalImportNamesByteCount);
-
 		GlobalMetaArchive << CulturePackageMap;
 	}
 
@@ -3228,15 +3711,15 @@ int32 CreateTarget(const FIoStoreArguments& Arguments, const FIoStoreWriterSetti
 		IOSTORE_CPU_SCOPE(CalculateHashes);
 
 		int32 TotalFileCount = 0;
-		for (FContainerTargetSpec& ContainerTarget : ContainerTargets)
+		for (const FContainerTargetSpec* ContainerTarget : ContainerTargets)
 		{
-			TotalFileCount += ContainerTarget.TargetFiles.Num();
+			TotalFileCount += ContainerTarget->TargetFiles.Num();
 		}
 		TArray<FContainerTargetFile*> AllTargetFiles;
 		AllTargetFiles.Reserve(TotalFileCount);
-		for (FContainerTargetSpec& ContainerTarget : ContainerTargets)
+		for (FContainerTargetSpec* ContainerTarget : ContainerTargets)
 		{
-			for (FContainerTargetFile& TargetFile : ContainerTarget.TargetFiles)
+			for (FContainerTargetFile& TargetFile : ContainerTarget->TargetFiles)
 			{
 				AllTargetFiles.Add(&TargetFile);
 			}
@@ -3255,18 +3738,18 @@ int32 CreateTarget(const FIoStoreArguments& Arguments, const FIoStoreWriterSetti
 			delete FileHandle;
 			if (!TargetFile->bIsBulkData)
 			{
-				IoBuffer = CreateExportBundleBuffer(TargetFile->Package, ObjectExports, IoBuffer);
+				IoBuffer = CreateExportBundleBuffer(*TargetFile, ObjectExports, IoBuffer);
 				TargetFile->TargetSize = IoBuffer.DataSize();
 			}
 			else
 			{
 				TargetFile->TargetSize = TargetFile->SourceSize;
 			}
-			FSHA1::HashBuffer(IoBuffer.Data(), IoBuffer.DataSize(), TargetFile->Hash.Hash);
+			TargetFile->ChunkHash = FIoChunkHash::HashBuffer(IoBuffer.Data(), IoBuffer.DataSize());
 #if OUTPUT_DEBUG_PACKAGE_HASHES
 			if (!TargetFile->bIsBulkData)
 			{
-				TargetFile->Package->Hashes.ExportBundleHash = TargetFile->Hash;
+				TargetFile->Package->Hashes.ExportBundleHash = TargetFile->ChunkHash;
 			}
 #endif
 			uint64 LocalFileIndex = CurrentFileIndex.IncrementExchange() + 1;
@@ -3300,6 +3783,20 @@ int32 CreateTarget(const FIoStoreArguments& Arguments, const FIoStoreWriterSetti
 #endif
 	}
 
+	{
+		IOSTORE_CPU_SCOPE(MapAdditionalBulkDataChunkIds);
+		for (FContainerTargetSpec* ContainerTarget : ContainerTargets)
+		{
+			for (FContainerTargetFile& TargetFile : ContainerTarget->TargetFiles)
+			{
+				if (TargetFile.bIsBulkData)
+				{
+					MapAdditionalBulkDataChunks(TargetFile, BulkDataManifest);
+				}
+			}
+		}
+	}
+
 	UE_LOG(LogIoStore, Display, TEXT("Creating disk layout..."));
 	CreateDiskLayout(ContainerTargets, Packages, Arguments.GameOrderMap, Arguments.CookerOrderMap, GeneralIoWriterSettings.CompressionBlockSize, Arguments.MemoryMappingAlignment);
 
@@ -3310,33 +3807,29 @@ int32 CreateTarget(const FIoStoreArguments& Arguments, const FIoStoreWriterSetti
 	FIoStoreWriter* GlobalIoStoreWriter;
 	{
 		IOSTORE_CPU_SCOPE(InitializeIoStoreWriters);
-		FString GlobalOutputPath = FString(Arguments.OutputDir) / TEXT("global");
-		GlobalIoStoreEnv.InitializeFileEnvironment(*GlobalOutputPath);
-		GlobalIoStoreWriter = new FIoStoreWriter(GlobalIoStoreEnv);
+		GlobalIoStoreEnv.InitializeFileEnvironment(*Arguments.GlobalContainerPath);
+		GlobalIoStoreWriter = new FIoStoreWriter(GlobalIoStoreEnv, FIoContainerId::FromIndex(0));
 		IoStoreWriters.Add(GlobalIoStoreWriter);
-		for (FContainerTargetSpec& ContainerTarget : ContainerTargets)
+		for (FContainerTargetSpec* ContainerTarget : ContainerTargets)
 		{
-			if (ContainerTarget.OutputPath.IsEmpty())
-			{
-				ContainerTarget.IoStoreWriter = GlobalIoStoreWriter;
-			}
-			else
+			check(ContainerTarget->Header.ContainerId.IsValid());
+			if (!ContainerTarget->OutputPath.IsEmpty())
 			{
-				ContainerTarget.IoStoreEnv.Reset(new FIoStoreEnvironment());
-				ContainerTarget.IoStoreEnv->InitializeFileEnvironment(ContainerTarget.OutputPath);
-				ContainerTarget.IoStoreWriter = new FIoStoreWriter(*ContainerTarget.IoStoreEnv);
-				IoStoreWriters.Add(ContainerTarget.IoStoreWriter);
+				ContainerTarget->IoStoreEnv.Reset(new FIoStoreEnvironment());
+				ContainerTarget->IoStoreEnv->InitializeFileEnvironment(ContainerTarget->OutputPath);
+				ContainerTarget->IoStoreWriter = new FIoStoreWriter(*ContainerTarget->IoStoreEnv, ContainerTarget->Header.ContainerId);
+				IoStoreWriters.Add(ContainerTarget->IoStoreWriter);
 			}
 		}
 		FIoStatus IoStatus = IoStoreWriterContext->Initialize(GeneralIoWriterSettings);
 		check(IoStatus.IsOk());
 		IoStatus = GlobalIoStoreWriter->Initialize(*IoStoreWriterContext, /* bIsCompressed */ false);
 		check(IoStatus.IsOk());
-		for (FContainerTargetSpec& ContainerTarget : ContainerTargets)
+		for (FContainerTargetSpec* ContainerTarget : ContainerTargets)
 		{
-			if (ContainerTarget.IoStoreWriter != GlobalIoStoreWriter)
+			if (ContainerTarget->IoStoreWriter && ContainerTarget->IoStoreWriter != GlobalIoStoreWriter)
 			{
-				IoStatus = ContainerTarget.IoStoreWriter->Initialize(*IoStoreWriterContext, ContainerTarget.bIsCompressed);
+				IoStatus = ContainerTarget->IoStoreWriter->Initialize(*IoStoreWriterContext, ContainerTarget->bIsCompressed);
 				check(IoStatus.IsOk());
 			}
 		}
@@ -3355,22 +3848,53 @@ int32 CreateTarget(const FIoStoreArguments& Arguments, const FIoStoreWriterSetti
 		};
 
 		int32 TotalFileCount = 0;
-		for (FContainerTargetSpec& ContainerTarget : ContainerTargets)
+		for (FContainerTargetSpec* ContainerTarget : ContainerTargets)
+		{
+			check(ContainerTarget->IoStoreWriter || ContainerTarget->TargetFiles.Num() == 0);
+			TotalFileCount += ContainerTarget->TargetFiles.Num();
+		}
+
+		TArray<uint8> ContainerHeaderPadding;
+		if (GeneralIoWriterSettings.CompressionBlockSize > 0)
 		{
-			TotalFileCount += ContainerTarget.TargetFiles.Num();
+			ContainerHeaderPadding.AddZeroed(GeneralIoWriterSettings.CompressionBlockSize);
 		}
 
 		TArray<FReadFileTask> ReadFileTasks;
 		ReadFileTasks.SetNum(TotalFileCount);
 		int32 ReadFileTaskIndex = 0;
-		for (FContainerTargetSpec& ContainerTarget : ContainerTargets)
+		for (FContainerTargetSpec* ContainerTarget : ContainerTargets)
 		{
-			for (FContainerTargetFile& TargetFile : ContainerTarget.TargetFiles)
+			if (!ContainerTarget->IoStoreWriter)
+			{
+				continue;
+			}
+			for (FContainerTargetFile& TargetFile : ContainerTarget->TargetFiles)
 			{
 				FReadFileTask& ReadFileTask = ReadFileTasks[ReadFileTaskIndex++];
 				ReadFileTask.ContainerTargetFile = &TargetFile;
-				ReadFileTask.IoStoreWriter = ContainerTarget.IoStoreWriter;
+				ReadFileTask.IoStoreWriter = ContainerTarget->IoStoreWriter;
+			}
+
+			FLargeMemoryWriter Ar(0, true);
+			Ar << ContainerTarget->Header;
+
+			if (GeneralIoWriterSettings.CompressionBlockSize > 0)
+			{
+				const uint64 RemainingInBlock = Align(Ar.TotalSize(), GeneralIoWriterSettings.CompressionBlockSize) - Ar.TotalSize();
+				if (RemainingInBlock > 0)
+				{
+					Ar.Serialize(ContainerHeaderPadding.GetData(), RemainingInBlock);
+				}
 			}
+
+			FIoWriteOptions WriteOptions;
+			WriteOptions.DebugName = TEXT("ContainerHeader");
+			FIoStatus Status = ContainerTarget->IoStoreWriter->Append(
+				CreateIoChunkId(0, ContainerTarget->Header.ContainerId.ToIndex(), EIoChunkType::ContainerHeader), 
+				FIoBuffer(FIoBuffer::Wrap, Ar.GetData(), Ar.TotalSize()), WriteOptions);
+
+			UE_CLOG(!Status.IsOk(), LogIoStore, Error, TEXT("Failed to serialize container header"));
 		}
 
 		FEvent* TaskStartedEvent = FPlatformProcess::GetSynchEventFromPool(false);
@@ -3401,7 +3925,7 @@ int32 CreateTarget(const FIoStoreArguments& Arguments, const FIoStoreWriterSetti
 				uint64 BufferSize = ReadFileTask.IoBuffer.DataSize();
 				if (!TargetFile.bIsBulkData)
 				{
-					ReadFileTask.IoBuffer = CreateExportBundleBuffer(TargetFile.Package, ObjectExports, ReadFileTask.IoBuffer);
+					ReadFileTask.IoBuffer = CreateExportBundleBuffer(TargetFile, ObjectExports, ReadFileTask.IoBuffer);
 				}
 
 				if (TargetFile.Padding)
@@ -3412,9 +3936,18 @@ int32 CreateTarget(const FIoStoreArguments& Arguments, const FIoStoreWriterSetti
 				WriteOptions.DebugName = *TargetFile.TargetPath;
 				WriteOptions.bForceUncompressed = TargetFile.bForceUncompressed;
 				WriteOptions.Alignment = TargetFile.Alignment;
-				FIoStatus Status = ReadFileTask.IoStoreWriter->Append(TargetFile.ChunkId, ReadFileTask.IoBuffer, WriteOptions);
+				FIoStatus Status = ReadFileTask.IoStoreWriter->Append(TargetFile.ChunkId, TargetFile.ChunkHash, ReadFileTask.IoBuffer, WriteOptions);
 				UE_CLOG(!Status.IsOk(), LogIoStore, Fatal, TEXT("Failed to append chunk to container file due to '%s'"), *Status.ToString());
 
+				for (const FContainerTargetFilePartialMapping& PartialMapping : TargetFile.PartialMappings)
+				{
+					const FIoStatus PartialResult = ReadFileTask.IoStoreWriter->MapPartialRange(TargetFile.ChunkId, PartialMapping.Offset, PartialMapping.Length, PartialMapping.PartialChunkId);
+					if (!PartialResult.IsOk())
+					{
+						UE_LOG(LogIoStore, Warning, TEXT("Failed to map partial range for '%s' due to: %s"), *TargetFile.Package->FileName, *PartialResult.ToString());
+					}
+				}
+
 				ReadFileTask.IoBuffer = FIoBuffer();
 				BufferMemoryAllocated -= BufferSize;
 				TaskFinishedEvent->Trigger();
@@ -3455,20 +3988,6 @@ int32 CreateTarget(const FIoStoreArguments& Arguments, const FIoStoreWriterSetti
 		FPlatformProcess::ReturnSynchEventToPool(TaskStartedEvent);
 	}
 
-	{
-		IOSTORE_CPU_SCOPE(MapAdditionalBulkDataChunkIds);
-		for (FContainerTargetSpec& ContainerTarget : ContainerTargets)
-		{
-			for (FContainerTargetFile& TargetFile : ContainerTarget.TargetFiles)
-			{
-				if (TargetFile.bIsBulkData)
-				{
-					MapAdditionalBulkDataChunks(TargetFile, BulkDataManifest, ContainerTarget.IoStoreWriter);
-				}
-			}
-		}
-	}
-
 	{
 		UE_LOG(LogIoStore, Display, TEXT("Saving global meta data to container file"));
 		FIoWriteOptions WriteOptions;
@@ -3500,7 +4019,7 @@ int32 CreateTarget(const FIoStoreArguments& Arguments, const FIoStoreWriterSetti
 
 		TArray<uint8> Names;
 		TArray<uint8> Hashes;
-		SaveNameBatch(NameMapBuilder.GetNameMap(), /* out */ Names, /* out */ Hashes);
+		SaveNameBatch(GlobalNameMapBuilder.GetNameMap(), /* out */ Names, /* out */ Hashes);
 
 		GlobalNamesMB = Names.Num() >> 20;
 		GlobalNameHashesMB = Hashes.Num() >> 20;
@@ -3534,7 +4053,7 @@ int32 CreateTarget(const FIoStoreArguments& Arguments, const FIoStoreWriterSetti
 
 	if (!Arguments.OutputReleaseVersionDir.IsEmpty())
 	{
-		SaveReleaseVersionMeta(*Arguments.OutputReleaseVersionDir, NameMapBuilder, PackageGlobalIdMap, GlobalPackageData.Imports, ContainerTargets);
+		SaveReleaseVersionMeta(*Arguments.OutputReleaseVersionDir, GlobalNameMapBuilder, PackageGlobalIdMap, GlobalPackageData.Imports, ContainerTargets);
 #if OUTPUT_DEBUG_PACKAGE_HASHES
 		FString PackageHashesOutputPath = FPaths::Combine(*Arguments.OutputReleaseVersionDir, TEXT("iodispatcher.upackagehashes"));
 		TUniquePtr<FArchive> PackageHashesArchive(IFileManager::Get().CreateFileWriter(*PackageHashesOutputPath));
@@ -3587,7 +4106,7 @@ int32 CreateTarget(const FIoStoreArguments& Arguments, const FIoStoreWriterSetti
 				if (PreviousHashes->ExportBundleHash != Package->Hashes.ExportBundleHash)
 				{
 					++ModifiedExportBundlesCount;
-					ModifiedExportBundlesSize += Package->HeaderData.Num() + Package->ExportsSerialSize;
+					ModifiedExportBundlesSize += Package->ExportBundlesHeaderSize + Package->ExportsSerialSize;
 				}
 				if (bModified)
 				{
@@ -3654,7 +4173,7 @@ int32 CreateTarget(const FIoStoreArguments& Arguments, const FIoStoreWriterSetti
 		ImportMapSize += Package.ImportMapSize;
 		ExportMapSize += Package.ExportMapSize;
 		NameMapSize += Package.NameMapSize;
-		NameMapCount += Package.NameIndices.Num();
+		NameMapCount += Package.NameMap.Num();
 		CircularPackagesCount += Package.bHasCircularImportDependencies;
 		TotalInternalArcCount += Package.InternalArcs.Num();
 		ImportedPackagesCount += Package.ImportedPackagesSerializeCount;
@@ -3682,7 +4201,7 @@ int32 CreateTarget(const FIoStoreArguments& Arguments, const FIoStoreWriterSetti
 	UE_LOG(LogIoStore, Display, TEXT("--------------------------------------------- IoStore Summary -----------------------------------------------"));
 
 	UE_LOG(LogIoStore, Display, TEXT(""));
-	UE_LOG(LogIoStore, Display, TEXT("%-35s %15s %15s %15s %25s"), TEXT("Container"), TEXT("TOC Size (KB)"), TEXT("TOC Entries"), TEXT("Size (MB)"), TEXT("Compressed (MB)"));
+	UE_LOG(LogIoStore, Display, TEXT("%-4s %-35s %15s %15s %15s %25s"), TEXT("ID"), TEXT("Container"), TEXT("TOC Size (KB)"), TEXT("TOC Entries"), TEXT("Size (MB)"), TEXT("Compressed (MB)"));
 	UE_LOG(LogIoStore, Display, TEXT("-------------------------------------------------------------------------------------------------------------"));
 	int64 TotalPaddingSize = 0;
 	for (const FIoStoreWriterResult& Result : IoStoreWriterResults)
@@ -3691,7 +4210,8 @@ int32 CreateTarget(const FIoStoreArguments& Arguments, const FIoStoreWriterSetti
 			? 0.0
 			: (double(Result.UncompressedContainerSize - Result.CompressedContainerSize) / double(Result.UncompressedContainerSize)) * 100.0;
 
-		UE_LOG(LogIoStore, Display, TEXT("%-35s %15.2lf %15d %15.2lf %25s"),
+		UE_LOG(LogIoStore, Display, TEXT("%-4d %-35s %15.2lf %15d %15.2lf %25s"),
+			Result.ContainerId.ToIndex(),
 			*Result.ContainerName,
 			(double)Result.TocSize / 1024.0,
 			Result.TocEntryCount,
@@ -3702,13 +4222,13 @@ int32 CreateTarget(const FIoStoreArguments& Arguments, const FIoStoreWriterSetti
 				*Result.CompressionMethod.ToString()));
 		TotalPaddingSize += Result.PaddingSize;
 	}
-	UE_LOG(LogIoStore, Display, TEXT("Padding: %8.2f MB"), TotalPaddingSize / 1024.0 / 1024.0);
+	UE_LOG(LogIoStore, Display, TEXT("Compression block padding: %8.2f MB"), TotalPaddingSize / 1024.0 / 1024.0);
 	UE_LOG(LogIoStore, Display, TEXT(""));
 	UE_LOG(LogIoStore, Display, TEXT("Packages: %8d total, %d circular dependencies, %d no preload dependencies, %d no import dependencies"),
 		PackageMap.Num(), CircularPackagesCount, PackagesWithoutPreloadDependenciesCount, PackagesWithoutImportDependenciesCount);
 	UE_LOG(LogIoStore, Display, TEXT("Bundles:  %8d total, %d entries, %d export objects"), BundleCount, BundleEntryCount, GlobalPackageData.Exports.Objects.Num());
 
-	UE_LOG(LogIoStore, Display, TEXT("IoStore: %8.2f MB GlobalNames, %d unique names"), (double)GlobalNamesMB, NameMapBuilder.GetNameMap().Num());
+	UE_LOG(LogIoStore, Display, TEXT("IoStore: %8.2f MB GlobalNames, %d unique names"), (double)GlobalNamesMB, GlobalNameMapBuilder.GetNameMap().Num());
 	UE_LOG(LogIoStore, Display, TEXT("IoStore: %8.2f MB GlobalNameHashes"), (double)GlobalNameHashesMB);
 	UE_LOG(LogIoStore, Display, TEXT("IoStore: %8.2f MB GlobalPackageStoreToc"), (double)StoreTocByteCount / 1024.0 / 1024.0);
 	UE_LOG(LogIoStore, Display, TEXT("IoStore: %8.2f MB GlobalPackageStoreData, %d imported packages, %d public exports, %d export bundles"),
@@ -4003,16 +4523,9 @@ int32 CreateIoStoreContainerFiles(const TCHAR* CmdLine)
 	FParse::Value(CmdLine, TEXT("-CreateReleaseVersionDirectory="), Arguments.OutputReleaseVersionDir);
 	FParse::Value(CmdLine, TEXT("-BasedOnReleaseVersionDirectory="), Arguments.BasedOnReleaseVersionDir);
 
-	FString CommandListFile;
-	if (FParse::Value(FCommandLine::Get(), TEXT("Commands="), CommandListFile))
+	if (FParse::Value(FCommandLine::Get(), TEXT("CreateGlobalContainer="), Arguments.GlobalContainerPath))
 	{
-		UE_LOG(LogIoStore, Display, TEXT("Using command list file: '%s'"), *CommandListFile);
-
-		if (!FParse::Value(FCommandLine::Get(), TEXT("OutputDirectory="), Arguments.OutputDir))
-		{
-			UE_LOG(LogIoStore, Error, TEXT("OutputDirectory must be specified"));
-			return 1;
-		}
+		Arguments.GlobalContainerPath = FPaths::ChangeExtension(Arguments.GlobalContainerPath, TEXT(""));
 
 		if (!FParse::Value(FCommandLine::Get(), TEXT("CookedDirectory="), Arguments.CookedDir))
 		{
@@ -4020,42 +4533,74 @@ int32 CreateIoStoreContainerFiles(const TCHAR* CmdLine)
 			return 1;
 		}
 
-		TArray<FString> Commands;
-		if (!FFileHelper::LoadFileToStringArray(Commands, *CommandListFile))
-		{
-			UE_LOG(LogIoStore, Error, TEXT("Failed to read command list file '%s'."), *CommandListFile);
-			return -1;
-		}
-
-		Arguments.Containers.Reserve(Commands.Num());
-		for (const FString& Command : Commands)
+		FContainerSourceSpec* CookedFilesVisitorContainerSpec = nullptr;
+		FString CommandListFile;
+		if (FParse::Value(FCommandLine::Get(), TEXT("Commands="), CommandListFile))
 		{
-			FContainerSourceSpec& ContainerSpec = Arguments.Containers.AddDefaulted_GetRef();
-
-			if (!FParse::Value(*Command, TEXT("Output="), ContainerSpec.OutputPath))
+			UE_LOG(LogIoStore, Display, TEXT("Using command list file: '%s'"), *CommandListFile);
+			TArray<FString> Commands;
+			if (!FFileHelper::LoadFileToStringArray(Commands, *CommandListFile))
 			{
-				UE_LOG(LogIoStore, Error, TEXT("Output argument missing from command '%s'"), *Command);
+				UE_LOG(LogIoStore, Error, TEXT("Failed to read command list file '%s'."), *CommandListFile);
 				return -1;
 			}
-			ContainerSpec.OutputPath = FPaths::Combine(Arguments.OutputDir, ContainerSpec.OutputPath);
 
-			FString ResponseFilePath;
-			if (!FParse::Value(*Command, TEXT("ResponseFile="), ResponseFilePath))
+			Arguments.Containers.Reserve(Commands.Num());
+			for (const FString& Command : Commands)
 			{
-				UE_LOG(LogIoStore, Error, TEXT("ResponseFile argument missing from command '%s'"), *Command);
-				return -1;
-			}
+				FContainerSourceSpec& ContainerSpec = Arguments.Containers.AddDefaulted_GetRef();
 
-			if (!ParsePakResponseFile(*ResponseFilePath, ContainerSpec.SourceFiles))
-			{
-				UE_LOG(LogIoStore, Error, TEXT("Failed to parse Pak response file '%s'"), *ResponseFilePath);
-				return -1;
+				if (!FParse::Value(*Command, TEXT("Output="), ContainerSpec.OutputPath))
+				{
+					UE_LOG(LogIoStore, Error, TEXT("Output argument missing from command '%s'"), *Command);
+					return -1;
+				}
+				ContainerSpec.OutputPath = FPaths::ChangeExtension(ContainerSpec.OutputPath, TEXT(""));
+
+				FString ContainerName;
+				if (!FParse::Value(*Command, TEXT("ContainerName="), ContainerName))
+				{
+					UE_LOG(LogIoStore, Error, TEXT("ContainerName argument missing from command '%s'"), *Command);
+					return -1;
+				}
+				ContainerSpec.Name = FName(ContainerName);
+
+				FString PatchSourceWildcard;
+				if (FParse::Value(*Command, TEXT("PatchSource="), PatchSourceWildcard))
+				{
+					IFileManager::Get().FindFiles(ContainerSpec.PatchSourceContainerFiles, *PatchSourceWildcard, true, false);
+					FString PatchSourceContainersDirectory = FPaths::GetPath(*PatchSourceWildcard);
+					for (FString& PatchSourceContainerFile : ContainerSpec.PatchSourceContainerFiles)
+					{
+						PatchSourceContainerFile = PatchSourceContainersDirectory / PatchSourceContainerFile;
+						FPaths::NormalizeFilename(PatchSourceContainerFile);
+					}
+				}
+
+				ContainerSpec.bGenerateDiffPatch = FParse::Param(*Command, TEXT("GenerateDiffPatch"));
+				
+				FString ResponseFilePath;
+				if (!FParse::Value(*Command, TEXT("ResponseFile="), ResponseFilePath))
+				{
+					UE_LOG(LogIoStore, Error, TEXT("ResponseFile argument missing from command '%s'"), *Command);
+					return -1;
+				}
+
+				if (!ParsePakResponseFile(*ResponseFilePath, ContainerSpec.SourceFiles))
+				{
+					UE_LOG(LogIoStore, Error, TEXT("Failed to parse Pak response file '%s'"), *ResponseFilePath);
+					return -1;
+				}
 			}
 		}
+		else
+		{
+			CookedFilesVisitorContainerSpec = &Arguments.Containers.AddDefaulted_GetRef();
+		}
 
 		UE_LOG(LogIoStore, Display, TEXT("Searching for cooked assets in folder '%s'"), *Arguments.CookedDir);
 		FCookedFileStatMap CookedFileStatMap;
-		FCookedFileVisitor CookedFileVistor(Arguments.CookedFileStatMap);
+		FCookedFileVisitor CookedFileVistor(Arguments.CookedFileStatMap, CookedFilesVisitorContainerSpec);
 		IFileManager::Get().IterateDirectoryStatRecursively(*Arguments.CookedDir, CookedFileVistor);
 		UE_LOG(LogIoStore, Display, TEXT("Found '%d' files"), Arguments.CookedFileStatMap.Num());
 
@@ -4067,30 +4612,8 @@ int32 CreateIoStoreContainerFiles(const TCHAR* CmdLine)
 	}
 	else
 	{
-		const TArray<ITargetPlatform*>& Platforms = GetTargetPlatformManagerRef().GetActiveTargetPlatforms();
-
-		for (const ITargetPlatform* TargetPlatform : Platforms)
-		{
-			Arguments.CookedDir = FPaths::ProjectSavedDir() / TEXT("Cooked") / TargetPlatform->PlatformName();
-			Arguments.OutputDir = Arguments.CookedDir / FApp::GetProjectName();
-
-			TArray<FContainerSourceSpec> Containers;
-			FContainerSourceSpec& ContainerSpec = Arguments.Containers.AddDefaulted_GetRef();
-			FCookedFileVisitor CookedFileVistor(Arguments.CookedFileStatMap, &ContainerSpec);
-
-			UE_LOG(LogIoStore, Display, TEXT("Searching for cooked files..."));
-			IFileManager::Get().IterateDirectoryStatRecursively(*Arguments.CookedDir, CookedFileVistor);
-
-			UE_LOG(LogIoStore, Display, TEXT("Found '%d' files"), ContainerSpec.SourceFiles.Num());
-
-			UE_LOG(LogIoStore, Display, TEXT("Creating target: '%s' using output path: '%s'"), *TargetPlatform->PlatformName(), *Arguments.OutputDir);
-
-			int32 ReturnValue = CreateTarget(Arguments, GeneralIoWriterSettings);
-			if (ReturnValue != 0)
-			{
-				return ReturnValue;
-			}
-		}
+		UE_LOG(LogIoStore, Error, TEXT("Nothing to do!"));
+		return -1;
 	}
 
 	return 0;
diff --git a/Engine/Source/Programs/AutomationTool/Scripts/CopyBuildToStagingDirectory.Automation.cs b/Engine/Source/Programs/AutomationTool/Scripts/CopyBuildToStagingDirectory.Automation.cs
--- a/Engine/Source/Programs/AutomationTool/Scripts/CopyBuildToStagingDirectory.Automation.cs
+++ b/Engine/Source/Programs/AutomationTool/Scripts/CopyBuildToStagingDirectory.Automation.cs
@@ -165,27 +165,33 @@ public partial class Project : CommandUtils
 
 	static private string GetIoStoreCommandArguments(
 		Dictionary<string, string> UnrealPakResponseFile,
-		StagedFileReference ContainerRelativePath,
-		FileReference PakOrderFileLocation,
-		string PlatformOptions,
-		bool Compressed,
+		string ContainerName,
+		FileReference PakOutputLocation,
+		bool bCompressed,
 		EncryptionAndSigning.CryptoSettings CryptoSettings,
 		string EncryptionKeyGuid,
-		FileReference SecondaryPakOrderFileLocation=null)
+		string PatchSourceContentPath,
+		bool bGenerateDiffPatch)
 	{
 		StringBuilder CmdLine = new StringBuilder();
-		CmdLine.AppendFormat("-Output={0}", MakePathSafeToUseWithCommandLine(ContainerRelativePath.Name));
+		CmdLine.AppendFormat("-Output={0}", MakePathSafeToUseWithCommandLine(Path.ChangeExtension(PakOutputLocation.FullName, ".utoc")));
+		CmdLine.AppendFormat("-ContainerName={0}", ContainerName);
+		if (!String.IsNullOrEmpty(PatchSourceContentPath))
+		{
+			CmdLine.AppendFormat(" -PatchSource={0}", CommandUtils.MakePathSafeToUseWithCommandLine(PatchSourceContentPath));
+		}
+		if (bGenerateDiffPatch)
+		{
+			CmdLine.Append(" -GenerateDiffPatch");
+		}
 
 		// Force encryption of ALL files if we're using specific encryption key. This should be made an option per encryption key in the settings, but for our initial
 		// implementation we will just assume that we require maximum security for this data.
 		bool bForceEncryption = !string.IsNullOrEmpty(EncryptionKeyGuid);
-		string ContainerName = Path.GetFileNameWithoutExtension(ContainerRelativePath.Name);
 		string UnrealPakResponseFileName = CombinePaths(CmdEnv.LogFolder, "PakListIoStore_" + ContainerName + ".txt");
-		WritePakResponseFile(UnrealPakResponseFileName, UnrealPakResponseFile, Compressed, CryptoSettings, bForceEncryption);
+		WritePakResponseFile(UnrealPakResponseFileName, UnrealPakResponseFile, bCompressed, CryptoSettings, bForceEncryption);
 		CmdLine.AppendFormat(" -ResponseFile={0}", CommandUtils.MakePathSafeToUseWithCommandLine(UnrealPakResponseFileName));
 
-		CmdLine.Append(PlatformOptions);
-
 		return CmdLine.ToString();
 	}
 
@@ -2336,31 +2342,22 @@ public partial class Project : CommandUtils
 							}
 						}
 
-						string ContainerName = PakParams.PakName + "-" + SC.FinalCookPlatform;
-						StagedFileReference ContainerRelativeLocation;
-						if (Params.HasDLCName)
-						{
-							ContainerRelativeLocation = StagedFileReference.Combine(Params.DLCFile.Directory.MakeRelativeTo(SC.ProjectRoot), "Content", "Paks", SC.FinalCookPlatform, Params.DLCFile.GetFileNameWithoutExtension() + ContainerName);
-						}
-						else
-						{
-							ContainerRelativeLocation = StagedFileReference.Combine("Content", "Paks", ContainerName);
-						}
-						if (SC.StageTargetPlatform.DeployLowerCaseFilenames())
+						string ContainerPatchSourcePath = null;
+						if (Params.HasBasedOnReleaseVersion)
 						{
-							ContainerRelativeLocation = ContainerRelativeLocation.ToLowerInvariant();
+							string ContainerWildcard = PakParams.PakName + "-" + SC.FinalCookPlatform + "*.utoc";
+							ContainerPatchSourcePath = CombinePaths(Params.GetBasedOnReleaseVersionPath(SC, Params.Client), ContainerWildcard);
 						}
-						ContainerRelativeLocation = SC.StageTargetPlatform.Remap(ContainerRelativeLocation);
-
+						bool bGenerateDiffPatch = bShouldGeneratePatch && !ShouldSkipGeneratingPatch(PlatformGameConfig, PakParams.PakName);
 						IoStoreCommands.Add(GetIoStoreCommandArguments(
 							IoStoreResponseFile,
-							ContainerRelativeLocation,
-							PrimaryOrderFile,
-							AdditionalArgs,
+							PakParams.PakName,
+							OutputLocation,
 							PakParams.bCompressed,
 							CryptoSettings,
 							PakParams.EncryptionKeyGuid,
-							SecondaryOrderFile));
+							ContainerPatchSourcePath,
+							bGenerateDiffPatch));
 					}
 
 					Commands.Add(GetUnrealPakArguments(
@@ -2446,6 +2443,11 @@ public partial class Project : CommandUtils
 
 				InternalUtils.SafeCreateDirectory(Path.GetDirectoryName(ReleaseVersionPath));
 				InternalUtils.SafeCopyFile(OutputLocation.FullName, ReleaseVersionPath);
+				if (ShouldCreateIoStoreContainerFiles(Params, SC))
+				{
+					InternalUtils.SafeCopyFile(Path.ChangeExtension(OutputLocation.FullName, ".utoc"), Path.ChangeExtension(ReleaseVersionPath, ".utoc"));
+					InternalUtils.SafeCopyFile(Path.ChangeExtension(OutputLocation.FullName, ".ucas"), Path.ChangeExtension(ReleaseVersionPath, ".ucas"));
+				}
 			}
 
 			if (Params.CreateChunkInstall)
@@ -2549,9 +2551,20 @@ public partial class Project : CommandUtils
 					string ExistingPatchSearchPath = SC.StageTargetPlatform.GetReleasePakFilePath(SC, Params, null);
 					if (Directory.Exists(ExistingPatchSearchPath))
 					{
-						IEnumerable<string> PakFileSet = Directory.EnumerateFiles(ExistingPatchSearchPath, PakName + "-" + SC.FinalCookPlatform + "*" + OutputFilenameExtension);
+						HashSet<string> IncludedExtensions = new HashSet<string>();
+						IncludedExtensions.Add(OutputFilenameExtension);
+						if (ShouldCreateIoStoreContainerFiles(Params, SC))
+						{
+							IncludedExtensions.Add(".ucas");
+							IncludedExtensions.Add(".utoc");
+						}
+						IEnumerable<string> PakFileSet = Directory.EnumerateFiles(ExistingPatchSearchPath, PakName + "-" + SC.FinalCookPlatform + "*.*");
 						foreach (string PakFilePath in PakFileSet)
 						{
+							if (!IncludedExtensions.Contains(Path.GetExtension(PakFilePath).ToLower()))
+							{
+								continue;
+							}
 							FileReference OutputDestinationPath = FileReference.Combine(OutputLocation.Directory, Path.GetFileName(PakFilePath));
 							if (!File.Exists(OutputDestinationPath.FullName))
 							{
@@ -2573,9 +2586,16 @@ public partial class Project : CommandUtils
 
 	private static void RunIoStore(ProjectParams Params, DeploymentContext SC, string CommandsFileName, FileReference GameOpenOrderFileLocation, FileReference CookerOpenOrderFileLocation, string AdditionalArgs)
 	{
-		DirectoryReference OutputLocation = SC.StageTargetPlatform.GetProjectRootForStage(SC.RuntimeRootDir, SC.RelativeProjectRootForStage);
+		StagedFileReference GlobalContainerOutputRelativeLocation;
+		GlobalContainerOutputRelativeLocation = StagedFileReference.Combine(SC.RelativeProjectRootForStage, "Content", "Paks", "global.utoc");
+		if (SC.StageTargetPlatform.DeployLowerCaseFilenames())
+		{
+			GlobalContainerOutputRelativeLocation = GlobalContainerOutputRelativeLocation.ToLowerInvariant();
+		}
+		GlobalContainerOutputRelativeLocation = SC.StageTargetPlatform.Remap(GlobalContainerOutputRelativeLocation);
+		FileReference GlobalContainerOutputLocation = FileReference.Combine(SC.RuntimeRootDir, GlobalContainerOutputRelativeLocation.Name);
 
-		string CommandletParams = String.Format("-OutputDirectory={0} -CookedDirectory={1} -Commands={2}", MakePathSafeToUseWithCommandLine(OutputLocation.FullName), MakePathSafeToUseWithCommandLine(SC.PlatformCookDir.ToString()), MakePathSafeToUseWithCommandLine(CommandsFileName));
+		string CommandletParams = String.Format("-CreateGlobalContainer={0} -CookedDirectory={1} -Commands={2}", MakePathSafeToUseWithCommandLine(GlobalContainerOutputLocation.FullName), MakePathSafeToUseWithCommandLine(SC.PlatformCookDir.ToString()), MakePathSafeToUseWithCommandLine(CommandsFileName));
 		if (GameOpenOrderFileLocation != null)
 		{
 			CommandletParams += String.Format(" -GameOrder={0}", MakePathSafeToUseWithCommandLine(GameOpenOrderFileLocation.FullName));
@@ -3177,6 +3197,8 @@ public partial class Project : CommandUtils
 	{
 		var StagedFilesDir = new DirectoryInfo(StagingDirectory);
 		StagedFilesDir.GetFiles("*.pak", SearchOption.AllDirectories).ToList().ForEach(File => File.Delete());
+		StagedFilesDir.GetFiles("*.ucas", SearchOption.AllDirectories).ToList().ForEach(File => File.Delete());
+		StagedFilesDir.GetFiles("*.utoc", SearchOption.AllDirectories).ToList().ForEach(File => File.Delete());
 	}
 
 	protected static void CleanDirectoryExcludingPakFiles(DirectoryInfo StagingDirectory)
diff --git a/Engine/Source/Runtime/Core/Private/GenericPlatform/GenericPlatformIoDispatcher.cpp b/Engine/Source/Runtime/Core/Private/GenericPlatform/GenericPlatformIoDispatcher.cpp
--- a/Engine/Source/Runtime/Core/Private/GenericPlatform/GenericPlatformIoDispatcher.cpp
+++ b/Engine/Source/Runtime/Core/Private/GenericPlatform/GenericPlatformIoDispatcher.cpp
@@ -103,7 +103,7 @@ bool FGenericFileIoStoreImpl::OpenContainer(const TCHAR* ContainerFilePath, uint
 void FGenericFileIoStoreImpl::ReadBlockFromFile(uint8* Target, uint64 InFileHandle, FFileIoStoreRawBlock* Block)
 {
 	FCachedBlock* CachedBlock = nullptr;
-	bool bIsCacheableBlock = CacheMemory != nullptr && !Block->DirectToRequest && Block->Size == ReadBufferSize;
+	bool bIsCacheableBlock = CacheMemory != nullptr && ((Block->Flags & FFileIoStoreRawBlock::Cacheable) > 0);
 	if (bIsCacheableBlock)
 	{
 		CachedBlock = CachedBlocks.FindRef(Block->Key.Hash);
diff --git a/Engine/Source/Runtime/Core/Private/IO/IoContainerId.cpp b/Engine/Source/Runtime/Core/Private/IO/IoContainerId.cpp
new file mode 100644
--- /dev/null
+++ b/Engine/Source/Runtime/Core/Private/IO/IoContainerId.cpp
@@ -0,0 +1,16 @@
+// Copyright Epic Games, Inc. All Rights Reserved.
+
+#include "IO/IoContainerId.h"
+#include "Serialization/StructuredArchive.h"
+
+FArchive& operator<<(FArchive& Ar, FIoContainerId& ContainerId)
+{
+	Ar << ContainerId.Id;
+
+	return Ar;
+}
+
+void operator<<(FStructuredArchiveSlot Slot, FIoContainerId& Value)
+{
+	Slot << Value.Id;
+}
diff --git a/Engine/Source/Runtime/Core/Private/IO/IoDispatcher.cpp b/Engine/Source/Runtime/Core/Private/IO/IoDispatcher.cpp
--- a/Engine/Source/Runtime/Core/Private/IO/IoDispatcher.cpp
+++ b/Engine/Source/Runtime/Core/Private/IO/IoDispatcher.cpp
@@ -26,6 +26,22 @@ TUniquePtr<FIoDispatcher> GIoDispatcher;
 //PRAGMA_DISABLE_OPTIMIZATION
 #endif
 
+/** A utility function to convert a EIoStoreResolveResult to the corresponding FIoStatus.*/
+FIoStatus ToStatus(EIoStoreResolveResult Result)
+{
+	switch (Result)
+	{
+	case IoStoreResolveResult_OK:
+		return FIoStatus(EIoErrorCode::Ok);
+
+	case IoStoreResolveResult_NotFound:
+		return FIoStatus(EIoErrorCode::NotFound);
+
+	default:
+		return FIoStatus(EIoErrorCode::Unknown);
+	}
+}
+
 template <typename T, uint32 BlockSize = 128>
 class TBlockAllocator
 {
@@ -135,7 +151,7 @@ class FIoDispatcherImpl
 public:
 	FIoDispatcherImpl(bool bInIsMultithreaded)
 		: bIsMultithreaded(bInIsMultithreaded)
-		, FileIoStore(EventQueue, bInIsMultithreaded)
+		, FileIoStore(EventQueue) 
 	{
 		FCoreDelegates::GetMemoryTrimDelegate().AddLambda([this]()
 		{
@@ -244,12 +260,35 @@ public:
 
 	TIoStatusOr<FIoMappedRegion> OpenMapped(const FIoChunkId& ChunkId, const FIoReadOptions& Options)
 	{
-		return FileIoStore.OpenMapped(ChunkId, Options);
+		if (ChunkId.IsValid())
+		{
+			return FileIoStore.OpenMapped(ChunkId, Options);
+		}
+		else
+		{
+			return FIoStatus(EIoErrorCode::InvalidParameter, TEXT("FIoChunkId is not valid"));
+		}
 	}
 
 	FIoStatus Mount(const FIoStoreEnvironment& Environment)
 	{
-		return FileIoStore.Mount(Environment);
+		TIoStatusOr<FIoContainerId> ContainerId = FileIoStore.Mount(Environment);
+
+		if (ContainerId.IsOk())
+		{
+			FIoDispatcherMountedContainer MountedContainer;
+			MountedContainer.ContainerId = ContainerId.ValueOrDie();
+			MountedContainer.Environment = Environment;
+			if (ContainerMountedEvent.IsBound())
+			{
+				ContainerMountedEvent.Broadcast(MountedContainer);
+			}
+			FScopeLock Lock(&MountedContainersCritical);
+			MountedContainers.Add(MoveTemp(MountedContainer));
+			return FIoStatus::Ok;
+		}
+
+		return ContainerId.Status();
 	}
 
 	bool DoesChunkExist(const FIoChunkId& ChunkId) const
@@ -259,7 +298,26 @@ public:
 
 	TIoStatusOr<uint64> GetSizeForChunk(const FIoChunkId& ChunkId) const
 	{
-		return FileIoStore.GetSizeForChunk(ChunkId);
+		// Only attempt to find the size if the FIoChunkId is valid
+		if (ChunkId.IsValid())
+		{
+			return FileIoStore.GetSizeForChunk(ChunkId);
+		}
+		else
+		{
+			return FIoStatus(EIoErrorCode::InvalidParameter, TEXT("FIoChunkId is not valid"));
+		}	
+	}
+
+	TArray<FIoDispatcherMountedContainer> GetMountedContainers() const
+	{
+		FScopeLock Lock(&MountedContainersCritical);
+		return MountedContainers;
+	}
+
+	FIoDispatcher::FIoContainerMountedEvent& OnContainerMounted()
+	{
+		return ContainerMountedEvent;
 	}
 
 	template<typename Func>
@@ -454,11 +512,20 @@ private:
 			{
 				TRACE_CPUPROFILER_EVENT_SCOPE(ResolveRequest);
 
-				EIoStoreResolveResult Result = FileIoStore.Resolve(Request);
-				if (Result == IoStoreResolveResult_NotFound)
+				// Make sure that the FIoChunkId in the request is valid before we try to do anything with it.
+				if (Request->ChunkId.IsValid())
+				{
+					EIoStoreResolveResult Result = FileIoStore.Resolve(Request);
+					if (Result != IoStoreResolveResult_OK)
+					{
+						Request->Status = ToStatus(Result);
+					}
+				}
+				else
 				{
-					Request->Status = FIoStatus(EIoErrorCode::NotFound);
+					Request->Status = FIoStatus(EIoErrorCode::InvalidParameter, TEXT("FIoChunkId is not valid"));
 				}
+
 				if (!SubmittedRequestsTail)
 				{
 					SubmittedRequestsHead = SubmittedRequestsTail = Request;
@@ -471,7 +538,18 @@ private:
 				Request->NextRequest = nullptr;
 			}
 
-			ProcessCompletedBlocks();
+			if (!bIsMultithreaded)
+			{
+				while (FileIoStore.ReadPendingBlock())
+				{
+					FileIoStore.FlushReads();
+					ProcessCompletedBlocks();
+				}
+			}
+			else
+			{
+				ProcessCompletedBlocks();
+			}
 		}
 	}
 
@@ -530,6 +608,9 @@ private:
 	FIoRequestImpl* SubmittedRequestsHead = nullptr;
 	FIoRequestImpl* SubmittedRequestsTail = nullptr;
 	TAtomic<bool> bStopRequested { false };
+	mutable FCriticalSection MountedContainersCritical;
+	TArray<FIoDispatcherMountedContainer> MountedContainers;
+	FIoDispatcher::FIoContainerMountedEvent ContainerMountedEvent;
 };
 
 FIoDispatcher::FIoDispatcher()
@@ -585,6 +666,18 @@ FIoDispatcher::GetSizeForChunk(const FIoChunkId& ChunkId) const
 	return Impl->GetSizeForChunk(ChunkId);
 }
 
+TArray<FIoDispatcherMountedContainer>
+FIoDispatcher::GetMountedContainers() const
+{
+	return Impl->GetMountedContainers();
+}
+
+FIoDispatcher::FIoContainerMountedEvent&
+FIoDispatcher::OnContainerMounted()
+{
+	return Impl->OnContainerMounted();
+}
+
 bool
 FIoDispatcher::IsInitialized()
 {
diff --git a/Engine/Source/Runtime/Core/Private/IO/IoDispatcherFileBackend.cpp b/Engine/Source/Runtime/Core/Private/IO/IoDispatcherFileBackend.cpp
--- a/Engine/Source/Runtime/Core/Private/IO/IoDispatcherFileBackend.cpp
+++ b/Engine/Source/Runtime/Core/Private/IO/IoDispatcherFileBackend.cpp
@@ -24,6 +24,13 @@ static FAutoConsoleVariableRef CVar_IoDispatcherBufferSizeKB(
 	TEXT("IoDispatcher read buffer size (in kilobytes).")
 );
 
+int32 GIoDispatcherBufferAlignment = 4096;
+static FAutoConsoleVariableRef CVar_IoDispatcherBufferAlignment(
+	TEXT("s.IoDispatcherBufferAlignment"),
+	GIoDispatcherBufferAlignment,
+	TEXT("IoDispatcher read buffer alignment.")
+);
+
 int32 GIoDispatcherBufferMemoryMB = 8;
 static FAutoConsoleVariableRef CVar_IoDispatcherBufferMemoryMB(
 	TEXT("s.IoDispatcherBufferMemoryMB"),
@@ -73,6 +80,8 @@ FIoStatus FFileIoStoreReader::Initialize(const FIoStoreEnvironment& Environment)
 	TStringBuilder<256> TocFilePath;
 	TocFilePath.Append(ContainerFilePath);
 
+	UE_LOG(LogIoDispatcher, Display, TEXT("Reading toc: %s"), *TocFilePath);
+
 	ContainerFilePath.Append(TEXT(".ucas"));
 	TocFilePath.Append(TEXT(".utoc"));
 
@@ -83,58 +92,27 @@ FIoStatus FFileIoStoreReader::Initialize(const FIoStoreEnvironment& Environment)
 
 	ContainerFile.FilePath = ContainerFilePath;
 
-	TUniquePtr<uint8[]> TocBuffer;
-	bool bTocReadOk = false;
-
-	{
-		TUniquePtr<IFileHandle>	TocFileHandle(Ipf.OpenRead(*TocFilePath, /* allowwrite */ false));
-
-		if (!TocFileHandle)
-		{
-			return FIoStatusBuilder(EIoErrorCode::FileOpenFailed) << TEXT("Failed to open IoStore TOC file '") << *TocFilePath << TEXT("'");
-		}
-
-		const int64 TocSize = TocFileHandle->Size();
-		TocBuffer = MakeUnique<uint8[]>(TocSize);
-		bTocReadOk = TocFileHandle->Read(TocBuffer.Get(), TocSize);
-	}
-
-	if (!bTocReadOk)
-	{
-		return FIoStatusBuilder(EIoErrorCode::CorruptToc) << TEXT("Failed to read IoStore TOC file '") << *TocFilePath << TEXT("'");
-	}
-
-	const FIoStoreTocHeader* Header = reinterpret_cast<const FIoStoreTocHeader*>(TocBuffer.Get());
-
-	if (!Header->CheckMagic())
-	{
-		return FIoStatusBuilder(EIoErrorCode::CorruptToc) << TEXT("TOC header magic mismatch while reading '") << *TocFilePath << TEXT("'");
-	}
-
-	if (Header->TocHeaderSize != sizeof(FIoStoreTocHeader))
-	{
-		return FIoStatusBuilder(EIoErrorCode::CorruptToc) << TEXT("TOC header size mismatch while reading '") << *TocFilePath << TEXT("'");
-	}
-
-	if (Header->TocEntrySize != sizeof(FIoStoreTocEntry))
+	FIoStoreTocReader TocReader;
+	FIoStatus Status = TocReader.Initialize(*TocFilePath);
+	if (!Status.IsOk())
 	{
-		return FIoStatusBuilder(EIoErrorCode::CorruptToc) << TEXT("TOC entry size mismatch while reading '") << *TocFilePath << TEXT("'");
+		return Status;
 	}
 
-	const FIoStoreTocEntry* TocEntry = reinterpret_cast<const FIoStoreTocEntry*>(TocBuffer.Get() + sizeof(FIoStoreTocHeader));
-	uint32 TocEntryCount = Header->TocEntryCount;
+	uint32 TocEntryCount;
+	const FIoStoreTocEntry* TocEntry = TocReader.GetEntries(TocEntryCount);
 
-	uint32 CompressedBlockEntryCount = Header->CompressionBlockCount;
-	const FIoStoreCompressedBlockEntry* CompressedBlockEntry = CompressedBlockEntryCount > 0
-		? reinterpret_cast<const FIoStoreCompressedBlockEntry*>(TocEntry + TocEntryCount)
-		: nullptr;
+	uint32 TocPartialEntryCount;
+	const FIoStoreTocPartialEntry* TocPartialEntry = TocReader.GetPartialEntries(TocPartialEntryCount);
 
-	const ANSICHAR* CompressionMethodNames = reinterpret_cast<const ANSICHAR*>(CompressedBlockEntry + CompressedBlockEntryCount);
+	uint32 CompressedBlockEntryCount;
+	const FIoStoreTocCompressedBlockEntry* CompressedBlockEntry = TocReader.GetCompressedBlockEntries(CompressedBlockEntryCount);
 
-	Toc.Reserve(TocEntryCount);
+	uint64 ContainerUncompressedSize = CompressedBlockEntryCount > 0 ? uint64(CompressedBlockEntryCount) * uint64(TocReader.GetCompressionBlockSize()) : ContainerFile.FileSize;
+	Toc.Reserve(TocEntryCount + TocPartialEntryCount);
 	while (TocEntryCount--)
 	{
-		if (!CompressedBlockEntryCount && (TocEntry->GetOffset() + TocEntry->GetLength()) > ContainerFile.FileSize)
+		if (TocEntry->GetOffset() + TocEntry->GetLength() > ContainerUncompressedSize)
 		{
 			return FIoStatusBuilder(EIoErrorCode::CorruptToc) << TEXT("TOC TocEntry out of container bounds while reading '") << *TocFilePath << TEXT("'");
 		}
@@ -142,35 +120,41 @@ FIoStatus FFileIoStoreReader::Initialize(const FIoStoreEnvironment& Environment)
 		Toc.Add(TocEntry->ChunkId, TocEntry->OffsetAndLength);
 		++TocEntry;
 	}
-
-	if (CompressedBlockEntryCount > 0)
+	while (TocPartialEntryCount--)
 	{
-		UE_LOG(LogIoDispatcher, Display, TEXT("Loading compressed toc: %s"), *TocFilePath);
-		ContainerFile.CompressionBlockSize = Header->CompressionBlockSize;
-		ContainerFile.CompressionBlocks.Reserve(Header->CompressionBlockCount);
-
-		while (CompressedBlockEntryCount--)
+		if (TocPartialEntry->GetOffset() + TocPartialEntry->GetLength() > ContainerUncompressedSize)
 		{
-			if (CompressedBlockEntry->OffsetAndLength.GetOffset() + CompressedBlockEntry->OffsetAndLength.GetLength() > ContainerFile.FileSize)
-			{
-				return FIoStatusBuilder(EIoErrorCode::CorruptToc) << TEXT("TOC TocCompressedBlockEntry out of container bounds while reading '") << *TocFilePath << TEXT("'");
-			}
-
-			ContainerFile.CompressionBlocks.Emplace(*CompressedBlockEntry);
-			++CompressedBlockEntry;
+			return FIoStatusBuilder(EIoErrorCode::CorruptToc) << TEXT("TOC TocPartialEntry out of container bounds while reading '") << *TocFilePath << TEXT("'");
 		}
+
+		Toc.Add(TocPartialEntry->ChunkId, TocPartialEntry->OffsetAndLength);
+		++TocPartialEntry;
 	}
-	else
+
+	uint32 CompressionMethodNameCount;
+	const FName* CompressionMethodName = TocReader.GetCompressionMethodNames(CompressionMethodNameCount);
+	while (CompressionMethodNameCount--)
 	{
-		UE_LOG(LogIoDispatcher, Display, TEXT("Loading uncompressed toc: %s"), *TocFilePath);
+		ContainerFile.CompressionMethods.Add(*CompressionMethodName);
+		++CompressionMethodName;
 	}
 
-	for (uint32 CompressonNameIndex = 0; CompressonNameIndex < Header->CompressionNameCount; CompressonNameIndex++)
+	ContainerFile.CompressionBlockSize = TocReader.GetCompressionBlockSize();
+	ContainerFile.CompressionBlocks.Reserve(CompressedBlockEntryCount);
+
+	while (CompressedBlockEntryCount--)
 	{
-		const ANSICHAR* CompressionMethodName = CompressionMethodNames + CompressonNameIndex * FIoStoreCompressionInfo::CompressionMethodNameLen;
-		ContainerFile.CompressionMethods.Add(FName(CompressionMethodName));
+		if (CompressedBlockEntry->OffsetAndLength.GetOffset() + CompressedBlockEntry->OffsetAndLength.GetLength() > ContainerFile.FileSize)
+		{
+			return (FIoStatus)(FIoStatusBuilder(EIoErrorCode::CorruptToc) << TEXT("TOC TocCompressedBlockEntry out of container bounds while reading '") << *TocFilePath << TEXT("'"));
+		}
+
+		ContainerFile.CompressionBlocks.Emplace(*CompressedBlockEntry);
+		++CompressedBlockEntry;
 	}
 
+	ContainerId = TocReader.GetContainerId();
+	Order = Environment.GetOrder();
 	return FIoStatus::Ok;
 }
 
@@ -210,17 +194,16 @@ IMappedFileHandle* FFileIoStoreReader::GetMappedContainerFileHandle()
 	return new FMappedFileProxy(ContainerFile.MappedFileHandle.Get(), ContainerFile.FileSize);
 }
 
-FFileIoStore::FFileIoStore(FIoDispatcherEventQueue& InEventQueue, bool bInIsMultithreaded)
+FFileIoStore::FFileIoStore(FIoDispatcherEventQueue& InEventQueue)
 	: ReadBufferSize(GIoDispatcherBufferSizeKB > 0 ? uint64(GIoDispatcherBufferSizeKB) << 10 : 256 << 10)
 	, EventQueue(InEventQueue)
-	, bIsMultithreaded(bInIsMultithreaded)
 	, PlatformImpl(InEventQueue, ReadBufferSize)
 	, BufferAvailableEvent(FPlatformProcess::GetSynchEventFromPool())
 	, PendingBlockEvent(FPlatformProcess::GetSynchEventFromPool())
 {
 	uint64 BufferCount = uint64(GIoDispatcherBufferMemoryMB > 0 ? uint64(GIoDispatcherBufferMemoryMB) << 20 : 32ull << 20) / ReadBufferSize;
 	uint64 MemorySize = BufferCount * ReadBufferSize;
-	BufferMemory = reinterpret_cast<uint8*>(FMemory::Malloc(MemorySize));
+	BufferMemory = reinterpret_cast<uint8*>(FMemory::Malloc(MemorySize, GIoDispatcherBufferAlignment));
 	for (uint64 BufferIndex = 0; BufferIndex < BufferCount; ++BufferIndex)
 	{
 		FFileIoStoreBuffer* Buffer = new FFileIoStoreBuffer();
@@ -246,16 +229,29 @@ FFileIoStore::~FFileIoStore()
 	FPlatformProcess::ReturnSynchEventToPool(PendingBlockEvent);
 }
 
-FIoStatus FFileIoStore::Mount(const FIoStoreEnvironment& Environment)
+TIoStatusOr<FIoContainerId> FFileIoStore::Mount(const FIoStoreEnvironment& Environment)
 {
 	TUniquePtr<FFileIoStoreReader> Reader(new FFileIoStoreReader(PlatformImpl));
 	FIoStatus IoStatus = Reader->Initialize(Environment);
-	if (IoStatus.IsOk())
+	if (!IoStatus.IsOk())
+	{
+		return IoStatus;
+	}
+	int32 InsertionIndex;
+	FIoContainerId ContainerId = Reader->GetContainerId();
 	{
 		FWriteScopeLock _(IoStoreReadersLock);
-		IoStoreReaders.Add(Reader.Release());
+		InsertionIndex = Algo::UpperBound(IoStoreReaders, Reader.Get(), [](const FFileIoStoreReader* A, const FFileIoStoreReader* B)
+		{
+			if (A->GetOrder() != B->GetOrder())
+			{
+				return A->GetOrder() > B->GetOrder();
+			}
+			return A->GetContainerId() < B->GetContainerId();
+		});
+		IoStoreReaders.Insert(Reader.Release(), InsertionIndex);
 	}
-	return IoStatus;
+	return ContainerId;
 }
 
 EIoStoreResolveResult FFileIoStore::Resolve(FIoRequestImpl* Request)
@@ -291,49 +287,7 @@ EIoStoreResolveResult FFileIoStore::Resolve(FIoRequestImpl* Request)
 					ResolvedRequest.Request->IoBuffer.SetSize(ResolvedRequest.ResolvedSize);
 				}
 
-				const FFileIoStoreContainerFile& ContainerFile = Reader->GetContainerFile();
-				const bool bIsCompressed = ContainerFile.CompressionBlockSize > 0;
-				const uint64 BlockSize = bIsCompressed ? ContainerFile.CompressionBlockSize : ReadBufferSize;
-				const uint64 RequestEndOffset = ResolvedRequest.ResolvedOffset + ResolvedRequest.ResolvedSize;
-				int32 RequestBeginBlockIndex = int32(ResolvedRequest.ResolvedOffset / BlockSize);
-				int32 RequestEndBlockIndex = int32((RequestEndOffset - 1) / BlockSize + 1);
-
-				if (bIsCompressed)
-				{
-					for (int32 BlockIndex = RequestBeginBlockIndex; BlockIndex < RequestEndBlockIndex; ++BlockIndex)
-					{
-						ReadBlockAndScatter(ReaderIndex, BlockIndex, ResolvedRequest);
-					}
-				}
-				else
-				{
-					int32 NumBlocks = RequestEndBlockIndex - RequestBeginBlockIndex;
-
-					const bool bFirstBlockIsPartial = ResolvedRequest.ResolvedOffset % BlockSize != 0;
-					if (bFirstBlockIsPartial)
-					{
-						ReadBlockAndScatter(ReaderIndex, RequestBeginBlockIndex, ResolvedRequest);
-						++RequestBeginBlockIndex;
-						--NumBlocks;
-					}
-
-					if (NumBlocks > 0)
-					{
-						const bool bLastBlockIsPartial = RequestEndOffset % BlockSize != 0;
-						const int32 NumMergedBlocks = bLastBlockIsPartial ? NumBlocks - 1  : NumBlocks;
-						if (NumMergedBlocks > 0)
-						{
-							const uint64 MergedReadOffset = RequestBeginBlockIndex * BlockSize;
-							const uint64 MergedReadSize = NumMergedBlocks * BlockSize;
-							ReadNoScatter(ReaderIndex, MergedReadOffset, MergedReadSize, ResolvedRequest);
-						}
-
-						if (bLastBlockIsPartial)
-						{
-							ReadBlockAndScatter(ReaderIndex, RequestEndBlockIndex - 1, ResolvedRequest);
-						}
-					}
-				}
+				ReadBlocks(ReaderIndex, ResolvedRequest);
 			}
 
 			return IoStoreResolveResult_OK;
@@ -621,6 +575,7 @@ TIoStatusOr<FIoMappedRegion> FFileIoStore::OpenMapped(const FIoChunkId& ChunkId,
 
 	IPlatformFile& Ipf = FPlatformFileManager::Get().GetPlatformFile();
 
+	FReadScopeLock _(IoStoreReadersLock);
 	for (FFileIoStoreReader* Reader : IoStoreReaders)
 	{
 		if (const FIoOffsetAndLength* OffsetAndLength = Reader->Resolve(ChunkId))
@@ -629,131 +584,125 @@ TIoStatusOr<FIoMappedRegion> FFileIoStore::OpenMapped(const FIoChunkId& ChunkId,
 			uint64 ResolvedSize = FMath::Min(Options.GetSize(), OffsetAndLength->GetLength());
 			
 			const FFileIoStoreContainerFile& ContainerFile = Reader->GetContainerFile();
-			const bool bIsCompressed = ContainerFile.CompressionBlockSize > 0;
-
+			
 			IMappedFileHandle* MappedFileHandle = Reader->GetMappedContainerFileHandle();
 			IMappedFileRegion* MappedFileRegion = nullptr;
 
-			if (bIsCompressed)
-			{
-				int32 BlockIndex = int32(ResolvedOffset / ContainerFile.CompressionBlockSize);
-				const FIoStoreCompressedBlockEntry& CompressionBlockEntry = ContainerFile.CompressionBlocks[BlockIndex];
-				const int64 BlockOffset = (int64)CompressionBlockEntry.OffsetAndLength.GetOffset();
-				check(BlockOffset > 0 && IsAligned(BlockOffset, FPlatformProperties::GetMemoryMappingAlignment()));
-
-				MappedFileRegion = MappedFileHandle->MapRegion(BlockOffset + Options.GetOffset(), ResolvedSize);
-				check(IsAligned(MappedFileRegion->GetMappedPtr(), FPlatformProperties::GetMemoryMappingAlignment()));
-			}
-			else
-			{
-				MappedFileRegion = MappedFileHandle->MapRegion(ResolvedOffset + Options.GetOffset(), ResolvedSize);
-			}
+			int32 BlockIndex = int32(ResolvedOffset / ContainerFile.CompressionBlockSize);
+			const FIoStoreTocCompressedBlockEntry& CompressionBlockEntry = ContainerFile.CompressionBlocks[BlockIndex];
+			const int64 BlockOffset = (int64)CompressionBlockEntry.OffsetAndLength.GetOffset();
+			check(BlockOffset > 0 && IsAligned(BlockOffset, FPlatformProperties::GetMemoryMappingAlignment()));
 
+			MappedFileRegion = MappedFileHandle->MapRegion(BlockOffset + Options.GetOffset(), ResolvedSize);
+			check(IsAligned(MappedFileRegion->GetMappedPtr(), FPlatformProperties::GetMemoryMappingAlignment()));
+			
 			return FIoMappedRegion { MappedFileHandle, MappedFileRegion };
 		}
 	}
 
-	return FIoStatus::Invalid;
+	// We didn't find any entry for the ChunkId.
+	return FIoStatus(EIoErrorCode::NotFound);
 }
 
-void FFileIoStore::ReadBlockAndScatter(uint32 ReaderIndex, uint32 BlockIndex, const FFileIoStoreResolvedRequest& ResolvedRequest)
+void FFileIoStore::ReadBlocks(uint32 ReaderIndex, const FFileIoStoreResolvedRequest& ResolvedRequest)
 {
 	/*TStringBuilder<256> ScopeName;
 	ScopeName.Appendf(TEXT("ReadBlock %d"), BlockIndex);
 	TRACE_CPUPROFILER_EVENT_SCOPE_TEXT(*ScopeName);*/
 
 	const FFileIoStoreContainerFile& ContainerFile = IoStoreReaders[ReaderIndex]->GetContainerFile();
+	const uint64 CompressionBlockSize = ContainerFile.CompressionBlockSize;
+	const uint64 RequestEndOffset = ResolvedRequest.ResolvedOffset + ResolvedRequest.ResolvedSize;
+	int32 RequestBeginBlockIndex = int32(ResolvedRequest.ResolvedOffset / CompressionBlockSize);
+	int32 RequestEndBlockIndex = int32((RequestEndOffset - 1) / CompressionBlockSize);
 
 	FFileIoStoreRawBlock* NewBlocksHead = nullptr;
 	FFileIoStoreRawBlock* NewBlocksTail = nullptr;
 
-	FFileIoStoreBlockKey CompressedBlockKey;
-	CompressedBlockKey.FileIndex = ReaderIndex;
-	CompressedBlockKey.BlockIndex = BlockIndex;
-	FFileIoStoreCompressedBlock* CompressedBlock = CompressedBlocksMap.FindRef(CompressedBlockKey);
-	if (!CompressedBlock)
+	uint64 RequestStartOffsetInBlock = ResolvedRequest.ResolvedOffset - RequestBeginBlockIndex * CompressionBlockSize;
+	uint64 RequestRemainingBytes = ResolvedRequest.ResolvedSize;
+	uint64 OffsetInRequest = 0;
+	for (int32 CompressedBlockIndex = RequestBeginBlockIndex; CompressedBlockIndex <= RequestEndBlockIndex; ++CompressedBlockIndex)
 	{
-		CompressedBlock = new FFileIoStoreCompressedBlock();
-		CompressedBlock->Key = CompressedBlockKey;
-		CompressedBlocksMap.Add(CompressedBlockKey, CompressedBlock);
-
-		uint64 RawOffset;
-		uint64 RawSize;
-		if (ContainerFile.CompressionBlockSize > 0)
-		{
-			const FIoStoreCompressedBlockEntry& CompressionBlockEntry = ContainerFile.CompressionBlocks[BlockIndex];
-			CompressedBlock->Size = ContainerFile.CompressionBlockSize;
-			CompressedBlock->CompressionMethod = 
-				(CompressionBlockEntry.CompressionMethodIndex == FIoStoreCompressionInfo::InvalidCompressionIndex) ?
-				NAME_None :
-				ContainerFile.CompressionMethods[CompressionBlockEntry.CompressionMethodIndex];
-			RawOffset = CompressionBlockEntry.OffsetAndLength.GetOffset();
-			RawSize = CompressionBlockEntry.OffsetAndLength.GetLength();
-		}
-		else
+		FFileIoStoreBlockKey CompressedBlockKey;
+		CompressedBlockKey.FileIndex = ReaderIndex;
+		CompressedBlockKey.BlockIndex = CompressedBlockIndex;
+		FFileIoStoreCompressedBlock* CompressedBlock = CompressedBlocksMap.FindRef(CompressedBlockKey);
+		if (!CompressedBlock)
 		{
-			CompressedBlock->Size = ReadBufferSize;
-			RawOffset = BlockIndex * CompressedBlock->Size;
-			RawSize = ReadBufferSize;
-		}
-		CompressedBlock->RawOffset = RawOffset;
-		CompressedBlock->RawSize = RawSize;
-		const uint32 RawBeginBlockIndex = uint32(RawOffset / ReadBufferSize);
-		const uint32 RawEndBlockIndex = uint32((RawOffset + RawSize - 1) / ReadBufferSize + 1);
-		const uint32 RawBlockCount = RawEndBlockIndex - RawBeginBlockIndex;
-		CompressedBlock->RawBlocksCount = RawBlockCount;
-		check(RawBlockCount > 0);
-		for (uint32 RawBlockIndex = RawBeginBlockIndex; RawBlockIndex < RawEndBlockIndex; ++RawBlockIndex)
-		{
-			FFileIoStoreBlockKey RawBlockKey;
-			RawBlockKey.BlockIndex = RawBlockIndex;
-			RawBlockKey.FileIndex = ReaderIndex;
+			CompressedBlock = new FFileIoStoreCompressedBlock();
+			CompressedBlock->Key = CompressedBlockKey;
+			CompressedBlocksMap.Add(CompressedBlockKey, CompressedBlock);
+
+			bool bCacheable = OffsetInRequest > 0 || RequestRemainingBytes < CompressionBlockSize;
 
-			FFileIoStoreRawBlock* RawBlock = RawBlocksMap.FindRef(RawBlockKey);
-			if (!RawBlock)
+			const FIoStoreTocCompressedBlockEntry& CompressionBlockEntry = ContainerFile.CompressionBlocks[CompressedBlockIndex];
+			CompressedBlock->Size = ContainerFile.CompressionBlockSize;
+			CompressedBlock->CompressionMethod = ContainerFile.CompressionMethods[CompressionBlockEntry.CompressionMethodIndex];
+			uint64 RawOffset = CompressionBlockEntry.OffsetAndLength.GetOffset();
+			uint64 RawSize = CompressionBlockEntry.OffsetAndLength.GetLength();
+			CompressedBlock->RawOffset = RawOffset;
+			CompressedBlock->RawSize = RawSize;
+			const uint32 RawBeginBlockIndex = uint32(RawOffset / ReadBufferSize);
+			const uint32 RawEndBlockIndex = uint32((RawOffset + RawSize - 1) / ReadBufferSize);
+			const uint32 RawBlockCount = RawEndBlockIndex - RawBeginBlockIndex + 1;
+			CompressedBlock->RawBlocksCount = RawBlockCount;
+			check(RawBlockCount > 0);
+			for (uint32 RawBlockIndex = RawBeginBlockIndex; RawBlockIndex <= RawEndBlockIndex; ++RawBlockIndex)
 			{
-				RawBlock = new FFileIoStoreRawBlock();
-				RawBlock->Key = RawBlockKey;
-				RawBlocksMap.Add(RawBlockKey, RawBlock);
-
-				RawBlock->Offset = RawBlockIndex * ReadBufferSize;
-				uint64 ReadSize = FMath::Min(ContainerFile.FileSize, RawBlock->Offset + ReadBufferSize) - RawBlock->Offset;
-				RawBlock->Size = ReadSize;
-				if (!NewBlocksTail)
+				FFileIoStoreBlockKey RawBlockKey;
+				RawBlockKey.BlockIndex = RawBlockIndex;
+				RawBlockKey.FileIndex = ReaderIndex;
+
+				FFileIoStoreRawBlock* RawBlock = RawBlocksMap.FindRef(RawBlockKey);
+				if (!RawBlock)
 				{
-					NewBlocksHead = NewBlocksTail = RawBlock;
+					RawBlock = new FFileIoStoreRawBlock();
+					RawBlock->Key = RawBlockKey;
+					RawBlocksMap.Add(RawBlockKey, RawBlock);
+					if (bCacheable)
+					{
+						RawBlock->Flags |= FFileIoStoreRawBlock::Cacheable;
+					}
+
+					RawBlock->Offset = RawBlockIndex * ReadBufferSize;
+					uint64 ReadSize = FMath::Min(ContainerFile.FileSize, RawBlock->Offset + ReadBufferSize) - RawBlock->Offset;
+					RawBlock->Size = ReadSize;
+					if (!NewBlocksTail)
+					{
+						NewBlocksHead = NewBlocksTail = RawBlock;
+					}
+					else
+					{
+						NewBlocksTail->Next = RawBlock;
+						NewBlocksTail = RawBlock;
+					}
 				}
-				else
+				if (RawBlockCount == 1)
 				{
-					NewBlocksTail->Next = RawBlock;
-					NewBlocksTail = RawBlock;
+					CompressedBlock->SingleRawBlock = RawBlock;
 				}
+				RawBlock->CompressedBlocks.Add(CompressedBlock);
+				++RawBlock->RefCount;
+				++CompressedBlock->UnfinishedRawBlocksCount;
 			}
-			if (RawBlockCount == 1)
-			{
-				CompressedBlock->SingleRawBlock = RawBlock;
-			}
-			RawBlock->CompressedBlocks.Add(CompressedBlock);
-			++RawBlock->RefCount;
-			++CompressedBlock->UnfinishedRawBlocksCount;
 		}
-	}
+		check(CompressedBlock->Size > RequestStartOffsetInBlock);
+		uint64 RequestSizeInBlock = FMath::Min<uint64>(CompressedBlock->Size - RequestStartOffsetInBlock, RequestRemainingBytes);
+		check(OffsetInRequest + RequestSizeInBlock <= ResolvedRequest.Request->IoBuffer.DataSize());
+		check(RequestStartOffsetInBlock + RequestSizeInBlock <= CompressedBlock->Size);
 
-	uint64 BlockOffset = BlockIndex * CompressedBlock->Size;
-	uint64 RequestStartOffsetInBlock = FMath::Max<int64>(0, int64(ResolvedRequest.ResolvedOffset) - BlockOffset);
-	uint64 RequestEndOffsetInBlock = FMath::Min<uint64>(CompressedBlock->Size, ResolvedRequest.ResolvedOffset + ResolvedRequest.ResolvedSize - BlockOffset);
-	uint64 BlockOffsetInRequest = FMath::Max<int64>(0, int64(BlockOffset) - ResolvedRequest.ResolvedOffset);
-	uint64 RequestSizeInBlock = RequestEndOffsetInBlock - RequestStartOffsetInBlock;
-	check(RequestSizeInBlock <= ResolvedRequest.Request->IoBuffer.DataSize());
-	check(RequestStartOffsetInBlock + RequestSizeInBlock <= CompressedBlock->Size);
-	check(BlockOffsetInRequest + RequestSizeInBlock <= ResolvedRequest.Request->IoBuffer.DataSize());
+		++ResolvedRequest.Request->UnfinishedReadsCount;
+		FFileIoStoreBlockScatter& Scatter = CompressedBlock->ScatterList.AddDefaulted_GetRef();
+		Scatter.Request = ResolvedRequest.Request;
+		Scatter.DstOffset = OffsetInRequest;
+		Scatter.SrcOffset = RequestStartOffsetInBlock;
+		Scatter.Size = RequestSizeInBlock;
 
-	++ResolvedRequest.Request->UnfinishedReadsCount;
-	FFileIoStoreBlockScatter& Scatter = CompressedBlock->ScatterList.AddDefaulted_GetRef();
-	Scatter.Request = ResolvedRequest.Request;
-	Scatter.DstOffset = BlockOffsetInRequest;
-	Scatter.SrcOffset = RequestStartOffsetInBlock;
-	Scatter.Size = RequestSizeInBlock;
+		RequestRemainingBytes -= RequestSizeInBlock;
+		OffsetInRequest += RequestSizeInBlock;
+		RequestStartOffsetInBlock = 0;
+	}
 
 	if (NewBlocksHead)
 	{
@@ -770,53 +719,6 @@ void FFileIoStore::ReadBlockAndScatter(uint32 ReaderIndex, uint32 BlockIndex, co
 			PendingBlocksTail = NewBlocksTail;
 		}
 		PendingBlockEvent->Trigger();
-		if (!bIsMultithreaded)
-		{
-			ReadPendingBlocks();
-		}
-	}
-}
-
-void FFileIoStore::ReadNoScatter(uint32 ReaderIndex, uint64 Offset, uint64 Size, const FFileIoStoreResolvedRequest& ResolvedRequest)
-{
-	/*TStringBuilder<256> ScopeName;
-	ScopeName.Appendf(TEXT("ReadNoScatter %lld %lld"), Offset, Size);
-	TRACE_CPUPROFILER_EVENT_SCOPE_TEXT(*ScopeName);*/
-
-	const FFileIoStoreContainerFile& ContainerFile = IoStoreReaders[ReaderIndex]->GetContainerFile();
-
-	FFileIoStoreRawBlock* NewBlocksHead = nullptr;
-	FFileIoStoreRawBlock* NewBlocksTail = nullptr;
-
-	FFileIoStoreRawBlock* RawBlock = new FFileIoStoreRawBlock();
-	RawBlock->Key.BlockIndex = -1;
-	RawBlock->Key.FileIndex = ReaderIndex;
-	RawBlock->Offset = Offset;
-	uint64 ReadSize = FMath::Min(ContainerFile.FileSize, RawBlock->Offset + Size) - RawBlock->Offset;
-	RawBlock->Size = ReadSize;
-	RawBlock->DirectToRequest = ResolvedRequest.Request;
-	check(Offset >= ResolvedRequest.ResolvedOffset);
-	RawBlock->DirectToRequestOffset = Offset - ResolvedRequest.ResolvedOffset;
-	++ResolvedRequest.Request->UnfinishedReadsCount;
-
-	AllocMemoryForRequest(ResolvedRequest.Request);
-
-	{
-		FScopeLock Lock(&PendingBlocksCritical);
-		if (!PendingBlocksTail)
-		{
-			PendingBlocksHead = RawBlock;
-		}
-		else
-		{
-			PendingBlocksTail->Next = RawBlock;
-		}
-		PendingBlocksTail = RawBlock;
-	}
-	PendingBlockEvent->Trigger();
-	if (!bIsMultithreaded)
-	{
-		ReadPendingBlocks();
 	}
 }
 
@@ -867,66 +769,53 @@ void FFileIoStore::FreeCompressionContext(FFileIoStoreCompressionContext* Compre
 	FirstFreeCompressionContext = CompressionContext;
 }
 
-void FFileIoStore::ReadPendingBlocks()
+bool FFileIoStore::ReadPendingBlock()
 {
-	while (!bStopRequested)
 	{
+		FScopeLock PendingBlocksLock(&PendingBlocksCritical);
+		if (PendingBlocksHead)
 		{
-			FScopeLock PendingBlocksLock(&PendingBlocksCritical);
-			if (PendingBlocksHead)
+			if (!ScheduledBlocksTail)
 			{
-				if (!ScheduledBlocksTail)
-				{
-					ScheduledBlocksHead = PendingBlocksHead;
-					ScheduledBlocksTail = PendingBlocksTail;
-				}
-				else
-				{
-					ScheduledBlocksTail->Next = PendingBlocksHead;
-					ScheduledBlocksTail = PendingBlocksTail;
-				}
-				PendingBlocksHead = PendingBlocksTail = nullptr;
+				ScheduledBlocksHead = PendingBlocksHead;
+				ScheduledBlocksTail = PendingBlocksTail;
 			}
+			else
+			{
+				ScheduledBlocksTail->Next = PendingBlocksHead;
+				ScheduledBlocksTail = PendingBlocksTail;
+			}
+			PendingBlocksHead = PendingBlocksTail = nullptr;
 		}
+	}
 
-		if (!ScheduledBlocksHead)
-		{
-			break;
-		}
+	if (!ScheduledBlocksHead)
+	{
+		return false;
+	}
 
-		FFileIoStoreRawBlock* BlockToRead = ScheduledBlocksHead;
+	FFileIoStoreRawBlock* BlockToRead = ScheduledBlocksHead;
 
-		/*TStringBuilder<256> ScopeName;
-		if (BlockToRead->DirectToRequest)
-		{
-			ScopeName.Appendf(TEXT("ReadNoScatter %lld %lld"), BlockToRead->Offset, BlockToRead->Size);
-		}
-		else
-		{
-			ScopeName.Appendf(TEXT("ReadBlock %d"), BlockToRead->Key.BlockIndex);
-		}
-		TRACE_CPUPROFILER_EVENT_SCOPE_TEXT(*ScopeName);*/
+	uint8* Target;
+	if (BlockToRead->DirectToRequest)
+	{
+		Target = BlockToRead->DirectToRequest->IoBuffer.Data() + BlockToRead->DirectToRequestOffset;
+	}
+	else
+	{
+		BlockToRead->Buffer = AllocBuffer();
+		Target = BlockToRead->Buffer->Memory;
+	}
 
-		uint8* Target;
-		if (BlockToRead->DirectToRequest)
-		{
-			Target = BlockToRead->DirectToRequest->IoBuffer.Data() + BlockToRead->DirectToRequestOffset;
-		}
-		else
-		{
-			BlockToRead->Buffer = AllocBuffer();
-			Target = BlockToRead->Buffer->Memory;
-		}
-		
-		ScheduledBlocksHead = ScheduledBlocksHead->Next;
-		if (!ScheduledBlocksHead)
-		{
-			ScheduledBlocksTail = nullptr;
-		}
-		
-		PlatformImpl.ReadBlockFromFile(Target, IoStoreReaders[BlockToRead->Key.FileIndex]->GetContainerFile().FileHandle, BlockToRead);
+	ScheduledBlocksHead = ScheduledBlocksHead->Next;
+	if (!ScheduledBlocksHead)
+	{
+		ScheduledBlocksTail = nullptr;
 	}
-	PlatformImpl.FlushReads();
+	
+	PlatformImpl.ReadBlockFromFile(Target, IoStoreReaders[BlockToRead->Key.FileIndex]->GetContainerFile().FileHandle, BlockToRead);
+
+	return true;
 }
 
 bool FFileIoStore::Init()
@@ -945,8 +834,12 @@ uint32 FFileIoStore::Run()
 {
 	while (!bStopRequested)
 	{
-		ReadPendingBlocks();
-		PendingBlockEvent->Wait();
+		const bool bDidReadBlock = ReadPendingBlock();
+		if (!bDidReadBlock)
+		{
+			PlatformImpl.FlushReads();
+			PendingBlockEvent->Wait();
+		}
 	}
 	return 0;
 }
diff --git a/Engine/Source/Runtime/Core/Private/IO/IoDispatcherFileBackend.h b/Engine/Source/Runtime/Core/Private/IO/IoDispatcherFileBackend.h
--- a/Engine/Source/Runtime/Core/Private/IO/IoDispatcherFileBackend.h
+++ b/Engine/Source/Runtime/Core/Private/IO/IoDispatcherFileBackend.h
@@ -18,7 +18,7 @@ struct FFileIoStoreContainerFile
 	uint64 FileSize = 0;
 	uint64 CompressionBlockSize = 0;
 	TArray<FName> CompressionMethods;
-	TArray<FIoStoreCompressedBlockEntry> CompressionBlocks;
+	TArray<FIoStoreTocCompressedBlockEntry> CompressionBlocks;
 	FString FilePath;
 	TUniquePtr<IMappedFileHandle> MappedFileHandle;
 };
@@ -92,6 +92,12 @@ struct FFileIoStoreCompressedBlock
 
 struct FFileIoStoreRawBlock
 {
+	enum EFlags
+	{
+		None = 0,
+		Cacheable = 1
+	};
+
 	FFileIoStoreRawBlock* Next = nullptr;
 	FFileIoStoreBlockKey Key;
 	uint64 Offset;
@@ -101,6 +107,7 @@ struct FFileIoStoreRawBlock
 	uint32 RefCount = 0;
 	FIoRequestImpl* DirectToRequest = nullptr;
 	uint64 DirectToRequestOffset = 0;
+	uint8 Flags = None;
 };
 
 struct FFileIoStoreResolvedRequest
@@ -120,26 +127,35 @@ public:
 	const FIoOffsetAndLength* Resolve(const FIoChunkId& ChunkId) const;
 	const FFileIoStoreContainerFile& GetContainerFile() { return ContainerFile; }
 	IMappedFileHandle* GetMappedContainerFileHandle();
+	const FIoContainerId& GetContainerId() const { return ContainerId; }
+	int32 GetOrder() const { return Order; }
 
 private:
 	FFileIoStoreImpl& PlatformImpl;
 
 	TMap<FIoChunkId, FIoOffsetAndLength> Toc;
 	FFileIoStoreContainerFile ContainerFile;
+	FIoContainerId ContainerId;
+	int32 Order;
 };
 
 class FFileIoStore
 	: public FRunnable
 {
 public:
-	FFileIoStore(FIoDispatcherEventQueue& InEventQueue, bool bInIsMultithreaded);
+	FFileIoStore(FIoDispatcherEventQueue& InEventQueue);
 	~FFileIoStore();
-	FIoStatus Mount(const FIoStoreEnvironment& Environment);
+	TIoStatusOr<FIoContainerId> Mount(const FIoStoreEnvironment& Environment);
 	EIoStoreResolveResult Resolve(FIoRequestImpl* Request);
 	bool DoesChunkExist(const FIoChunkId& ChunkId) const;
 	TIoStatusOr<uint64> GetSizeForChunk(const FIoChunkId& ChunkId) const;
+	bool ReadPendingBlock();
 	void ProcessCompletedBlocks();
 	TIoStatusOr<FIoMappedRegion> OpenMapped(const FIoChunkId& ChunkId, const FIoReadOptions& Options);
+	void FlushReads()
+	{ 
+		PlatformImpl.FlushReads();
+	}
 
 	static bool IsValidEnvironment(const FIoStoreEnvironment& Environment);
 
@@ -181,8 +197,7 @@ private:
 	};
 
 	void InitCache();
-	void ReadBlockAndScatter(uint32 ReaderIndex, uint32 BlockIndex, const FFileIoStoreResolvedRequest& ResolvedRequest);
-	void ReadNoScatter(uint32 ReaderIndex, uint64 Offset, uint64 Size, const FFileIoStoreResolvedRequest& ResolvedRequest);
+	void ReadBlocks(uint32 ReaderIndex, const FFileIoStoreResolvedRequest& ResolvedRequest);
 	FFileIoStoreBuffer* AllocBuffer();
 	void FreeBuffer(FFileIoStoreBuffer* Buffer);
 	FFileIoStoreCompressionContext* AllocCompressionContext();
@@ -190,11 +205,9 @@ private:
 	void ScatterBlock(FFileIoStoreCompressedBlock* CompressedBlock, bool bIsAsync);
 	void AllocMemoryForRequest(FIoRequestImpl* Request);
 	void FinalizeCompressedBlock(FFileIoStoreCompressedBlock* CompressedBlock);
-	void ReadPendingBlocks();
 
 	const uint64 ReadBufferSize;
 	FIoDispatcherEventQueue& EventQueue;
-	bool bIsMultithreaded;
 	FFileIoStoreImpl PlatformImpl;
 
 	FRunnableThread* Thread;
diff --git a/Engine/Source/Runtime/Core/Private/IO/IoStore.cpp b/Engine/Source/Runtime/Core/Private/IO/IoStore.cpp
--- a/Engine/Source/Runtime/Core/Private/IO/IoStore.cpp
+++ b/Engine/Source/Runtime/Core/Private/IO/IoStore.cpp
@@ -10,6 +10,8 @@
 #include "Async/Async.h"
 #include "Async/Future.h"
 #include "Async/ParallelFor.h"
+#include "HAL/PlatformFilemanager.h"
+#include "Misc/StringBuilder.h"
 
 //////////////////////////////////////////////////////////////////////////
 
@@ -17,6 +19,45 @@ constexpr char FIoStoreTocHeader::TocMagicImg[];
 
 //////////////////////////////////////////////////////////////////////////
 
+/**
+ * I/O store compresssion info.
+ */
+struct FIoStoreCompressionInfo
+{
+	enum
+	{
+		/** Compression method name max length. */
+		CompressionMethodNameLen = 32
+	};
+
+	uint8 GetCompressionMethodIndex(FName CompressionMethod)
+	{
+		if (CompressionMethod == NAME_None)
+		{
+			return 0;
+		}
+
+		for (int32 Idx = 0; Idx < CompressionMethods.Num(); ++Idx)
+		{
+			if (CompressionMethods[Idx] == CompressionMethod)
+			{
+				return uint8(Idx) + 1;
+			}
+		}
+
+		const uint8 Idx = uint8(CompressionMethods.Num());
+		CompressionMethods.Add(CompressionMethod);
+
+		return Idx + 1;
+	}
+
+	TArray<FIoStoreTocCompressedBlockEntry> BlockEntries;
+	TArray<FName> CompressionMethods;
+	int64 BlockSize = 0;
+	int64 UncompressedContainerSize = 0;
+	int64 CompressedContainerSize = 0;
+};
+
 FIoStoreEnvironment::FIoStoreEnvironment()
 {
 }
@@ -25,9 +66,10 @@ FIoStoreEnvironment::~FIoStoreEnvironment()
 {
 }
 
-void FIoStoreEnvironment::InitializeFileEnvironment(FStringView InPath)
+void FIoStoreEnvironment::InitializeFileEnvironment(FStringView InPath, int32 InOrder)
 {
 	Path = InPath;
+	Order = InOrder;
 }
 
 //////////////////////////////////////////////////////////////////////////
@@ -211,7 +253,7 @@ public:
 						check(FileHandle->Tell() % CurrentWriteQueueItem->Alignment == 0);
 					}
 					const int64 CompressedFileOffset = FileHandle->Tell();
-					FIoStoreCompressedBlockEntry& CompressedBlockEntry = CompressionInfo.BlockEntries.AddDefaulted_GetRef();
+					FIoStoreTocCompressedBlockEntry& CompressedBlockEntry = CompressionInfo.BlockEntries.AddDefaulted_GetRef();
 					CompressedBlockEntry.OffsetAndLength.SetOffset(CompressedFileOffset);
 					const uint8* SourceData;
 					uint64 SourceDataSize;
@@ -226,7 +268,6 @@ public:
 						SourceDataSize = CurrentWriteQueueItem->CompressedSize;
 					}
 
-					FIoBuffer& SourceBuffer = CurrentWriteQueueItem->CompressionMethod == NAME_None ? CurrentWriteQueueItem->UncompressedData : CurrentWriteQueueItem->CompressedData;
 					CompressedBlockEntry.OffsetAndLength.SetLength(SourceDataSize);
 					CompressedBlockEntry.CompressionMethodIndex = CompressionInfo.GetCompressionMethodIndex(CurrentWriteQueueItem->CompressionMethod);
 					{
@@ -340,20 +381,6 @@ public:
 		
 		check(!WriteQueueHead);
 		check(!WriteQueueTail);
-		bool bIsUncompressed = true;
-		for (const FIoStoreCompressedBlockEntry& CompressedBlock : CompressionInfo.BlockEntries)
-		{
-			if (CompressedBlock.CompressionMethodIndex != FIoStoreCompressionInfo::InvalidCompressionIndex)
-			{
-				bIsUncompressed = false;
-				break;
-			}
-		}
-		if (bIsUncompressed)
-		{
-			CompressionInfo.BlockEntries.Empty();
-			CompressionInfo.BlockSize = 0;
-		}
 		CompressionInfo.UncompressedContainerSize = UncompressedOffset;
 		CompressionInfo.CompressedContainerSize = FileHandle.IsValid() ? FileHandle->Tell() : 0;
 
@@ -461,8 +488,9 @@ private:
 class FIoStoreWriterImpl
 {
 public:
-	FIoStoreWriterImpl(FIoStoreEnvironment& InEnvironment)
-	:	Environment(InEnvironment)
+	FIoStoreWriterImpl(FIoStoreEnvironment& InEnvironment, FIoContainerId InContainerId)
+		: Environment(InEnvironment)
+		, ContainerId(InContainerId)
 	{
 	}
 
@@ -473,6 +501,7 @@ public:
 		FString TocFilePath = Environment.GetPath() + TEXT(".utoc");
 		FString ContainerFilePath = Environment.GetPath() + TEXT(".ucas");
 
+		Result.ContainerId = ContainerId;
 		Result.ContainerName = FPaths::GetBaseFilename(Environment.GetPath());
 		Result.CompressionMethod = InContext.Settings().CompressionMethod;
 
@@ -517,7 +546,12 @@ public:
 		return FIoStatus::Ok;
 	}
 
-	UE_NODISCARD FIoStatus Append(FIoChunkId ChunkId, FIoBuffer Chunk, FIoWriteOptions WriteOptions)
+	UE_NODISCARD FIoStatus Append(const FIoChunkId& ChunkId, FIoBuffer Chunk, const FIoWriteOptions& WriteOptions)
+	{
+		return Append(ChunkId, FIoChunkHash::HashBuffer(Chunk.Data(), Chunk.DataSize()), Chunk, WriteOptions);
+	}
+
+	UE_NODISCARD FIoStatus Append(const FIoChunkId& ChunkId, const FIoChunkHash& ChunkHash, FIoBuffer Chunk, const FIoWriteOptions& WriteOptions)
 	{
 		if (!ChunkWriter)
 		{
@@ -541,6 +575,7 @@ public:
 		if (TocEntryStatus.IsOk())
 		{
 			FIoStoreTocEntry TocEntry = TocEntryStatus.ConsumeValueOrDie();
+			TocEntry.ChunkHash = ChunkHash;
 			Toc.Add(ChunkId, TocEntry);
 
 			if (CsvArchive)
@@ -588,13 +623,14 @@ public:
 			return FIoStatus(EIoErrorCode::InvalidParameter, TEXT("The given range (Offset/Length) is not within the bounds of OriginalChunkId's data"));
 		}
 
-		FIoStoreTocEntry TocEntry;
+		FIoStoreTocPartialEntry TocPartialEntry;
 
-		TocEntry.SetOffset(Entry->GetOffset() + Offset);
-		TocEntry.SetLength(Length);
-		TocEntry.ChunkId = ChunkIdPartialRange;
+		TocPartialEntry.SetOffset(Entry->GetOffset() + Offset);
+		TocPartialEntry.SetLength(Length);
+		TocPartialEntry.ChunkId = ChunkIdPartialRange;
+		TocPartialEntry.OriginalChunkId = OriginalChunkId;
 
-		Toc.Add(ChunkIdPartialRange, TocEntry);
+		PartialToc.Add(ChunkIdPartialRange, TocPartialEntry);
 
 		IsMetadataDirty = true;
 
@@ -620,9 +656,14 @@ public:
 		TocHeader.TocHeaderSize = sizeof(TocHeader);
 		TocHeader.TocEntryCount = Toc.Num();
 		TocHeader.TocEntrySize = sizeof(FIoStoreTocEntry);
-		TocHeader.CompressionBlockCount = CompressionInfo.BlockEntries.Num();
+		TocHeader.TocPartialEntryCount = PartialToc.Num();
+		TocHeader.TocPartialEntrySize = sizeof(FIoStoreTocPartialEntry);
+		TocHeader.TocCompressedBlockEntryCount = CompressionInfo.BlockEntries.Num();
+		TocHeader.TocCompressedBlockEntrySize = sizeof(FIoStoreTocCompressedBlockEntry);
 		TocHeader.CompressionBlockSize = uint32(CompressionInfo.BlockSize);
-		TocHeader.CompressionNameCount = CompressionInfo.CompressionMethods.Num();
+		TocHeader.CompressionMethodNameCount = CompressionInfo.CompressionMethods.Num();
+		TocHeader.CompressionMethodNameLength = FIoStoreCompressionInfo::CompressionMethodNameLen;
+		TocHeader.ContainerId = ContainerId;
 
 		TocFileHandle->Seek(0);
 		if (!TocFileHandle->Write(reinterpret_cast<const uint8*>(&TocHeader), sizeof(TocHeader)))
@@ -641,8 +682,18 @@ public:
 			}
 		}
 
+		for (auto& _: PartialToc)
+		{
+			FIoStoreTocPartialEntry& TocEntry = _.Value;
+
+			if (!TocFileHandle->Write(reinterpret_cast<const uint8*>(&TocEntry), sizeof(TocEntry)))
+			{
+				return FIoStatus(EIoErrorCode::WriteError, TEXT("Failed to write TOC entry"));
+			}
+		}
+
 		// Compression blocks
-		for (const FIoStoreCompressedBlockEntry& CompressedBlockEntry : CompressionInfo.BlockEntries)
+		for (const FIoStoreTocCompressedBlockEntry& CompressedBlockEntry : CompressionInfo.BlockEntries)
 		{
 			if (!TocFileHandle->Write(reinterpret_cast<const uint8*>(&CompressedBlockEntry), sizeof(CompressedBlockEntry)))
 			{
@@ -665,7 +716,7 @@ public:
 		}
 
 		Result.TocSize = TocFileHandle->Tell();
-		Result.TocEntryCount = TocHeader.TocEntryCount;
+		Result.TocEntryCount = TocHeader.TocEntryCount + TocHeader.TocPartialEntryCount;
 		Result.PaddingSize = ChunkWriter->GetPaddingSize();
 		Result.UncompressedContainerSize = CompressionInfo.UncompressedContainerSize;
 		Result.CompressedContainerSize = CompressionInfo.CompressedContainerSize;
@@ -677,15 +728,17 @@ private:
 	friend class FIoStoreWriter;
 	FIoStoreEnvironment&				Environment;
 	TMap<FIoChunkId, FIoStoreTocEntry>	Toc;
+	TMap<FIoChunkId, FIoStoreTocPartialEntry> PartialToc;
 	TUniquePtr<FChunkWriter>			ChunkWriter;
 	TUniquePtr<IFileHandle>				TocFileHandle;
 	TUniquePtr<FArchive>				CsvArchive;
 	FIoStoreWriterResult				Result;
+	FIoContainerId						ContainerId;
 	bool								IsMetadataDirty = true;
 };
 
-FIoStoreWriter::FIoStoreWriter(FIoStoreEnvironment& InEnvironment)
-:	Impl(new FIoStoreWriterImpl(InEnvironment))
+FIoStoreWriter::FIoStoreWriter(FIoStoreEnvironment& InEnvironment, FIoContainerId InContainerId)
+:	Impl(new FIoStoreWriterImpl(InEnvironment, InContainerId))
 {
 }
 
@@ -699,11 +752,16 @@ FIoStatus FIoStoreWriter::Initialize(const FIoStoreWriterContext& Context, bool
 	return Impl->Initialize(*Context.Impl, bIsContainerCompressed);
 }
 
-FIoStatus FIoStoreWriter::Append(FIoChunkId ChunkId, FIoBuffer Chunk, FIoWriteOptions WriteOptions)
+FIoStatus FIoStoreWriter::Append(const FIoChunkId& ChunkId, FIoBuffer Chunk, const FIoWriteOptions& WriteOptions)
 {
 	return Impl->Append(ChunkId, Chunk, WriteOptions);
 }
 
+FIoStatus FIoStoreWriter::Append(const FIoChunkId& ChunkId, const FIoChunkHash& ChunkHash, FIoBuffer Chunk, const FIoWriteOptions& WriteOptions)
+{
+	return Impl->Append(ChunkId, ChunkHash, Chunk, WriteOptions);
+}
+
 FIoStatus FIoStoreWriter::AppendPadding(uint64 Count)
 {
 	return Impl->AppendPadding(Count);
@@ -718,3 +776,290 @@ TIoStatusOr<FIoStoreWriterResult> FIoStoreWriter::Flush()
 {
 	return Impl->Flush();
 }
+
+class FIoStoreReaderImpl
+{
+public:
+	FIoStoreReaderImpl()
+	{
+
+	}
+
+	UE_NODISCARD FIoStatus Initialize(const FIoStoreEnvironment& Environment)
+	{
+		IPlatformFile& Ipf = FPlatformFileManager::Get().GetPlatformFile();
+
+		TStringBuilder<256> ContainerFilePath;
+		ContainerFilePath.Append(Environment.GetPath());
+
+		TStringBuilder<256> TocFilePath;
+		TocFilePath.Append(ContainerFilePath);
+
+		ContainerFilePath.Append(TEXT(".ucas"));
+		TocFilePath.Append(TEXT(".utoc"));
+
+		ContainerFileHandle.Reset(Ipf.OpenRead(*ContainerFilePath, /* allowwrite */ false));
+		if (!ContainerFileHandle)
+		{
+			return FIoStatusBuilder(EIoErrorCode::FileOpenFailed) << TEXT("Failed to open IoStore container file '") << *TocFilePath << TEXT("'");
+		}
+
+		FIoStoreTocReader TocReader;
+		FIoStatus Status = TocReader.Initialize(*TocFilePath);
+		if (!Status.IsOk())
+		{
+			return Status;
+		}
+
+		uint32 TocEntryCount;
+		const FIoStoreTocEntry* TocEntry = TocReader.GetEntries(TocEntryCount);
+
+		uint32 TocPartialEntryCount;
+		const FIoStoreTocPartialEntry* TocPartialEntry = TocReader.GetPartialEntries(TocPartialEntryCount);
+
+		TocEntries.Reserve(TocEntryCount);
+		int32 EntryIndex = 0;
+		while (TocEntryCount--)
+		{
+			TocEntries.Add(*TocEntry);
+			TocMap.Add(TocEntry->ChunkId, EntryIndex++);
+			++TocEntry;
+		}
+		while (TocPartialEntryCount--)
+		{
+			TocPartialEntries.Add(*TocPartialEntry);
+			TocMap.Add(TocPartialEntry->ChunkId, EntryIndex++);
+			++TocPartialEntry;
+		}
+
+		uint32 CompressionMethodNameCount;
+		const FName* CompressionMethodNames = TocReader.GetCompressionMethodNames(CompressionMethodNameCount);
+
+		CompressionBlockSize = TocReader.GetCompressionBlockSize();
+		CompressedBuffer.SetNumUninitialized(int32(CompressionBlockSize));
+		UncompressedBuffer.SetNumUninitialized(int32(CompressionBlockSize));
+
+		uint32 CompressedBlockEntryCount;
+		const FIoStoreTocCompressedBlockEntry* CompressedBlockEntry = TocReader.GetCompressedBlockEntries(CompressedBlockEntryCount);
+		CompressionBlocks.Reserve(CompressedBlockEntryCount);
+
+		while (CompressedBlockEntryCount--)
+		{
+			FCompressionBlock& CompressionBlock = CompressionBlocks.AddDefaulted_GetRef();
+			CompressionBlock.OffsetAndLength = CompressedBlockEntry->OffsetAndLength;
+			CompressionBlock.CompressionMethod = CompressionMethodNames[CompressedBlockEntry->CompressionMethodIndex];
+			++CompressedBlockEntry;
+		}
+
+		return FIoStatus::Ok;
+	}
+
+	void EnumerateChunks(TFunction<bool(const FIoStoreTocChunkInfo&)>&& Callback) const
+	{
+		for (const FIoStoreTocEntry& TocEntry : TocEntries)
+		{
+			FIoStoreTocChunkInfo ChunkInfo;
+			ChunkInfo.Id = TocEntry.ChunkId;
+			ChunkInfo.Hash = TocEntry.ChunkHash;
+			ChunkInfo.Offset = TocEntry.GetOffset();
+			ChunkInfo.Size = TocEntry.GetLength();
+			if (!Callback(ChunkInfo))
+			{
+				break;
+			}
+		}
+	}
+
+	void EnumeratePartialChunks(TFunction<bool(const FIoStoreTocPartialChunkInfo&)>&& Callback) const
+	{
+		for (const FIoStoreTocPartialEntry& TocPartialEntry : TocPartialEntries)
+		{
+			FIoStoreTocPartialChunkInfo ChunkInfo;
+			ChunkInfo.Id = TocPartialEntry.ChunkId;
+			ChunkInfo.OrginalId = TocPartialEntry.OriginalChunkId;
+			const int32* FindOriginalEntry = TocMap.Find(TocPartialEntry.OriginalChunkId);
+			check(FindOriginalEntry);
+			check(*FindOriginalEntry < TocEntries.Num());
+			const FIoStoreTocEntry& OriginalTocEntry = TocEntries[*FindOriginalEntry];
+			ChunkInfo.RangeStart = TocPartialEntry.GetOffset() - OriginalTocEntry.GetOffset();
+			ChunkInfo.Size = TocPartialEntry.GetLength();
+			if (!Callback(ChunkInfo))
+			{
+				break;
+			}
+		}
+	}
+
+	TIoStatusOr<FIoBuffer> Read(const FIoChunkId& Chunk, const FIoReadOptions& Options) const
+	{
+		FIoOffsetAndLength OffsetAndLength;
+		const int32* FindEntryIndex = TocMap.Find(Chunk);
+		if (!FindEntryIndex)
+		{
+			return FIoStatus(EIoErrorCode::NotFound, TEXT("Unknown chunk ID"));
+		}
+		if (*FindEntryIndex < TocEntries.Num())
+		{
+			OffsetAndLength = TocEntries[*FindEntryIndex].OffsetAndLength;
+		}
+		else
+		{
+			OffsetAndLength = TocPartialEntries[*FindEntryIndex - TocEntries.Num()].OffsetAndLength;
+		}
+
+		FIoBuffer IoBuffer = FIoBuffer(OffsetAndLength.GetLength());
+		int32 FirstBlockIndex = int32(OffsetAndLength.GetOffset() / CompressionBlockSize);
+		int32 LastBlockIndex = int32((Align(OffsetAndLength.GetOffset() + OffsetAndLength.GetLength(), CompressionBlockSize) - 1) / CompressionBlockSize);
+		uint64 OffsetInBlock = OffsetAndLength.GetOffset() % CompressionBlockSize;
+		uint8* Dst = IoBuffer.Data();
+		uint8* Src = nullptr;
+		uint64 RemainingSize = OffsetAndLength.GetLength();
+		for (int32 BlockIndex = FirstBlockIndex; BlockIndex <= LastBlockIndex; ++BlockIndex)
+		{
+			const FCompressionBlock& CompressionBlock = CompressionBlocks[BlockIndex];
+			if (CompressedBuffer.Num() < CompressionBlock.OffsetAndLength.GetLength())
+			{
+				CompressedBuffer.SetNumUninitialized(int32(CompressionBlock.OffsetAndLength.GetLength()));
+			}
+			ContainerFileHandle->Seek(CompressionBlock.OffsetAndLength.GetOffset());
+			ContainerFileHandle->Read(CompressedBuffer.GetData(), CompressionBlock.OffsetAndLength.GetLength());
+			if (CompressionBlock.CompressionMethod.IsNone())
+			{
+				Src = CompressedBuffer.GetData();
+			}
+			else
+			{
+				bool bUncompressed = FCompression::UncompressMemory(CompressionBlock.CompressionMethod, UncompressedBuffer.GetData(), int32(CompressionBlockSize), CompressedBuffer.GetData(), int32(CompressionBlock.OffsetAndLength.GetLength()));
+				if (!bUncompressed)
+				{
+					return FIoStatus(EIoErrorCode::CorruptToc, TEXT("Failed uncompressing block"));
+				}
+				Src = UncompressedBuffer.GetData();
+			}
+			uint64 SizeInBlock = FMath::Min(CompressionBlockSize - OffsetInBlock, RemainingSize);
+			FMemory::Memcpy(Dst, Src + OffsetInBlock, SizeInBlock);
+			OffsetInBlock = 0;
+			RemainingSize -= SizeInBlock;
+			Dst += SizeInBlock;
+		}
+		
+		return IoBuffer;
+	}
+
+private:
+	struct FCompressionBlock
+	{
+		FIoOffsetAndLength OffsetAndLength;
+		FName CompressionMethod;
+	};
+
+	TArray<FIoStoreTocEntry> TocEntries;
+	TArray<FIoStoreTocPartialEntry> TocPartialEntries;
+	TMap<FIoChunkId, int32> TocMap;
+	TArray<FCompressionBlock> CompressionBlocks;
+	uint64 CompressionBlockSize = 0;
+	TUniquePtr<IFileHandle> ContainerFileHandle;
+	mutable TArray<uint8> CompressedBuffer;
+	mutable TArray<uint8> UncompressedBuffer;
+};
+
+
+FIoStoreReader::FIoStoreReader()
+	: Impl(new FIoStoreReaderImpl())
+{
+}
+
+FIoStoreReader::~FIoStoreReader()
+{
+	delete Impl;
+}
+
+FIoStatus FIoStoreReader::Initialize(const FIoStoreEnvironment& InEnvironment)
+{
+	return Impl->Initialize(InEnvironment);
+}
+
+void FIoStoreReader::EnumerateChunks(TFunction<bool(const FIoStoreTocChunkInfo&)>&& Callback) const
+{
+	Impl->EnumerateChunks(MoveTemp(Callback));
+}
+
+void FIoStoreReader::EnumeratePartialChunks(TFunction<bool(const FIoStoreTocPartialChunkInfo&)>&& Callback) const
+{
+	Impl->EnumeratePartialChunks(MoveTemp(Callback));
+}
+
+TIoStatusOr<FIoBuffer> FIoStoreReader::Read(const FIoChunkId& Chunk, const FIoReadOptions& Options) const
+{
+	return Impl->Read(Chunk, Options);
+}
+
+FIoStatus FIoStoreTocReader::Initialize(const TCHAR* TocFilePath)
+{
+	bool bTocReadOk = false;
+
+	{
+		IPlatformFile& Ipf = FPlatformFileManager::Get().GetPlatformFile();
+		TUniquePtr<IFileHandle>	TocFileHandle(Ipf.OpenRead(TocFilePath, /* allowwrite */ false));
+
+		if (!TocFileHandle)
+		{
+			return FIoStatusBuilder(EIoErrorCode::FileOpenFailed) << TEXT("Failed to open IoStore TOC file '") << TocFilePath << TEXT("'");
+		}
+
+		const int64 TocSize = TocFileHandle->Size();
+		TocBuffer = MakeUnique<uint8[]>(TocSize);
+		bTocReadOk = TocFileHandle->Read(TocBuffer.Get(), TocSize);
+	}
+
+	if (!bTocReadOk)
+	{
+		return FIoStatusBuilder(EIoErrorCode::CorruptToc) << TEXT("Failed to read IoStore TOC file '") << TocFilePath << TEXT("'");
+	}
+
+	Header = reinterpret_cast<const FIoStoreTocHeader*>(TocBuffer.Get());
+
+	if (!Header->CheckMagic())
+	{
+		return FIoStatusBuilder(EIoErrorCode::CorruptToc) << TEXT("TOC header magic mismatch while reading '") << TocFilePath << TEXT("'");
+	}
+
+	if (Header->TocHeaderSize != sizeof(FIoStoreTocHeader))
+	{
+		return FIoStatusBuilder(EIoErrorCode::CorruptToc) << TEXT("TOC header size mismatch while reading '") << TocFilePath << TEXT("'");
+	}
+
+	if (Header->TocEntrySize != sizeof(FIoStoreTocEntry))
+	{
+		return FIoStatusBuilder(EIoErrorCode::CorruptToc) << TEXT("TOC entry size mismatch while reading '") << TocFilePath << TEXT("'");
+	}
+
+	if (Header->TocPartialEntrySize != sizeof(FIoStoreTocPartialEntry))
+	{
+		return FIoStatusBuilder(EIoErrorCode::CorruptToc) << TEXT("TOC partial entry size mismatch while reading '") << TocFilePath << TEXT("'");
+	}
+
+	if (Header->TocCompressedBlockEntrySize != sizeof(FIoStoreTocCompressedBlockEntry))
+	{
+		return FIoStatusBuilder(EIoErrorCode::CorruptToc) << TEXT("TOC compressed block entry size mismatch while reading '") << TocFilePath << TEXT("'");
+	}
+
+	Entries = reinterpret_cast<const FIoStoreTocEntry*>(TocBuffer.Get() + sizeof(FIoStoreTocHeader));
+	EntryCount = Header->TocEntryCount;
+
+	PartialEntries = reinterpret_cast<const FIoStoreTocPartialEntry*>(Entries + EntryCount);
+	PartialEntryCount = Header->TocPartialEntryCount;
+
+	CompressionMethodNames.Add(NAME_None);
+	CompressedBlockEntries = reinterpret_cast<const FIoStoreTocCompressedBlockEntry*>(PartialEntries + PartialEntryCount);
+	CompressedBlockEntryCount = Header->TocCompressedBlockEntryCount;
+
+	const ANSICHAR* AnsiCompressionMethodNames = reinterpret_cast<const ANSICHAR*>(CompressedBlockEntries + CompressedBlockEntryCount);
+	for (uint32 CompressonNameIndex = 0; CompressonNameIndex < Header->CompressionMethodNameCount; CompressonNameIndex++)
+	{
+		const ANSICHAR* AnsiCompressionMethodName = AnsiCompressionMethodNames + CompressonNameIndex * Header->CompressionMethodNameLength;
+		CompressionMethodNames.Add(FName(AnsiCompressionMethodName));
+	}
+
+	return FIoStatus::Ok;
+}
\ No newline at end of file
diff --git a/Engine/Source/Runtime/Core/Private/IO/IoStore.h b/Engine/Source/Runtime/Core/Private/IO/IoStore.h
--- a/Engine/Source/Runtime/Core/Private/IO/IoStore.h
+++ b/Engine/Source/Runtime/Core/Private/IO/IoStore.h
@@ -15,10 +15,15 @@ struct FIoStoreTocHeader
 	uint32	TocHeaderSize;
 	uint32	TocEntryCount;
 	uint32	TocEntrySize;	// For sanity checking
-	uint32	CompressionBlockCount;
+	uint32	TocPartialEntryCount;
+	uint32	TocPartialEntrySize;	// For sanity checking
+	uint32	TocCompressedBlockEntryCount;
+	uint32	TocCompressedBlockEntrySize;	// For sanity checking
+	uint32	CompressionMethodNameCount;
+	uint32	CompressionMethodNameLength;
 	uint32	CompressionBlockSize;
-	uint32	CompressionNameCount;
-	uint32	TocPad[22];
+	FIoContainerId ContainerId;
+	uint8 Pad[70];
 
 	void MakeMagic()
 	{
@@ -86,11 +91,36 @@ private:
  */
 struct FIoStoreTocEntry
 {
-	FIoOffsetAndLength	OffsetAndLength;
+	FIoOffsetAndLength OffsetAndLength;
+	FIoChunkId ChunkId;
+	FIoChunkHash ChunkHash;
+	
+	inline uint64 GetOffset() const
+	{
+		return OffsetAndLength.GetOffset();
+	}
+
+	inline uint64 GetLength() const
+	{
+		return OffsetAndLength.GetLength();
+	}
 
-	// TBD: should the chunk ID use content addressing, or names, or a
-	//      mix of both?
-	FIoChunkId	ChunkId;
+	inline void SetOffset(uint64 Offset)
+	{
+		OffsetAndLength.SetOffset(Offset);
+	}
+
+	inline void SetLength(uint64 Length)
+	{
+		OffsetAndLength.SetLength(Length);
+	}
+};
+
+struct FIoStoreTocPartialEntry
+{
+	FIoOffsetAndLength OffsetAndLength;
+	FIoChunkId ChunkId;
+	FIoChunkId OriginalChunkId;
 
 	inline uint64 GetOffset() const
 	{
@@ -112,10 +142,11 @@ struct FIoStoreTocEntry
 		OffsetAndLength.SetLength(Length);
 	}
 };
+
 /**
  * Compression block entry.
  */
-struct FIoStoreCompressedBlockEntry
+struct FIoStoreTocCompressedBlockEntry
 {
 	/** Offset and size of the compressed block. */
 	FIoOffsetAndLength OffsetAndLength;
@@ -123,43 +154,54 @@ struct FIoStoreCompressedBlockEntry
 	uint8 CompressionMethodIndex;
 };
 
-/**
- * I/O store compresssion info.
- */
-struct FIoStoreCompressionInfo
+class FIoStoreTocReader
 {
-	enum
+public:
+	UE_NODISCARD FIoStatus Initialize(const TCHAR* TocFilePath);
+	
+	const FIoStoreTocEntry* GetEntries(uint32& OutCount) const
 	{
-		/** Compression method name max length. */
-		CompressionMethodNameLen = 32,
-		/** No compression. */
-		InvalidCompressionIndex = 255,
-	};
+		OutCount = EntryCount;
+		return EntryCount > 0 ? Entries : nullptr;
+	}
 
-	uint8 GetCompressionMethodIndex(FName CompressionMethod)
+	const FIoStoreTocPartialEntry* GetPartialEntries(uint32& OutCount) const
 	{
-		if (CompressionMethod == NAME_None)
-		{
-			return InvalidCompressionIndex;
-		}
+		OutCount = PartialEntryCount;
+		return PartialEntryCount > 0 ? PartialEntries : nullptr;
+	}
 
-		for (int32 Idx = 0; Idx < CompressionMethods.Num(); ++Idx)
-		{
-			if (CompressionMethods[Idx] == CompressionMethod)
-			{
-				return uint8(Idx);
-			}
-		}
+	const FIoStoreTocCompressedBlockEntry* GetCompressedBlockEntries(uint32& OutCount) const
+	{
+		OutCount = CompressedBlockEntryCount;
+		return CompressedBlockEntryCount > 0 ? CompressedBlockEntries : nullptr;
+	}
 
-		const uint8 Idx = uint8(CompressionMethods.Num());
-		CompressionMethods.Add(CompressionMethod);
+	const FName* GetCompressionMethodNames(uint32& OutCount) const
+	{
+		OutCount = CompressionMethodNames.Num();
+		return CompressionMethodNames.Num() > 0 ? CompressionMethodNames.GetData() : nullptr;
+	}
 
-		return Idx;
+	uint32 GetCompressionBlockSize() const
+	{
+		return Header->CompressionBlockSize;
 	}
 
-	TArray<FIoStoreCompressedBlockEntry> BlockEntries;
-	TArray<FName> CompressionMethods;
-	int64 BlockSize = 0;
-	int64 UncompressedContainerSize = 0;
-	int64 CompressedContainerSize = 0;
+	FIoContainerId GetContainerId() const
+	{
+		return Header->ContainerId;
+	}
+
+private:
+	TArray<FName> CompressionMethodNames;
+	TUniquePtr<uint8[]> TocBuffer;
+	const FIoStoreTocHeader* Header = nullptr;
+	const FIoStoreTocEntry* Entries = nullptr;
+	const FIoStoreTocPartialEntry* PartialEntries = nullptr;
+	const FIoStoreTocCompressedBlockEntry* CompressedBlockEntries = nullptr;
+	uint32 EntryCount = 0;
+	uint32 PartialEntryCount = 0;
+	uint32 CompressedBlockEntryCount = 0;
 };
+
diff --git a/Engine/Source/Runtime/Core/Public/IO/IoContainerId.h b/Engine/Source/Runtime/Core/Public/IO/IoContainerId.h
new file mode 100644
--- /dev/null
+++ b/Engine/Source/Runtime/Core/Public/IO/IoContainerId.h
@@ -0,0 +1,69 @@
+// Copyright Epic Games, Inc. All Rights Reserved.
+
+#pragma once
+
+#include "CoreMinimal.h"
+
+class FArchive;
+class FStructuredArchiveSlot;
+
+/**
+ * Container ID.
+ */
+class FIoContainerId
+{
+public:
+	inline FIoContainerId() = default;
+	inline FIoContainerId(const FIoContainerId& Other) = default;
+	inline FIoContainerId(FIoContainerId&& Other) = default;
+	inline FIoContainerId& operator=(const FIoContainerId& Other) = default;
+
+	static FIoContainerId FromIndex(const uint16 InIndex)
+	{
+		return FIoContainerId(InIndex);
+	}
+
+	inline bool IsValid() const
+	{ 
+		return Id != InvalidId;
+	}
+
+	inline uint16 ToIndex() const
+	{
+		check(Id != InvalidId);
+		return Id;
+	}
+
+	inline bool operator<(FIoContainerId Other) const
+	{
+		return Id < Other.Id;
+	}
+
+	inline bool operator==(FIoContainerId Other) const
+	{
+		return Id == Other.Id;
+	}
+
+	inline bool operator!=(FIoContainerId Other) const
+	{
+		return Id != Other.Id;
+	}
+
+	inline friend uint32 GetTypeHash(const FIoContainerId& In)
+	{
+		return In.Id;
+	}
+
+	CORE_API friend FArchive& operator<<(FArchive& Ar, FIoContainerId& ContainerId);
+
+	CORE_API friend void operator<<(FStructuredArchiveSlot Slot, FIoContainerId& Value);
+
+private:
+	inline explicit FIoContainerId(const uint16 InId)
+		: Id(InId) { }
+
+	static constexpr uint16 InvalidId = ~uint16(0u);
+
+	uint16 Id = InvalidId;
+};
+
diff --git a/Engine/Source/Runtime/Core/Public/IO/IoDispatcher.h b/Engine/Source/Runtime/Core/Public/IO/IoDispatcher.h
--- a/Engine/Source/Runtime/Core/Public/IO/IoDispatcher.h
+++ b/Engine/Source/Runtime/Core/Public/IO/IoDispatcher.h
@@ -3,6 +3,7 @@
 #pragma once
 
 #include "CoreTypes.h"
+#include "IO/IoContainerId.h"
 #include "Containers/StringView.h"
 #include "Containers/UnrealString.h"
 #include "Logging/LogMacros.h"
@@ -10,6 +11,7 @@
 #include "Templates/UnrealTemplate.h"
 #include "Templates/TypeCompatibleBytes.h"
 #include "HAL/PlatformAtomics.h"
+#include "Misc/SecureHash.h"
 
 #if __cplusplus >= 201703L
 #	define UE_NODISCARD		[[nodiscard]]
@@ -29,6 +31,7 @@ class FIoBatchImpl;
 class FIoDispatcherImpl;
 class FIoStoreWriterContextImpl;
 class FIoStoreWriterImpl;
+class FIoStoreReaderImpl;
 class IMappedFileHandle;
 class IMappedFileRegion;
 
@@ -377,6 +380,26 @@ TIoStatusOr<T>& TIoStatusOr<T>::operator=(const TIoStatusOr<U>& Other)
 	return *this;
 }
 
+//////////////////////////////////////////////////////////////////////////
+
+/** Helper used to manage creation of I/O store file handles etc
+  */
+class FIoStoreEnvironment
+{
+public:
+	CORE_API FIoStoreEnvironment();
+	CORE_API ~FIoStoreEnvironment();
+
+	CORE_API void InitializeFileEnvironment(FStringView InPath, int32 InOrder = 0);
+
+	CORE_API const FString& GetPath() const { return Path; }
+	CORE_API int32 GetOrder() const { return Order; }
+
+private:
+	FString			Path;
+	int32			Order = 0;
+};
+
 /** Reference to buffer data used by I/O dispatcher APIs
   */
 class FIoBuffer
@@ -533,6 +556,47 @@ private:
 	friend class FIoBufferManager;
 };
 
+class FIoChunkHash
+{
+public:
+	friend uint32 GetTypeHash(const FIoChunkHash& InChunkHash)
+	{
+		uint32 Result = 5381;
+		for (int i = 0; i < sizeof Hash; ++i)
+		{
+			Result = Result * 33 + InChunkHash.Hash[i];
+		}
+		return Result;
+	}
+
+	friend FArchive& operator<<(FArchive& Ar, FIoChunkHash& ChunkHash)
+	{
+		Ar.Serialize(&ChunkHash.Hash, sizeof Hash);
+		return Ar;
+	}
+
+	inline bool operator ==(const FIoChunkHash& Rhs) const
+	{
+		return 0 == FMemory::Memcmp(Hash, Rhs.Hash, sizeof Hash);
+	}
+
+	inline bool operator !=(const FIoChunkHash& Rhs) const
+	{
+		return !(*this == Rhs);
+	}
+
+	static FIoChunkHash HashBuffer(const void* Data, uint64 DataSize)
+	{
+		FIoChunkHash Result;
+		FSHA1::HashBuffer(Data, DataSize, Result.Hash);
+		FMemory::Memset(Result.Hash + 20, 0, 12);
+		return Result;
+	}
+
+private:
+	uint8	Hash[32];
+};
+
 /**
  * Identifier to a chunk of data.
  */
@@ -605,7 +669,8 @@ enum class EIoChunkType : uint8
 	LoaderGlobalMeta,
 	LoaderInitialLoadMeta,
 	LoaderGlobalNames,
-	LoaderGlobalNameHashes
+	LoaderGlobalNameHashes,
+	ContainerHeader
 };
 
 /**
@@ -817,11 +882,19 @@ struct FIoMappedRegion
 	IMappedFileRegion* MappedFileRegion = nullptr;
 };
 
+struct FIoDispatcherMountedContainer
+{
+	FIoStoreEnvironment Environment;
+	FIoContainerId ContainerId;
+};
+
 /** I/O dispatcher
   */
 class FIoDispatcher
 {
 public:
+	DECLARE_EVENT_OneParam(FIoDispatcher, FIoContainerMountedEvent, const FIoDispatcherMountedContainer&);
+
 	CORE_API						FIoDispatcher();
 	CORE_API virtual				~FIoDispatcher();
 
@@ -837,6 +910,10 @@ public:
 	// Polling methods
 	CORE_API bool					DoesChunkExist(const FIoChunkId& ChunkId) const;
 	CORE_API TIoStatusOr<uint64>	GetSizeForChunk(const FIoChunkId& ChunkId) const;
+	CORE_API TArray<FIoDispatcherMountedContainer> GetMountedContainers() const;
+
+	// Events
+	CORE_API FIoContainerMountedEvent& OnContainerMounted();
 
 	FIoDispatcher(const FIoDispatcher&) = default;
 	FIoDispatcher& operator=(const FIoDispatcher&) = delete;
@@ -857,24 +934,6 @@ private:
 
 //////////////////////////////////////////////////////////////////////////
 
-/** Helper used to manage creation of I/O store file handles etc 
-  */
-class FIoStoreEnvironment
-{
-public:
-	CORE_API FIoStoreEnvironment();
-	CORE_API ~FIoStoreEnvironment();
-
-	CORE_API void InitializeFileEnvironment(FStringView InPath);
-
-	CORE_API const FString& GetPath() const { return Path; }
-
-private:
-	FString			Path;
-};
-
-//////////////////////////////////////////////////////////////////////////
-
 struct FIoStoreWriterSettings
 {
 	FName CompressionMethod = NAME_None;
@@ -885,6 +944,7 @@ struct FIoStoreWriterSettings
 
 struct FIoStoreWriterResult
 {
+	FIoContainerId ContainerId;
 	FString ContainerName;
 	int64 TocSize = 0;
 	int64 TocEntryCount = 0;
@@ -918,14 +978,15 @@ private:
 class FIoStoreWriter
 {
 public:
-	CORE_API 			FIoStoreWriter(FIoStoreEnvironment& InEnvironment);
+	CORE_API 			FIoStoreWriter(FIoStoreEnvironment& InEnvironment, FIoContainerId InContainerId);
 	CORE_API virtual	~FIoStoreWriter();
 
 	FIoStoreWriter(const FIoStoreWriter&) = delete;
 	FIoStoreWriter& operator=(const FIoStoreWriter&) = delete;
 
 	UE_NODISCARD CORE_API FIoStatus	Initialize(const FIoStoreWriterContext& Context, bool bIsContainerCompressed);
-	UE_NODISCARD CORE_API FIoStatus	Append(FIoChunkId ChunkId, FIoBuffer Chunk, FIoWriteOptions WriteOptions);
+	UE_NODISCARD CORE_API FIoStatus	Append(const FIoChunkId& ChunkId, const FIoChunkHash& ChunkHash, FIoBuffer Chunk, const FIoWriteOptions& WriteOptions);
+	UE_NODISCARD CORE_API FIoStatus	Append(const FIoChunkId& ChunkId, FIoBuffer Chunk, const FIoWriteOptions& WriteOptions);
 	UE_NODISCARD CORE_API FIoStatus	AppendPadding(uint64 Count);
 
 	/**
@@ -943,5 +1004,36 @@ private:
 	FIoStoreWriterImpl*		Impl;
 };
 
+struct FIoStoreTocChunkInfo
+{
+	FIoChunkId Id;
+	FIoChunkHash Hash;
+	uint64 Offset;
+	uint64 Size;
+};
+
+struct FIoStoreTocPartialChunkInfo
+{
+	FIoChunkId Id;
+	FIoChunkId OrginalId;
+	uint64 RangeStart;
+	uint64 Size;
+};
+
+class FIoStoreReader
+{
+public:
+	CORE_API FIoStoreReader();
+	CORE_API ~FIoStoreReader();
+
+	UE_NODISCARD CORE_API FIoStatus Initialize(const FIoStoreEnvironment& InEnvironment);
+	CORE_API void EnumerateChunks(TFunction<bool(const FIoStoreTocChunkInfo&)>&& Callback) const;
+	CORE_API void EnumeratePartialChunks(TFunction<bool(const FIoStoreTocPartialChunkInfo&)>&& Callback) const;
+	CORE_API TIoStatusOr<FIoBuffer> Read(const FIoChunkId& Chunk, const FIoReadOptions& Options) const;
+
+private:
+	FIoStoreReaderImpl* Impl;
+};
+
 //////////////////////////////////////////////////////////////////////////
 
diff --git a/Engine/Source/Runtime/CoreUObject/Private/Serialization/AsyncLoading2.cpp b/Engine/Source/Runtime/CoreUObject/Private/Serialization/AsyncLoading2.cpp
--- a/Engine/Source/Runtime/CoreUObject/Private/Serialization/AsyncLoading2.cpp
+++ b/Engine/Source/Runtime/CoreUObject/Private/Serialization/AsyncLoading2.cpp
@@ -46,23 +46,34 @@
 #include "Serialization/AsyncPackage.h"
 #include "Serialization/UnversionedPropertySerialization.h"
 #include "Serialization/Zenaphore.h"
-#include "IO/IoDispatcher.h"
 #include "UObject/GCObject.h"
 #include "UObject/ObjectRedirector.h"
 #include "Serialization/BulkData.h"
 #include "Serialization/LargeMemoryReader.h"
+#include "Serialization/MemoryReader.h"
 #include "UObject/UObjectClusters.h"
 #include "ProfilingDebugging/CountersTrace.h"
+#include "Async/Async.h"
 #include "HAL/LowLevelMemStats.h"
 
 #if UE_BUILD_DEVELOPMENT || UE_BUILD_DEBUG
 PRAGMA_DISABLE_OPTIMIZATION
 #endif
 
-FArchive& operator<<(FArchive& Ar, FNameMapEntry& NameMapEntry)
+FArchive& operator<<(FArchive& Ar, FMappedName& MappedName)
 {
-	Ar << NameMapEntry.Index;
-	Ar << NameMapEntry.Number;
+	Ar << MappedName.Index << MappedName.Number;
+
+	return Ar;
+}
+
+FArchive& operator<<(FArchive& Ar, FContainerHeader& ContainerHeader)
+{
+	Ar << ContainerHeader.ContainerId;
+	Ar << ContainerHeader.Names;
+	Ar << ContainerHeader.NameHashes;
+	Ar << ContainerHeader.PackageIds;
+	Ar << ContainerHeader.PackageNames;
 
 	return Ar;
 }
@@ -93,7 +104,7 @@ FArchive& operator<<(FArchive& Ar, FExportBundleMetaEntry& ExportBundleMetaEntry
 
 FArchive& operator<<(FArchive& Ar, FScriptObjectEntry& ScriptObjectEntry)
 {
-	Ar << ScriptObjectEntry.ObjectName;
+	Ar << ScriptObjectEntry.ObjectName.Index << ScriptObjectEntry.ObjectName.Number;
 	Ar << ScriptObjectEntry.OuterIndex;
 	Ar << ScriptObjectEntry.CDOClassIndex;
 	return Ar;
@@ -103,7 +114,7 @@ FArchive& operator<<(FArchive& Ar, FExportMapEntry& ExportMapEntry)
 {
 	Ar << ExportMapEntry.CookedSerialOffset;
 	Ar << ExportMapEntry.CookedSerialSize;
-	Ar << ExportMapEntry.ObjectName.NameEntryIndex << ExportMapEntry.ObjectName.Number;
+	Ar << ExportMapEntry.ObjectName;
 	Ar << ExportMapEntry.OuterIndex;
 	Ar << ExportMapEntry.ClassIndex;
 	Ar << ExportMapEntry.SuperIndex;
@@ -382,10 +393,10 @@ private:
 	}
 };
 
-class FGlobalNameMap
+class FNameMap
 {
 public:
-	void Load(FIoDispatcher& IoDispatcher)
+	void LoadGlobal(FIoDispatcher& IoDispatcher)
 	{
 		check(NameEntries.Num() == 0);
 
@@ -405,23 +416,29 @@ public:
 		FIoBuffer NameBuffer = NameRequest.GetResult().ConsumeValueOrDie();
 		FIoBuffer HashBuffer = HashRequest.GetResult().ConsumeValueOrDie();
 
-		LoadNameBatch(/* Out */ NameEntries, 
-						MakeArrayView(NameBuffer.Data(), NameBuffer.DataSize()),
-						MakeArrayView(HashBuffer.Data(), HashBuffer.DataSize()));
+		Load(MakeArrayView(NameBuffer.Data(), NameBuffer.DataSize()), MakeArrayView(HashBuffer.Data(), HashBuffer.DataSize()), FMappedName::EType::Global);
 
 		IoDispatcher.FreeBatch(Batch);
 	}
 
-	FName GetName(const FNameMapIndex& NameMapIndex) const
+	void Load(TArrayView<const uint8> NameBuffer, TArrayView<const uint8> HashBuffer, FMappedName::EType InNameMapType)
 	{
-		FNameEntryId NameEntry = NameEntries[NameMapIndex.NameEntryIndex];
-		return FName::CreateFromDisplayId(NameEntry, NameMapIndex.Number);
+		LoadNameBatch(NameEntries, NameBuffer, HashBuffer);
+		NameMapType = InNameMapType;
 	}
 
-	FNameMapEntry FromSerializedNameMapEntry(const FNameMapEntry& SerializedNameMapEntry) const
+	FName GetName(const FMappedName& MappedName) const
 	{
-		FNameEntryId NameEntry = NameEntries[SerializedNameMapEntry.Index];
-		return FNameMapEntry { NameEntry.ToUnstableInt(), SerializedNameMapEntry.Number };
+		check(MappedName.GetType() == NameMapType);
+		FNameEntryId NameEntry = NameEntries[MappedName.GetIndex()];
+		return FName::CreateFromDisplayId(NameEntry, MappedName.GetNumber());
+	}
+
+	FMinimalName GetMinimalName(const FMappedName& MappedName) const
+	{
+		check(MappedName.GetType() == NameMapType);
+		FNameEntryId NameEntry = NameEntries[MappedName.GetIndex()];
+		return FMinimalName(NameEntry, MappedName.GetNumber());
 	}
 
 	const TArray<FNameEntryId>& GetNameEntries() const
@@ -431,40 +448,43 @@ public:
 
 private:
 	TArray<FNameEntryId> NameEntries;
+	FMappedName::EType NameMapType = FMappedName::EType::Global;
 };
 
 struct FGlobalImportStore
 {
-	TArray<FGlobalImport> Imports[FPackageObjectIndex::ImportTypeCount];
+	TArray<UObject*> ScriptObjects;
+	TArray<FGlobalImport> PublicExports;
 	TArray<FScriptObjectEntry> ScriptObjectEntries;
-	TArray<FName> PackageImportNames;
 
 	inline FName GetName(FPackageObjectIndex GlobalIndex)
 	{
 		check(GlobalIndex.IsImport());
-		FName Name;
-		if (GlobalIndex.IsPackageImport())
-		{
-			Name = PackageImportNames[GlobalIndex.GetIndex()];
-		}
-		else if (GlobalIndex.IsScriptImport() && ScriptObjectEntries.Num() > 0)
-		{
-			Name = ScriptObjectEntries[GlobalIndex.GetIndex()].ObjectName.ToFName();
-		}
-		return Name;
+
+		return GlobalIndex.IsScriptImport() && ScriptObjectEntries.Num() > 0
+			? MinimalNameToName(ScriptObjectEntries[GlobalIndex.GetIndex()].ObjectName)
+			: NAME_None;
 	}
 
 	inline FGlobalImport& GetImport(FPackageObjectIndex GlobalIndex)
 	{
-		check(GlobalIndex.IsImport());
-		FGlobalImport& Import = Imports[GlobalIndex.GetType()][GlobalIndex.GetIndex()];
+		check(GlobalIndex.IsPackageImport());
+		FGlobalImport& Import = PublicExports[GlobalIndex.GetIndex()];
 		return Import;
 	}
 
 	inline UObject* GetImportObject(FPackageObjectIndex GlobalIndex)
 	{
-		FGlobalImport& Import = GetImport(GlobalIndex);
-		return Import.GetObject();
+		if (GlobalIndex.IsScriptImport())
+		{
+			return ScriptObjects[GlobalIndex.GetIndex()];
+		}
+		if (GlobalIndex.IsPackageImport())
+		{
+			return PublicExports[GlobalIndex.GetIndex()].GetObject();
+		}
+		check(false);
+		return nullptr;
 	}
 
 	UObject* FindScriptImportObjectFromIndex(int32 ScriptImportIndex);
@@ -505,13 +525,23 @@ struct FGlobalImportStore
 		return Object;
 	}
 
-	void StoreGlobalImportObject(FPackageObjectIndex GlobalIndex, UObject* InObject)
+	void StoreGlobalObject(FPackageObjectIndex GlobalIndex, UObject* InObject)
 	{
-		if (GlobalIndex.IsImport())
+		if (GlobalIndex.IsScriptImport())
+		{
+			UObject*& Object = ScriptObjects[GlobalIndex.GetIndex()];
+			check(!Object || Object == InObject);
+			Object = InObject;
+		}
+		else if (GlobalIndex.IsPackageImport())
 		{
-			FGlobalImport& Import = Imports[GlobalIndex.GetType()][GlobalIndex.GetIndex()];
+			FGlobalImport& Import = PublicExports[GlobalIndex.GetIndex()];
 			Import.SetObject(InObject);
 		}
+		else
+		{
+			check(GlobalIndex.IsNull());
+		}
 	}
 
 	/** Reference tracking for GC management */
@@ -521,22 +551,40 @@ struct FGlobalImportStore
 	void OnPreGarbageCollect(bool bInIsLoadingPackages);
 	void OnPostGarbageCollect();
 
-	void FindAllScriptImports();
+	void FindAllScriptObjects();
 };
 
 class FPackageStore
 {
 public:
+	FPackageStore(FIoDispatcher& InIoDispatcher, FNameMap& InGlobalNameMap)
+		: IoDispatcher(InIoDispatcher)
+		, GlobalNameMap(InGlobalNameMap) { }
+	
+	struct FLoadedContainer
+	{
+		TUniquePtr<FNameMap> ContainerNameMap;
+		int32 Order;
+		bool bValid = false;
+	};
+
+	FIoDispatcher& IoDispatcher;
+	FNameMap& GlobalNameMap;
+	TArray<FLoadedContainer> LoadedContainers;
+
 	FCriticalSection PackageNameToPackageIdCritical;
 	TMap<FName, FPackageId> PackageNameToPackageId;
 
+	FCulturePackageMap CulturePackageMap;
+	FSourceToLocalizedPackageIdMap* LocalizedPackageIdMap = nullptr;
+
 	FGlobalImportStore ImportStore;
 	FPackageStoreEntry* StoreEntries = nullptr;
 	int32 PackageCount = 0;
 	int32 ScriptArcsCount = 0;
 
 public:
-	void Load(FIoDispatcher& IoDispatcher, FGlobalNameMap& GlobalNameMap)
+	void Load()
 	{
 		FIoBuffer IoBuffer;
 		FEvent* Event = FPlatformProcess::GetSynchEventFromPool();
@@ -574,7 +622,9 @@ public:
 			TRACE_CPUPROFILER_EVENT_SCOPE(LoadPackageStoreTocFixup);
 
 			int32 PackageByteCount = 0;
+			int32 PackageImportCount = 0;
 			GlobalMetaAr << PackageCount;
+			GlobalMetaAr << PackageImportCount;
 			GlobalMetaAr << PackageByteCount;
 
 			StoreEntries = reinterpret_cast<FPackageStoreEntry*>(FMemory::Malloc(PackageByteCount));
@@ -582,135 +632,247 @@ public:
 			// In-place loading
 			GlobalMetaAr.Serialize(StoreEntries, PackageByteCount);
 
-			// FName fixup
+			ImportStore.PublicExports.AddDefaulted(PackageImportCount);
+
+			// add 10% slack for temp package names
 			{
-				TRACE_CPUPROFILER_EVENT_SCOPE(LoadPackageStoreFNameFixup);
-				for (int I = 0; I < PackageCount; ++I)
-				{
-					StoreEntries[I].Name = GlobalNameMap.FromSerializedNameMapEntry(StoreEntries[I].Name);
-				}
+				FScopeLock Lock(&PackageNameToPackageIdCritical);
+				PackageNameToPackageId.Reserve(PackageCount + PackageCount / 10);
 			}
+		}
+
+		{
+			TRACE_CPUPROFILER_EVENT_SCOPE(LoadPackageStoreLocalization);
+
+			GlobalMetaAr << CulturePackageMap;
+
+			FString CurrentCultureName = FInternationalization::Get().GetCurrentCulture()->GetName();
+			FParse::Value(FCommandLine::Get(), TEXT("CULTURE="), CurrentCultureName);
+
+			LocalizedPackageIdMap = CulturePackageMap.Find(CurrentCultureName);
+		}
+
+		// Load initial loading meta data
+		{
+			TRACE_CPUPROFILER_EVENT_SCOPE(LoadPackageStoreInitialLoadIo);
+			InitialLoadEvent->Wait();
+			FPlatformProcess::ReturnSynchEventToPool(InitialLoadEvent);
+		}
+
+		{
+			TRACE_CPUPROFILER_EVENT_SCOPE(LoadPackageStoreInitialLoadFixup);
+			FLargeMemoryReader InitialLoadArchive(InitialLoadIoBuffer.Data(), InitialLoadIoBuffer.DataSize());
+			int32 NumScriptObjects = 0;
+			InitialLoadArchive << NumScriptObjects;
+			ImportStore.ScriptObjectEntries = MakeArrayView(reinterpret_cast<const FScriptObjectEntry*>(InitialLoadIoBuffer.Data() + InitialLoadArchive.Tell()), NumScriptObjects);
 
-			// Global package index
+			for (FScriptObjectEntry& ScriptObjectEntry : ImportStore.ScriptObjectEntries)
 			{
-				TRACE_CPUPROFILER_EVENT_SCOPE(LoadPackageStoreGlobalIds);
-				// add 10% slack for temp package names
-				PackageNameToPackageId.Reserve(PackageCount + PackageCount / 10);
-				for (int I = 0; I < PackageCount; ++I)
-				{
-					PackageNameToPackageId.Add(StoreEntries[I].Name.ToFName(), FPackageId::FromIndex(I));
-				}
+				const FMappedName& MappedName = FMappedName::FromMinimalName(ScriptObjectEntry.ObjectName);
+				check(MappedName.IsGlobal());
+				ScriptObjectEntry.ObjectName = GlobalNameMap.GetMinimalName(MappedName);
 			}
+
+			ImportStore.ScriptObjects.AddDefaulted(ImportStore.ScriptObjectEntries.Num());
 		}
 
+		FPackageName::DoesPackageExistOverride().BindLambda([this](FName InPackageName)
 		{
-			TRACE_CPUPROFILER_EVENT_SCOPE(LoadPackageStoreGlobalImports);
+			FScopeLock Lock(&PackageNameToPackageIdCritical);
+			return PackageNameToPackageId.Contains(InPackageName);
+		});
 
-			int32 Count = 0;
-			GlobalMetaAr << Count; 
+		LoadContainers(IoDispatcher.GetMountedContainers());
+		IoDispatcher.OnContainerMounted().AddRaw(this, &FPackageStore::OnContainerMounted);
+	}
+
+	void LoadContainers(TArrayView<const FIoDispatcherMountedContainer> Containers)
+	{
+		TRACE_CPUPROFILER_EVENT_SCOPE(LoadContainers);
 
-			TArray<FGlobalImport>& Imports = ImportStore.Imports[FPackageObjectIndex::PackageImport];
-			TArray<FName>& Names = ImportStore.PackageImportNames;
-			Imports.AddDefaulted(Count);
-			Names.Reserve(Count);
+		int32 ContainersToLoad = 0;
 
-			FNameMapEntry* NameMapEntries = reinterpret_cast<FNameMapEntry*>(IoBuffer.Data() + GlobalMetaAr.Tell());
-			for (int I = 0; I < Count; ++I)
+		int32 MaxContainerIndex = LoadedContainers.Num();
+		for (const FIoDispatcherMountedContainer& Container : Containers)
+		{
+			const int32 ContainerIndex = Container.ContainerId.ToIndex();
+			if (ContainerIndex != 0)
 			{
-				Names.Emplace(NameMapEntries[I].ToFName());
+				++ContainersToLoad;
+				MaxContainerIndex = FMath::Max(MaxContainerIndex, ContainerIndex);
 			}
-			GlobalMetaAr.Seek(GlobalMetaAr.Tell() + Count * sizeof(FName));
 		}
 
+		if (!ContainersToLoad)
 		{
-			TRACE_CPUPROFILER_EVENT_SCOPE(LoadPackageStoreLocalization);
+			return;
+		}
 
-			FCulturePackageMap CulturePackageMap;
+		const int32 NewContainerMapSize = MaxContainerIndex + 1;
+		if (NewContainerMapSize > LoadedContainers.Num())
+		{
+			LoadedContainers.SetNum(NewContainerMapSize, false);
+		}
 
-			GlobalMetaAr << CulturePackageMap;
+		TAtomic<int32> Remaining(ContainersToLoad);
 
-			FString CurrentCultureName = FInternationalization::Get().GetCurrentCulture()->GetName();
-			FParse::Value(FCommandLine::Get(), TEXT("CULTURE="), CurrentCultureName);
+		FEvent* Event = FPlatformProcess::GetSynchEventFromPool();
+
+		for (const FIoDispatcherMountedContainer& Container : Containers)
+		{
+			FIoContainerId ContainerId = Container.ContainerId;
+			if (ContainerId.ToIndex() == 0)
+			{
+				continue;
+			}
+
+			FLoadedContainer& LoadedContainer = LoadedContainers[ContainerId.ToIndex()];
+			if (LoadedContainer.bValid && LoadedContainer.Order >= Container.Environment.GetOrder())
+			{
+				UE_LOG(LogStreaming, Log, TEXT("Skipping loading mounted container ID '%d', already loaded with higher order"), ContainerId.ToIndex());
+				if (--Remaining == 0)
+				{
+					Event->Trigger();
+				}
+				continue;
+			}
 
-			FSourceToLocalizedPackageIdMap* LocalizedPackageIdMap = CulturePackageMap.Find(CurrentCultureName);
-			if (LocalizedPackageIdMap)
+			UE_LOG(LogStreaming, Log, TEXT("Loading mounted container ID '%d'"), ContainerId.ToIndex());
+			LoadedContainer.bValid = true;
+			LoadedContainer.Order = Container.Environment.GetOrder();
+
+			FIoChunkId HeaderChunkId = CreateIoChunkId(0, ContainerId.ToIndex(), EIoChunkType::ContainerHeader);
+			IoDispatcher.ReadWithCallback(HeaderChunkId, FIoReadOptions(), [this, &Remaining, Event, &LoadedContainer](TIoStatusOr<FIoBuffer> Result)
 			{
-				for (auto It = LocalizedPackageIdMap->CreateIterator(); It; ++It)
+				// Execution method Thread will run the async block synchronously when multithreading is NOT supported
+				const EAsyncExecution ExecutionMethod = FPlatformProcess::SupportsMultithreading() ? EAsyncExecution::TaskGraph : EAsyncExecution::Thread;
+
+				Async(ExecutionMethod, [this, &Remaining, Event, IoBuffer = Result.ConsumeValueOrDie(), &LoadedContainer]()
 				{
-					const FName SourceName = StoreEntries[It.Key().ToIndex()].Name.ToFName();
-					FPackageId* FoundPackageId = PackageNameToPackageId.Find(SourceName);
+					FMemoryReaderView Ar(MakeArrayView(IoBuffer.Data(), IoBuffer.DataSize()));
 
-					UE_CLOG(!FoundPackageId, LogStreaming, Warning,
-						TEXT("Skip remapping for localized package %s (%d) since the source package %s (%d) is unknown."),
-						It.Value().ToIndex(),
-						*StoreEntries[It.Value().ToIndex()].Name.ToFName().ToString(),
-						It.Key().ToIndex(),
-						*SourceName.ToString());
+					FContainerHeader ContainerHeader;
+					Ar << ContainerHeader;
 
-					if (FoundPackageId)
+					const bool bHasContainerLocalNameMap = ContainerHeader.Names.Num() > 0;
+					if (bHasContainerLocalNameMap)
 					{
-						const FPackageId LocalizedPackageId = It.Value();
-						*FoundPackageId = LocalizedPackageId;
-
-#if ALT2_LOG_VERBOSE
-						UE_LOG(LogStreaming, Verbose, TEXT("LoadPackageStore: RemapLocalizedPackage: %s (%d) %s (%d)."),
-							*StoreEntries[It.Key().ToIndex()].Name.ToFName().ToString(),
-							It.Key().ToIndex(),
-							*StoreEntries[It.Value().ToIndex()].Name.ToFName().ToString(),
-							It.Value().ToIndex());
-#endif
+						TRACE_CPUPROFILER_EVENT_SCOPE(LoadContainerNameMap);
+						LoadedContainer.ContainerNameMap.Reset(new FNameMap());
+						LoadedContainer.ContainerNameMap->Load(ContainerHeader.Names, ContainerHeader.NameHashes, FMappedName::EType::Container);
 					}
-				}
 
-				for (int I = 0; I < PackageCount; ++I)
-				{
-					FPackageStoreEntry& StoreEntry = StoreEntries[I];
-					for (FPackageId& ImportedPackageIndex : StoreEntry.ImportedPackages)
+					FNameMap& NameMap = bHasContainerLocalNameMap ? *LoadedContainer.ContainerNameMap : GlobalNameMap;
+
 					{
-						if (FPackageId* Value = LocalizedPackageIdMap->Find(ImportedPackageIndex))
+						TRACE_CPUPROFILER_EVENT_SCOPE(AddPackages);
+						FScopeLock Lock(&PackageNameToPackageIdCritical);
+
+						int32 NameIndex = 0;
+						for (FPackageId PackageId : ContainerHeader.PackageIds)
 						{
-							ImportedPackageIndex = *Value;
+							FName PackageName = NameMap.GetName(ContainerHeader.PackageNames[NameIndex++]);
+							FPackageId& Id = PackageNameToPackageId.FindOrAdd(PackageName);
+							if (!Id.IsValid())
+							{
+								Id = PackageId;
+								StoreEntries[PackageId.ToIndex()].Name = NameToMinimalName(PackageName);
+							}
 						}
 					}
-				}
-			}
+
+					if (--Remaining == 0)
+					{
+						Event->Trigger();
+					}
+				});
+			});
 		}
 
-		// Load initial loading meta data
+		Event->Wait();
+		FPlatformProcess::ReturnSynchEventToPool(Event);
+
+		ApplyLocalization();
+	}
+
+	void OnContainerMounted(const FIoDispatcherMountedContainer& Container)
+	{
+		LoadContainers(MakeArrayView(&Container, 1));
+	}
+
+	void ApplyLocalization()
+	{
+		TRACE_CPUPROFILER_EVENT_SCOPE(ApplyLocalization);
+
+		if (!LocalizedPackageIdMap)
 		{
-			TRACE_CPUPROFILER_EVENT_SCOPE(LoadPackageStoreInitialLoadIo);
-			InitialLoadEvent->Wait();
-			FPlatformProcess::ReturnSynchEventToPool(InitialLoadEvent);
+			return;
 		}
 
+		// Mark already localized entries if this becomes a perfromance problem
+
+		for (auto It = LocalizedPackageIdMap->CreateIterator(); It; ++It)
 		{
-			TRACE_CPUPROFILER_EVENT_SCOPE(LoadPackageStoreInitialLoadFixup);
-			FLargeMemoryReader InitialLoadArchive(InitialLoadIoBuffer.Data(), InitialLoadIoBuffer.DataSize());
-			int32 NumScriptObjects = 0;
-			InitialLoadArchive << NumScriptObjects;
-			ImportStore.ScriptObjectEntries = MakeArrayView(reinterpret_cast<const FScriptObjectEntry*>(InitialLoadIoBuffer.Data() + InitialLoadArchive.Tell()), NumScriptObjects);
+			const FPackageStoreEntry& PackageEntry = StoreEntries[It.Key().ToIndex()];
 
-			for (FScriptObjectEntry& ScriptObjectEntry : ImportStore.ScriptObjectEntries)
+			if (!FMappedName::IsResolvedToMinimalName(PackageEntry.Name))
 			{
-				ScriptObjectEntry.ObjectName = GlobalNameMap.FromSerializedNameMapEntry(ScriptObjectEntry.ObjectName);
+				// Not mounted yet
+				continue;
 			}
 
-			TArray<FGlobalImport>& Imports = ImportStore.Imports[FPackageObjectIndex::ScriptImport];
-			Imports.AddDefaulted(ImportStore.ScriptObjectEntries.Num());
+			const FName SourceName = MinimalNameToName(PackageEntry.Name);
+			FScopeLock Lock(&PackageNameToPackageIdCritical);
+			FPackageId* FoundPackageId = PackageNameToPackageId.Find(SourceName);
+
+			UE_CLOG(!FoundPackageId, LogStreaming, Warning,
+				TEXT("Skip remapping for localized package %s (%d) since the source package %s (%d) is unknown."),
+				It.Value().ToIndex(),
+				*FMappedName::SafeMinimalNameToName(StoreEntries[It.Value().ToIndex()].Name).ToString(),
+				It.Key().ToIndex(),
+				*SourceName.ToString());
+
+			if (FoundPackageId)
+			{
+				const FPackageId LocalizedPackageId = It.Value();
+				*FoundPackageId = LocalizedPackageId;
+
+#if ALT2_LOG_VERBOSE
+				UE_LOG(LogStreaming, Verbose, TEXT("LoadPackageStore: RemapLocalizedPackage: %s (%d) %s (%d)."),
+					*FMappedName::SafeMinimalNameToName(StoreEntries[It.Key().ToIndex()].Name).ToString(),
+					It.Key().ToIndex(),
+					*FMappedName::SafeMinimalNameToName(StoreEntries[It.Value().ToIndex()].Name).ToString(),
+					It.Value().ToIndex());
+#endif
+			}
 		}
 
-		FPackageName::DoesPackageExistOverride().BindLambda([this](FName InPackageName)
+		for (int I = 0; I < PackageCount; ++I)
 		{
-			return PackageNameToPackageId.Contains(InPackageName);
-		});
+			FPackageStoreEntry& StoreEntry = StoreEntries[I];
+
+			if (!FMappedName::IsResolvedToMinimalName(StoreEntry.Name))
+			{
+				// Not mounted yet
+				continue;
+			}
+
+			for (FPackageId& ImportedPackageIndex : StoreEntry.ImportedPackages)
+			{
+				if (FPackageId* Value = LocalizedPackageIdMap->Find(ImportedPackageIndex))
+				{
+					ImportedPackageIndex = *Value;
+				}
+			}
+		}
 	}
 
 	void FinalizeInitialLoad()
 	{
-		ImportStore.FindAllScriptImports();
+		ImportStore.FindAllScriptObjects();
 
-		UE_LOG(LogStreaming, Display, TEXT("AsyncLoading2 - InitialLoad Finalized: Script Imports: %d"),
-			ImportStore.Imports[FPackageObjectIndex::ScriptImport].Num());
+		UE_LOG(LogStreaming, Display, TEXT("AsyncLoading2 - InitialLoad Finalized: Script Objects: %d"),
+			ImportStore.ScriptObjects.Num());
 	}
 
 	inline FGlobalImportStore& GetGlobalImportStore()
@@ -747,6 +909,11 @@ public:
 	{
 		return StoreEntries[PackageId.ToIndex()];
 	}
+
+	inline const FNameMap& GetPackageNameMap(const uint16 NameMapIndex) const
+	{
+		return NameMapIndex > 0 ? *LoadedContainers[NameMapIndex].ContainerNameMap : GlobalNameMap;
+	}
 };
 
 struct FPackageImportStore
@@ -812,9 +979,9 @@ struct FPackageImportStore
 		return FName();
 	}
 
-	void StoreGlobalImportObject(FPackageObjectIndex GlobalIndex, UObject* InObject)
+	void StoreGlobalObject(FPackageObjectIndex GlobalIndex, UObject* InObject)
 	{
-		GlobalImportStore.StoreGlobalImportObject(GlobalIndex, InObject);
+		GlobalImportStore.StoreGlobalObject(GlobalIndex, InObject);
 	}
 
 	void ClearReferences()
@@ -990,8 +1157,8 @@ public:
 
 #if ALT2_LOG_VERBOSE
 				const FExportMapEntry& Export = ExportMap[ExportIndex];
-				FNameEntryId NameEntry = (*GlobalNameMap)[Export.ObjectName.NameEntryIndex];
-				FName ObjectName = FName::CreateFromDisplayId(NameEntry, Export.ObjectName.Number);
+				FNameEntryId NameEntry = (*NameMap)[Export.ObjectName.GetIndex()];
+				FName ObjectName = FName::CreateFromDisplayId(NameEntry, Export.ObjectName.GetNumber());
 				UE_ASYNC_PACKAGE_CLOG_VERBOSE(!Object, VeryVerbose, *PackageDesc,
 					TEXT("FExportArchive: Object"), TEXT("Export %s at index %d is null."),
 					*ObjectName.ToString(), 
@@ -1043,7 +1210,7 @@ public:
 	FORCENOINLINE void HandleBadNameIndex(int32 NameIndex, FName& Name)
 	{
 		UE_ASYNC_PACKAGE_LOG(Error, *PackageDesc, TEXT("HandleBadNameIndex"),
-			TEXT("Index: %d/%d"), NameIndex, GlobalNameMap->Num());
+			TEXT("Index: %d/%d"), NameIndex, NameMap->Num());
 
 		Name = FName();
 		SetCriticalError();
@@ -1059,10 +1226,10 @@ public:
 
 		NameIndex = PackageNameMap[NameIndex];
 
-		if (GlobalNameMap->IsValidIndex(NameIndex))
+		if (NameMap->IsValidIndex(NameIndex))
 		{
 			// if the name wasn't loaded (because it wasn't valid in this context)
-			FNameEntryId MappedName = (*GlobalNameMap)[NameIndex];
+			FNameEntryId MappedName = (*NameMap)[NameIndex];
 
 			// simply create the name from the NameMap's name and the serialized instance number
 			Name = FName::CreateFromDisplayId(MappedName, Number);
@@ -1084,7 +1251,7 @@ private:
 	FPackageImportStore* ImportStore = nullptr;
 	TArray<FExternalReadCallback>* ExternalReadDependencies;
 	const int32* PackageNameMap = nullptr;
-	const TArray<FNameEntryId>* GlobalNameMap = nullptr;
+	const TArray<FNameEntryId>* NameMap = nullptr;
 	const FExportObjects* Exports = nullptr;
 	const FExportMapEntry* ExportMap = nullptr;
 	int32 ExportCount = 0;
@@ -1535,6 +1702,7 @@ private:
 	const int32* PackageNameMap = nullptr;
 	FExportObjects Exports;
 	FPackageImportStore ImportStore;
+	const FNameMap* NameMap = nullptr;
 
 	int32 ExportBundleCount = 0;
 	const FExportBundleHeader* ExportBundles = nullptr;
@@ -1791,9 +1959,9 @@ private:
 	FThreadSafeCounter PackageRequestID;
 
 	/** I/O Dispatcher */
-	FGlobalNameMap GlobalNameMap;
 	FIoDispatcher& IoDispatcher;
 
+	FNameMap GlobalNameMap;
 	FPackageStore GlobalPackageStore;
 
 	struct FBundleIoRequest
@@ -2153,31 +2321,26 @@ struct FAsyncLoadingTickScope2
 
 void FAsyncLoadingThread2::InitializeLoading()
 {
-	FIoStoreEnvironment GlobalEnvironment;
-	GlobalEnvironment.InitializeFileEnvironment(FPaths::ProjectDir() / TEXT("global"));
-	FIoStatus MountStatus = IoDispatcher.Mount(GlobalEnvironment);
-	UE_CLOG(!MountStatus.IsOk(), LogInit, Fatal, TEXT("Failed to initialize I/O dispatcher: '%s'"), *MountStatus.ToString());
-
 #if USE_NEW_BULKDATA
 	FBulkDataBase::SetIoDispatcher(&IoDispatcher);
 #endif
 
 	{
 		TRACE_CPUPROFILER_EVENT_SCOPE(LoadGlobalNameMap);
-		GlobalNameMap.Load(IoDispatcher);
+		GlobalNameMap.LoadGlobal(IoDispatcher);
 	}
 
 	{
 		TRACE_CPUPROFILER_EVENT_SCOPE(LoadPackageStore);
-		GlobalPackageStore.Load(IoDispatcher, GlobalNameMap);
+		GlobalPackageStore.Load();
 	}
 
 	AsyncThreadReady.Increment();
 
-	UE_LOG(LogStreaming, Display, TEXT("AsyncLoading2 - Initialized: Packages: %d, PackageImports: %d, ScriptImports %d, FNames: %d"),
+	UE_LOG(LogStreaming, Display, TEXT("AsyncLoading2 - Initialized: Packages: %d, PublicExports: %d, ScriptObjects %d, FNames: %d"),
 		GlobalPackageStore.PackageCount,
-		GlobalPackageStore.ImportStore.Imports[FPackageObjectIndex::PackageImport].Num(),
-		GlobalPackageStore.ImportStore.Imports[FPackageObjectIndex::ScriptImport].Num(),
+		GlobalPackageStore.ImportStore.PublicExports.Num(),
+		GlobalPackageStore.ImportStore.ScriptObjects.Num(),
 		GlobalNameMap.GetNameEntries().Num());
 }
 
@@ -2722,27 +2885,25 @@ void FAsyncPackage2::SetupSerializedArcs(const uint8* GraphData, uint64 GraphDat
 }
 
 static UObject* GFindExistingScriptImport(int32 GlobalImportIndex,
-	TArray<FGlobalImport>& ScriptImports,
+	TArray<UObject*>& ScriptObjects,
 	const TArray<FScriptObjectEntry>& ScriptObjectEntries)
 {
-	FGlobalImport& ScriptImport = ScriptImports[GlobalImportIndex];
-	UObject* Object = ScriptImport.GetObject();
+	UObject*& Object = ScriptObjects[GlobalImportIndex];
 	if (!Object)
 	{
 		const FScriptObjectEntry& Entry = ScriptObjectEntries[GlobalImportIndex];
 		if (Entry.OuterIndex.IsNull())
 		{
-			Object = StaticFindObjectFast(UPackage::StaticClass(), nullptr, Entry.ObjectName.ToFName(), true);
+			Object = StaticFindObjectFast(UPackage::StaticClass(), nullptr, MinimalNameToName(Entry.ObjectName), true);
 		}
 		else
 		{
-			UObject* Outer = GFindExistingScriptImport(Entry.OuterIndex.GetIndex(), ScriptImports, ScriptObjectEntries);
+			UObject* Outer = GFindExistingScriptImport(Entry.OuterIndex.GetIndex(), ScriptObjects, ScriptObjectEntries);
 			if (Outer)
 			{
-				Object = StaticFindObjectFast(UObject::StaticClass(), Outer, Entry.ObjectName.ToFName(), false, true);
+				Object = StaticFindObjectFast(UObject::StaticClass(), Outer, MinimalNameToName(Entry.ObjectName), false, true);
 			}
 		}
-		ScriptImport.SetObject(Object);
 	}
 	return Object;
 }
@@ -2750,19 +2911,18 @@ static UObject* GFindExistingScriptImport(int32 GlobalImportIndex,
 UObject* FGlobalImportStore::FindScriptImportObjectFromIndex(int32 GlobalImportIndex)
 {
 	check(ScriptObjectEntries.Num() > 0);
-	return GFindExistingScriptImport(
-		GlobalImportIndex, Imports[FPackageObjectIndex::ScriptImport], ScriptObjectEntries);
+	return GFindExistingScriptImport(GlobalImportIndex, ScriptObjects, ScriptObjectEntries);
 }
 
-void FGlobalImportStore::FindAllScriptImports()
+void FGlobalImportStore::FindAllScriptObjects()
 {
-	TRACE_CPUPROFILER_EVENT_SCOPE(FindAllScriptImports);
+	TRACE_CPUPROFILER_EVENT_SCOPE(FindAllScriptObjects);
 	const int32 ScriptImportCount = ScriptObjectEntries.Num();
 	check(ScriptImportCount > 0);
 	for (int32 GlobalImportIndex = 0; GlobalImportIndex < ScriptImportCount; ++GlobalImportIndex)
 	{
 		UObject* Object = GFindExistingScriptImport(
-			GlobalImportIndex, Imports[FPackageObjectIndex::ScriptImport], ScriptObjectEntries);
+			GlobalImportIndex, ScriptObjects, ScriptObjectEntries);
 #if DO_CHECK
 		if (!Object)
 		{
@@ -2771,14 +2931,16 @@ void FGlobalImportStore::FindAllScriptImports()
 			{
 				UE_LOG(LogStreaming, Warning,
 					TEXT("AsyncLoading2 - Failed to find import script package after initial load: %s"),
-					*ScriptObjectEntries[GlobalImportIndex].ObjectName.ToFName().ToString());
+					*MinimalNameToName(ScriptObjectEntries[GlobalImportIndex].ObjectName).ToString());
 			}
 			else
 			{
 				UE_LOG(LogStreaming, Warning,
-					TEXT("AsyncLoading2 - Failed to find import script object after initial load: %s - %s"),
-					*ScriptObjectEntries[Entry.OuterIndex.GetIndex()].ObjectName.ToFName().ToString(),
-					*ScriptObjectEntries[GlobalImportIndex].ObjectName.ToFName().ToString());
+					TEXT("AsyncLoading2 - Failed to find import script object after initial load: %s (%d) in outer %s (%d)"),
+					*MinimalNameToName(ScriptObjectEntries[GlobalImportIndex].ObjectName).ToString(),
+					GlobalImportIndex,
+					*MinimalNameToName(ScriptObjectEntries[Entry.OuterIndex.GetIndex()].ObjectName).ToString(),
+					*MinimalNameToName(ScriptObjectEntries[GlobalImportIndex].ObjectName).ToString());
 			}
 		}
 #endif
@@ -2805,6 +2967,12 @@ void FAsyncPackage2::ImportPackagesRecursive()
 	{
 		// AreAllImportsInImportPackageAlreadyLoaded?
 		const FPackageStoreEntry& ImportEntry = GlobalPackageStore.GetStoreEntry(ImportedPackageId);
+#if DO_CHECK
+		UE_CLOG(!FMappedName::IsResolvedToMinimalName(ImportEntry.Name), LogStreaming, Warning,
+			TEXT("Package '%s' is trying to import non mounted package ID '%d'"),
+			*Desc.Name.ToString(),
+			ImportedPackageId.ToIndex());
+#endif
 		bool bNeedToLoadPackage = false;
 		for (const FPackageObjectIndex& GlobalImportIndex : ImportEntry.PublicExports)
 		{
@@ -2818,12 +2986,11 @@ void FAsyncPackage2::ImportPackagesRecursive()
 
 		if (bNeedToLoadPackage)
 		{
-			const FGlobalNameMap& GlobalNameMap = AsyncLoadingThread.GlobalNameMap;
 			const FPackageId PackageIdToLoad = ImportedPackageId;
-			const FName NameToLoad = ImportEntry.Name.ToFName();
+			const FName NameToLoad = MinimalNameToName(ImportEntry.Name);
 			const FName SourceName =
 				ImportEntry.SourcePackageId.IsValid() ?
-				GlobalPackageStore.StoreEntries[ImportEntry.SourcePackageId.ToIndex()].Name.ToFName() :
+				MinimalNameToName(GlobalPackageStore.StoreEntries[ImportEntry.SourcePackageId.ToIndex()].Name) :
 				NameToLoad;
 
 			FAsyncPackageDesc2 PackageDesc(INDEX_NONE, PackageIdToLoad, PackageIdToLoad, SourceName, NameToLoad);
@@ -2923,6 +3090,7 @@ EAsyncPackageState::Type FAsyncPackage2::Event_ProcessExportBundle(FAsyncPackage
 			const uint64 PackageSummarySize = GraphData + PackageSummary->GraphDataSize - PackageSummaryData;
 
 			Package->CookedHeaderSize = PackageSummary->CookedHeaderSize;
+			Package->NameMap = &Package->AsyncLoadingThread.GlobalPackageStore.GetPackageNameMap(PackageSummary->NameMapIndex);
 			Package->PackageNameMap = reinterpret_cast<const int32*>(PackageSummaryData + PackageSummary->NameMapOffset);
 			Package->ImportStore.ImportMap = reinterpret_cast<const FPackageObjectIndex*>(PackageSummaryData + PackageSummary->ImportMapOffset);
 			Package->ImportStore.ImportMapCount = (PackageSummary->ExportMapOffset - PackageSummary->ImportMapOffset) / sizeof(int32);
@@ -3115,8 +3283,7 @@ void FAsyncPackage2::EventDrivenCreateExport(int32 LocalExportIndex)
 	FName ObjectName;
 	{
 		TRACE_CPUPROFILER_EVENT_SCOPE(ObjectNameFixup);
-		const FGlobalNameMap& GlobalNameMap = AsyncLoadingThread.GlobalNameMap;
-		ObjectName = GlobalNameMap.GetName(Export.ObjectName);
+		ObjectName = NameMap->GetName(Export.ObjectName);
 	}
 
 	TRACE_LOADTIME_CREATE_EXPORT_SCOPE(this, &Object);
@@ -3272,7 +3439,7 @@ void FAsyncPackage2::EventDrivenCreateExport(int32 LocalExportIndex)
 	}
 
 	check(Object);
-	ImportStore.StoreGlobalImportObject(Export.GlobalImportIndex, Object);
+	ImportStore.StoreGlobalObject(Export.GlobalImportIndex, Object);
 }
 
 void FAsyncPackage2::EventDrivenSerializeExport(int32 LocalExportIndex, FExportArchive& Ar)
@@ -3811,6 +3978,7 @@ EAsyncPackageState::Type FAsyncLoadingThread2::TickAsyncLoadingFromGameThread(bo
 FAsyncLoadingThread2::FAsyncLoadingThread2(FIoDispatcher& InIoDispatcher)
 	: Thread(nullptr)
 	, IoDispatcher(InIoDispatcher)
+	, GlobalPackageStore(InIoDispatcher, GlobalNameMap)
 {
 	GEventDrivenLoaderEnabled = true;
 
@@ -4250,7 +4418,7 @@ void FGlobalImportStore::OnPreGarbageCollect(bool bInIsLoadingPackages)
 	bNeedToHandleGarbageCollect = bInIsLoadingPackages;
 
 	int32 NumWeak = 0;
-	for (FGlobalImport& GlobalImport : Imports[FPackageObjectIndex::PackageImport])
+	for (FGlobalImport& GlobalImport : PublicExports)
 	{
 		UObject* Object = GlobalImport.GetObjectIfRawPointer();
 		if (!Object)
@@ -4283,10 +4451,10 @@ void FGlobalImportStore::OnPreGarbageCollect(bool bInIsLoadingPackages)
 		check(KeepAliveObjects.Num() == 0);
 	}
 
-#ifdef ALT2_VERIFY_ASYNC_FLAGS
+#if ALT2_VERIFY_ASYNC_FLAGS
 	if (!bInIsLoadingPackages)
 	{
-		for (FGlobalImport& GlobalImport : Imports[FPackageObjectIndex::PackageImport])
+		for (FGlobalImport& GlobalImport : PublicExports)
 		{
 			check(GlobalImport.GetObjectIfRawPointer() == nullptr);
 			check(GlobalImport.GetRefCount() == 0);
diff --git a/Engine/Source/Runtime/CoreUObject/Private/Serialization/AsyncPackageLoader.cpp b/Engine/Source/Runtime/CoreUObject/Private/Serialization/AsyncPackageLoader.cpp
--- a/Engine/Source/Runtime/CoreUObject/Private/Serialization/AsyncPackageLoader.cpp
+++ b/Engine/Source/Runtime/CoreUObject/Private/Serialization/AsyncPackageLoader.cpp
@@ -1,6 +1,7 @@
 // Copyright Epic Games, Inc. All Rights Reserved.
 
 #include "Serialization/AsyncPackageLoader.h"
+#include "HAL/IConsoleManager.h"
 #include "Serialization/AsyncLoadingThread.h"
 #include "Serialization/AsyncLoading2.h"
 #include "UObject/GCObject.h"
diff --git a/Engine/Source/Runtime/CoreUObject/Private/Serialization/BulkData2.cpp b/Engine/Source/Runtime/CoreUObject/Private/Serialization/BulkData2.cpp
--- a/Engine/Source/Runtime/CoreUObject/Private/Serialization/BulkData2.cpp
+++ b/Engine/Source/Runtime/CoreUObject/Private/Serialization/BulkData2.cpp
@@ -162,7 +162,7 @@ namespace FileTokenSystem
 			StringTable.Remove(TokenData[ID].PackageName);
 			TokenData.RemoveAt(ID);
 
-			check(StringTable.Num() <= TokenData.Num());
+			checkf(StringTable.Num() <= TokenData.Num(), TEXT("FStringTable has more strings (%d) than file entries (%d)"), StringTable.Num(), TokenData.Num());
 		}
 	}
 
@@ -332,7 +332,7 @@ protected:
 			FScopeLock Lock(&FReadChunkIdRequestEvent);
 			if (bRequestOutstanding)
 			{
-				check(DoneEvent == nullptr);
+				checkf(DoneEvent == nullptr, TEXT("Multiple threads attempting to wait on the same FReadChunkIdRequest"));
 				DoneEvent = FPlatformProcess::GetSynchEventFromPool(true);
 			}
 		}
@@ -348,8 +348,8 @@ protected:
 		}
 
 		// Make sure everything has completed
-		check(bRequestOutstanding == false);
-		check(PollCompletion() == true);
+		checkf(bRequestOutstanding == false, TEXT("Request has not completed by the end of WaitCompletionImpl"));
+		checkf(PollCompletion() == true, TEXT("Request and callback has not completed by the end of WaitCompletionImpl"));
 	}
 
 	virtual void CancelImpl() override
@@ -444,13 +444,14 @@ public:
 			DataResult = nullptr;
 		}
 
-		check(!IoBatch.IsValid()); // Should be freed by the callback!
+		// Should be freed by the callback!
+		checkf(!IoBatch.IsValid(), TEXT("FBulkDataIoDispatcherRequest::IoBatch was not freed"));
 	}
 
 	void StartAsyncWork()
 	{		
-		check(RequestArray.Num() > 0);
-		check(!IoBatch.IsValid());
+		checkf(RequestArray.Num() > 0, TEXT("RequestArray cannot be empty"));
+		checkf(!IoBatch.IsValid(), TEXT("FBulkDataIoDispatcherRequest::StartAsyncWork was called twice"));
 
 		auto Callback = [this](TIoStatusOr<FIoBuffer> Result)
 		{
@@ -468,7 +469,9 @@ public:
 				DataResult = IoBuffer.Data();
 			}
 
-			check(DataResult);
+			checkf(DataResult != nullptr, TEXT("Nothing was loaded!"));
+
+			FBulkDataBase::GetIoDispatcher()->FreeBatch(IoBatch);
 
 			bIsCompleted = true;
 
@@ -477,9 +480,7 @@ public:
 			if (CompleteCallback)
 			{
 				CompleteCallback(bIsCanceled, this);
-			}
-
-			FBulkDataBase::GetIoDispatcher()->FreeBatch(IoBatch);
+			}	
 		};
 
 		IoBatch = FBulkDataBase::GetIoDispatcher()->NewBatch();
@@ -591,7 +592,7 @@ FBulkDataBase::FBulkDataBase(FBulkDataBase&& Other)
 	, BulkDataFlags(Other.BulkDataFlags)
 
 {
-	check(Other.LockStatus == LOCKSTATUS_Unlocked); // Make sure that the other object wasn't inuse
+	checkf(Other.LockStatus != LOCKSTATUS_ReadWriteLock, TEXT("Attempting to read from a BulkData object that is locked for write")); 
 
 	if (!Other.IsUsingIODispatcher())
 	{
@@ -603,8 +604,8 @@ FBulkDataBase& FBulkDataBase::operator=(const FBulkDataBase& Other)
 {
 	DECLARE_SCOPE_CYCLE_COUNTER(TEXT("FBulkDataBase::operator="), STAT_UBD_Constructor, STATGROUP_Memory);
 
-	check(LockStatus == LOCKSTATUS_Unlocked);
-	check(Other.LockStatus == LOCKSTATUS_Unlocked);
+	checkf(LockStatus == LOCKSTATUS_Unlocked, TEXT("Attempting to modify a BulkData object that is locked"));
+	checkf(Other.LockStatus != LOCKSTATUS_ReadWriteLock, TEXT("Attempting to read from a BulkData object that is locked for write"));
 
 	RemoveBulkData();
 
@@ -637,7 +638,7 @@ FBulkDataBase& FBulkDataBase::operator=(const FBulkDataBase& Other)
 	}
 	else
 	{
-		// Note we don't need a fallback since thge original already managed the load, if we fail now then it
+		// Note we don't need a fallback since the original already managed the load, if we fail now then it
 		// is an actual error.
 		if (Other.IsUsingIODispatcher())
 		{
@@ -661,7 +662,7 @@ FBulkDataBase::~FBulkDataBase()
 {
 	FlushAsyncLoading();
 	
-	check(LockStatus == LOCKSTATUS_Unlocked);
+	checkf(LockStatus == LOCKSTATUS_Unlocked, TEXT("Attempting to modify a BulkData object that is locked"));
 
 	FreeData();
 	if (!IsUsingIODispatcher())
@@ -678,10 +679,10 @@ void FBulkDataBase::Serialize(FArchive& Ar, UObject* Owner, int32 /*Index*/, boo
 
 #if WITH_EDITOR == 0 && WITH_EDITORONLY_DATA == 0
 	if (Ar.IsPersistent() && !Ar.IsObjectReferenceCollector() && !Ar.ShouldSkipBulkData())
-	{
-		check(Ar.IsLoading());				// Only support loading from cooked data!
-		check(!GIsEditor);					// The editor path is not supported
-		check(LockStatus == LOCKSTATUS_Unlocked);
+	{		
+		checkf(Ar.IsLoading(), TEXT("FBulkDataBase only works with loading")); // Only support loading from cooked data!
+		checkf(!GIsEditor, TEXT("FBulkDataBase does not work in the editor"));
+		checkf(LockStatus == LOCKSTATUS_Unlocked, TEXT("Attempting to modify a BulkData object that is locked"));
 
 		Ar << BulkDataFlags;
 
@@ -721,9 +722,9 @@ void FBulkDataBase::Serialize(FArchive& Ar, UObject* Owner, int32 /*Index*/, boo
 		// Assuming that Owner/Package/Linker are all valid, the old BulkData system would
 		// generally fail if any of these were nullptr but had plenty of inconsistent checks
 		// scattered throughout.
-		check(Owner != nullptr);
+		checkf(Owner != nullptr, TEXT("FBulkDataBase::Serialize requires a valid Owner"));
 		const UPackage* Package = Owner->GetOutermost();
-		check(Package != nullptr);
+		checkf(Package != nullptr, TEXT("FBulkDataBase::Serialize requires n Owner that returns a valid UPackage"));
 
 		if (!IsInlined() && bUseIoDispatcher)
 		{
@@ -799,7 +800,7 @@ void FBulkDataBase::Serialize(FArchive& Ar, UObject* Owner, int32 /*Index*/, boo
 				}
 				else
 				{
-					check(Filename != nullptr);
+					checkf(Filename != nullptr, TEXT("Could not find Filename"));
 					FString MemoryMappedFilename = ConvertFilenameFromFlags(*Filename);
 					if (!MemoryMapBulkData(MemoryMappedFilename, BulkDataOffsetInFile, BulkDataSize))
 					{
@@ -835,13 +836,13 @@ void FBulkDataBase::Serialize(FArchive& Ar, UObject* Owner, int32 /*Index*/, boo
 		}
 	}
 #else
-	check(false); // Only implemented for cooked builds!
+	checkf(false, TEXT("FBulkDataBase does not work in the editor")); // Only implemented for cooked builds!
 #endif
 }
 
 void* FBulkDataBase::Lock(uint32 LockFlags)
 {
-	check(LockStatus == LOCKSTATUS_Unlocked);
+	checkf(LockStatus == LOCKSTATUS_Unlocked, TEXT("Attempting to lock a BulkData object that is already locked"));
 	
 	ForceBulkDataResident(); 	// If nothing is currently loaded then load from disk
 
@@ -865,7 +866,7 @@ void* FBulkDataBase::Lock(uint32 LockFlags)
 
 const void* FBulkDataBase::LockReadOnly() const
 {
-	check(LockStatus == LOCKSTATUS_Unlocked);
+	checkf(LockStatus == LOCKSTATUS_Unlocked, TEXT("Attempting to lock a BulkData object that is already locked"));
 	LockStatus = LOCKSTATUS_ReadOnlyLock;
 
 	return GetDataBufferReadOnly();
@@ -873,8 +874,7 @@ const void* FBulkDataBase::LockReadOnly() const
 
 void FBulkDataBase::Unlock()
 {
-	check(LockStatus != LOCKSTATUS_Unlocked);
-
+	checkf(LockStatus != LOCKSTATUS_Unlocked, TEXT("Attempting to unlock a BulkData object that is not locked"));
 	LockStatus = LOCKSTATUS_Unlocked;
 
 	// Free pointer if we're guaranteed to only to access the data once.
@@ -893,14 +893,13 @@ void* FBulkDataBase::Realloc(int64 SizeInBytes)
 {
 	DECLARE_SCOPE_CYCLE_COUNTER(TEXT("FBulkDataBase::Realloc"), STAT_UBD_Realloc, STATGROUP_Memory);
 
-	check(LockStatus == LOCKSTATUS_ReadWriteLock);
+	checkf(LockStatus == LOCKSTATUS_Unlocked, TEXT("Attempting to modify a BulkData object that is locked"));
 	checkf(!CanLoadFromDisk(), TEXT("Cannot re-allocate a FBulkDataBase object that represents a file on disk!"));
 
-	AllocateData(SizeInBytes);
+	// We might want to consider this a valid use case if anyone can come up with one?
+	checkf(!IsUsingIODispatcher(), TEXT("Attempting to re-allocate data loaded from the IoDispatcher"));
 
-	check(!IsUsingIODispatcher());	// This case should get caught above but if someone tries to change that in the
-									// future then this check is a reminder that we need to handle the IoDispatcher 
-									// vs fallback case.
+	AllocateData(SizeInBytes);
 
 	Data.Fallback.BulkDataSize = SizeInBytes;
 
@@ -911,8 +910,9 @@ void FBulkDataBase::GetCopy(void** DstBuffer, bool bDiscardInternalCopy)
 {
 	DECLARE_SCOPE_CYCLE_COUNTER(TEXT("FBulkDataBase::GetCopy"), STAT_UBD_GetCopy, STATGROUP_Memory);
 
-	check(LockStatus == LOCKSTATUS_Unlocked);
-	check(DstBuffer);
+	checkf(LockStatus == LOCKSTATUS_Unlocked, TEXT("Attempting to modify a BulkData object that is locked"));
+	checkf(DstBuffer, TEXT("FBulkDataBase::GetCopy requires a valid DstBuffer"));	// Really should just use a reference to a pointer 
+																					// but we are stuck trying to stick with the old bulkdata API for now
 
 	// Wait for anything that might be currently loading
 	FlushAsyncLoading();
@@ -963,42 +963,35 @@ void FBulkDataBase::GetCopy(void** DstBuffer, bool bDiscardInternalCopy)
 
 void FBulkDataBase::SetBulkDataFlags(uint32 BulkDataFlagsToSet)
 {
-	check(!CanLoadFromDisk());	// We only want to allow the editing of flags if the BulkData
-								// was dynamically created at runtime, not loaded off disk
-
 	BulkDataFlags = EBulkDataFlags(BulkDataFlags | BulkDataFlagsToSet);
 }
 
 void FBulkDataBase::ResetBulkDataFlags(uint32 BulkDataFlagsToSet)
 {
-	check(!CanLoadFromDisk());	// We only want to allow the editing of flags if the BulkData
-								// was dynamically created at runtime, not loaded off disk
-
 	BulkDataFlags = (EBulkDataFlags)BulkDataFlagsToSet;
 }
 
 void FBulkDataBase::ClearBulkDataFlags(uint32 BulkDataFlagsToClear)
 { 
-	check(!CanLoadFromDisk());	// We only want to allow the editing of flags if the BulkData
-								// was dynamically created at runtime, not loaded off disk
-
 	BulkDataFlags = EBulkDataFlags(BulkDataFlags & ~BulkDataFlagsToClear);
 }
 
 void FBulkDataBase::SetRuntimeBulkDataFlags(uint32 BulkDataFlagsToSet)
 {
-	check(	BulkDataFlagsToSet == BULKDATA_UsesIoDispatcher || 
+	checkf(	BulkDataFlagsToSet == BULKDATA_UsesIoDispatcher || 
 			BulkDataFlagsToSet == BULKDATA_DataIsMemoryMapped || 
-			BulkDataFlagsToSet == BULKDATA_HasAsyncReadPending);
+			BulkDataFlagsToSet == BULKDATA_HasAsyncReadPending,
+			TEXT("Attempting to set an invalid runtime flag"));
 			
 	BulkDataFlags = EBulkDataFlags(BulkDataFlags | BulkDataFlagsToSet);
 }
 
 void FBulkDataBase::ClearRuntimeBulkDataFlags(uint32 BulkDataFlagsToClear)
 {
-	check(	BulkDataFlagsToClear == BULKDATA_UsesIoDispatcher ||
+	checkf(	BulkDataFlagsToClear == BULKDATA_UsesIoDispatcher ||
 			BulkDataFlagsToClear == BULKDATA_DataIsMemoryMapped ||
-			BulkDataFlagsToClear == BULKDATA_HasAsyncReadPending);
+			BulkDataFlagsToClear == BULKDATA_HasAsyncReadPending,
+			TEXT("Attempting to clear an invalid runtime flag"));
 
 	BulkDataFlags = EBulkDataFlags(BulkDataFlags & ~BulkDataFlagsToClear);
 }
@@ -1140,14 +1133,14 @@ IBulkDataIORequest* FBulkDataBase::CreateStreamingRequest(int64 OffsetInBulkData
 		const int64 BulkDataSize = GetBulkDataSize();
 		FileTokenSystem::Data FileData = FileTokenSystem::GetFileData(Data.Fallback.Token);
 
-		check(FileData.PackageHeaderFilename.IsEmpty() == false);
+		checkf(FileData.PackageHeaderFilename.IsEmpty() == false, TEXT("BulkData file name is missing"));
 		const FString Filename = ConvertFilenameFromFlags(FileData.PackageHeaderFilename);
 
 		UE_CLOG(IsStoredCompressedOnDisk(), LogSerialization, Fatal, TEXT("Package level compression is no longer supported (%s)."), *Filename);
 		UE_CLOG(BulkDataSize <= 0, LogSerialization, Error, TEXT("(%s) has invalid bulk data size."), *Filename);
 
 		IAsyncReadFileHandle* IORequestHandle = FPlatformFileManager::Get().GetPlatformFile().OpenAsyncRead(*Filename);
-		check(IORequestHandle); // this generally cannot fail because it is async
+		checkf(IORequestHandle, TEXT("OpenAsyncRead failed")); // An assert as there shouldn't be a way for this to fail
 
 		if (IORequestHandle == nullptr)
 		{
@@ -1172,11 +1165,11 @@ IBulkDataIORequest* FBulkDataBase::CreateStreamingRequest(int64 OffsetInBulkData
 
 IBulkDataIORequest* FBulkDataBase::CreateStreamingRequestForRange(const BulkDataRangeArray& RangeArray, EAsyncIOPriorityAndFlags Priority, FBulkDataIORequestCallBack* CompleteCallback)
 {
-	check(RangeArray.Num() > 0);
+	checkf(RangeArray.Num() > 0, TEXT("RangeArray cannot be empty"));
 
 	const FBulkDataBase& Start = *(RangeArray[0]);
 
-	check(!Start.IsInlined());
+	checkf(!Start.IsInlined(), TEXT("Cannot stream inlined BulkData"));
 
 	if (Start.IsUsingIODispatcher())
 	{
@@ -1195,14 +1188,14 @@ IBulkDataIORequest* FBulkDataBase::CreateStreamingRequestForRange(const BulkData
 	{	
 		const FBulkDataBase& End = *RangeArray[RangeArray.Num() - 1];
 
-		check(Start.GetFilename() == End.GetFilename());
+		checkf(Start.GetFilename() == End.GetFilename(), TEXT("BulkData range does not come from the same file (%s vs %s)"), *Start.GetFilename(), *End.GetFilename());
 
 		const int64 ReadOffset = Start.GetBulkDataOffsetInFile();
-		const int64 ReadSize = (End.GetBulkDataOffsetInFile() + End.GetBulkDataSize()) - ReadOffset;
+		const int64 ReadLength = (End.GetBulkDataOffsetInFile() + End.GetBulkDataSize()) - ReadOffset;
 
-		check(ReadSize > 0);
+		checkf(ReadLength > 0, TEXT("Read length is 0"));
 
-		return Start.CreateStreamingRequest(0, ReadSize, Priority, CompleteCallback, nullptr);
+		return Start.CreateStreamingRequest(0, ReadLength, Priority, CompleteCallback, nullptr);
 	}
 }
 
@@ -1223,14 +1216,14 @@ void FBulkDataBase::ForceBulkDataResident()
 
 FOwnedBulkDataPtr* FBulkDataBase::StealFileMapping()
 {
-	check(LockStatus == LOCKSTATUS_Unlocked);
+	checkf(LockStatus == LOCKSTATUS_Unlocked, TEXT("Attempting to modify a BulkData object that is locked"));
 
 	return DataAllocation.StealFileMapping(this);
 }
 
 void FBulkDataBase::RemoveBulkData()
 {
-	check(LockStatus == LOCKSTATUS_Unlocked);
+	checkf(LockStatus == LOCKSTATUS_Unlocked, TEXT("Attempting to modify a BulkData object that is locked"));
 
 	FreeData();
 
@@ -1263,7 +1256,7 @@ bool FBulkDataBase::StartAsyncLoading()
 		return false; // Early out if we cannot load from disk
 	}
 
-	check(LockStatus == LOCKSTATUS_Unlocked)
+	checkf(LockStatus == LOCKSTATUS_Unlocked, TEXT("Attempting to modify a BulkData object that is locked"));
 
 	LockStatus = LOCKSTATUS_ReadWriteLock; // Bulkdata is effectively locked while streaming!
 
@@ -1273,10 +1266,12 @@ bool FBulkDataBase::StartAsyncLoading()
 
 	AsyncCallback Callback = [this](TIoStatusOr<FIoBuffer> Result)
 	{
-		check(Result.IsOk());
+		CHECK_IOSTATUS(Result.Status(), TEXT("FBulkDataBase::StartAsyncLoading"));
 		FIoBuffer IoBuffer = Result.ConsumeValueOrDie();
 
-		check(!IoBuffer.IsMemoryOwned());
+		// It is assumed that ::LoadDataAsynchronously will allocate memory for the loaded data, so that we
+		// do not need to take ownership here. This should guard against any future change in functionality.
+		checkf(!IoBuffer.IsMemoryOwned(), TEXT("The loaded data is not owned by the BulkData object"));
 		DataAllocation.SetData(this, IoBuffer.Data());
 
 		FPlatformMisc::MemoryBarrier();
@@ -1333,7 +1328,7 @@ bool FBulkDataBase::CanDiscardInternalData() const
 	//		since we will not be able to reload inline data when the IoStore is
 	//		active.
 
-	// TODO: This is currently called from ::GetCopy but not ::Unlock, we should investe unifying the
+	// TODO: This is currently called from ::GetCopy but not ::Unlock, we should investigate unifying the
 	// rules for discarding data
 	return CanLoadFromDisk() || IsSingleUse() || (IsInlined() && IsIoDispatcherEnabled());
 }
@@ -1454,8 +1449,8 @@ void FBulkDataBase::InternalLoadFromIoStore(void** DstBuffer)
 	NewBatch.Issue();
 	NewBatch.Wait(); // Blocking wait until all requests in the batch are done
 
-	check(Request.IsOk());
-
+	CHECK_IOSTATUS(Request.Status(), TEXT("FIoRequest"));
+	
 	FBulkDataBase::GetIoDispatcher()->FreeBatch(NewBatch);
 }
 
@@ -1500,7 +1495,7 @@ void FBulkDataBase::ProcessDuplicateData(FArchive& Ar, const UPackage* Package,
 	}
 	else
 	{
-		check(Filename != nullptr);
+		checkf(Filename != nullptr, TEXT("If IoDispatcher is not used then ProcessDuplicateData requires a valid Filename pointer"));
 		const FString OptionalDataFilename = FPathViews::ChangeExtension(*Filename, BulkDataExt::Optional);
 
 		if (IFileManager::Get().FileExists(*OptionalDataFilename))
@@ -1544,7 +1539,7 @@ void FBulkDataBase::SerializeDuplicateData(FArchive& Ar, EBulkDataFlags& OutBulk
 
 void FBulkDataBase::SerializeBulkData(FArchive& Ar, void* DstBuffer, int64 DataLength)
 {	
-	check(Ar.IsLoading()); // Currently only support loading
+	checkf(Ar.IsLoading(), TEXT("BulkData2 only supports serialization for loading"));
 
 	if (IsAvailableForUse()) // skip serializing of unused data
 	{
@@ -1557,7 +1552,7 @@ void FBulkDataBase::SerializeBulkData(FArchive& Ar, void* DstBuffer, int64 DataL
 		return;
 	}
 
-	check(DstBuffer != nullptr);
+	checkf(DstBuffer != nullptr, TEXT("No destination buffer was provided for serialization"));
 
 	if (IsStoredCompressedOnDisk())
 	{
@@ -1572,7 +1567,7 @@ void FBulkDataBase::SerializeBulkData(FArchive& Ar, void* DstBuffer, int64 DataL
 
 bool FBulkDataBase::MemoryMapBulkData(const FString& Filename, int64 OffsetInBulkData, int64 BytesToRead)
 {
-	check(!IsBulkDataLoaded());
+	checkf(!IsBulkDataLoaded(), TEXT("Attempting to memory map BulkData that is already loaded"));
 
 	IMappedFileHandle* MappedHandle = nullptr;
 	IMappedFileRegion* MappedRegion = nullptr;
diff --git a/Engine/Source/Runtime/CoreUObject/Public/Serialization/AsyncLoading2.h b/Engine/Source/Runtime/CoreUObject/Public/Serialization/AsyncLoading2.h
--- a/Engine/Source/Runtime/CoreUObject/Public/Serialization/AsyncLoading2.h
+++ b/Engine/Source/Runtime/CoreUObject/Public/Serialization/AsyncLoading2.h
@@ -10,6 +10,7 @@
 #include "UObject/ObjectResource.h"
 #include "UObject/PackageId.h"
 #include "Serialization/Archive.h"
+#include "IO/IoContainerId.h"
 
 class FArchive;
 class IAsyncPackageLoader;
@@ -19,29 +20,92 @@ class IEDLBootNotificationManager;
 using FSourceToLocalizedPackageIdMap = TMap<FPackageId, FPackageId>;
 using FCulturePackageMap = TMap<FString, FSourceToLocalizedPackageIdMap>;
 
-struct FNameMapEntry
+class FMappedName
 {
-	uint32 Index;
-	uint32 Number;
-#if WITH_EDITOR
-	FName ToFName() const
+	static constexpr uint32 InvalidIndex = ~uint32(0);
+	static constexpr uint32 IndexBits = 31u;
+	static constexpr uint32 IndexMask = (1u << IndexBits) - 1u;
+	static constexpr uint32 TypeMask = ~IndexMask;
+	static constexpr uint32 TypeShift = IndexBits;
+
+public:
+	enum class EType { Container, Global };
+
+	inline FMappedName() = default;
+
+	static inline FMappedName Create(const uint32 InIndex, const uint32 InNumber, EType InType)
 	{
-		return FName::CreateFromDisplayId(FNameEntryId::FromUnstableInt(Index), Number);
+		check(InIndex <= MAX_int32);
+		return FMappedName((uint32(InType) << TypeShift) | InIndex, InNumber);
 	}
-#else
-	FName ToFName() const
+
+	static inline FMappedName FromMinimalName(const FMinimalName& MinimalName)
 	{
-		return *reinterpret_cast<const FName*>(this);
+		return *reinterpret_cast<const FMappedName*>(&MinimalName);
+	}
+
+	static inline bool IsResolvedToMinimalName(const FMinimalName& MinimalName)
+	{
+		// Not completely safe, relies on that no FName will have its Index and Number equal to Max_uint32
+		const FMappedName MappedName = FromMinimalName(MinimalName);
+		return MappedName.IsValid();
+	}
+
+	static inline FName SafeMinimalNameToName(const FMinimalName& MinimalName)
+	{
+		return IsResolvedToMinimalName(MinimalName) ? MinimalNameToName(MinimalName) : NAME_None;
+	}
+
+	inline FMinimalName ToUnresolvedMinimalName() const
+	{
+		return *reinterpret_cast<const FMinimalName*>(this);
+	}
+
+	inline bool IsValid() const
+	{
+		return Index != InvalidIndex && Number != InvalidIndex;
+	}
+
+	inline EType GetType() const
+	{
+		return static_cast<EType>(uint32((Index & TypeMask) >> TypeShift));
+	}
+
+	inline bool IsGlobal() const
+	{
+		return ((Index & TypeMask) >> TypeShift) != 0;
+	}
+
+	inline uint32 GetIndex() const
+	{
+		return Index & IndexMask;
+	}
+
+	inline uint32 GetNumber() const
+	{
+		return Number;
 	}
-#endif
 
-	COREUOBJECT_API friend FArchive& operator<<(FArchive& Ar, FNameMapEntry& NameMapEntry);
+	COREUOBJECT_API friend FArchive& operator<<(FArchive& Ar, FMappedName& MappedName);
+
+private:
+	inline FMappedName(const uint32 InIndex, const uint32 InNumber)
+		: Index(InIndex)
+		, Number(InNumber) { }
+
+	uint32 Index = InvalidIndex;
+	uint32 Number = InvalidIndex;
 };
 
-struct FNameMapIndex
+struct FContainerHeader
 {
-	uint32 NameEntryIndex;
-	uint32 Number;
+	FIoContainerId ContainerId;
+	TArray<uint8> Names;
+	TArray<uint8> NameHashes;
+	TArray<FPackageId> PackageIds;
+	TArray<FMappedName> PackageNames;
+
+	COREUOBJECT_API friend FArchive& operator<<(FArchive& Ar, FContainerHeader& ContainerHeader);
 };
 
 class FPackageObjectIndex
@@ -174,6 +238,8 @@ struct FPackageSummary
 {
 	uint32 PackageFlags;
 	uint32 CookedHeaderSize;
+	uint16 NameMapIndex;
+	uint16 Pad;
 	int32 NameMapOffset;
 	int32 ImportMapOffset;
 	int32 ExportMapOffset;
@@ -231,10 +297,9 @@ public:
 	inline T& operator[](uint32 Index)				{ return Data()[Index]; }
 };
 
-
 struct FPackageStoreEntry
 {
-	FNameMapEntry Name;
+	FMinimalName Name;
 	FPackageId SourcePackageId;
 	int32 ExportCount;
 	TPackageStoreEntryCArrayView<FPackageId> ImportedPackages;
@@ -255,7 +320,7 @@ struct FExportBundleHeader
 
 struct FScriptObjectEntry
 {
-	FNameMapEntry ObjectName;
+	FMinimalName ObjectName;
 	FPackageObjectIndex OuterIndex;
 	FPackageObjectIndex CDOClassIndex;
 
@@ -269,7 +334,7 @@ struct FExportMapEntry
 {
 	uint64 CookedSerialOffset;
 	uint64 CookedSerialSize;
-	FNameMapIndex ObjectName;
+	FMappedName ObjectName;
 	FPackageObjectIndex OuterIndex;
 	FPackageObjectIndex ClassIndex;
 	FPackageObjectIndex SuperIndex;
diff --git a/Engine/Source/Runtime/CoreUObject/Public/Serialization/BulkData2.h b/Engine/Source/Runtime/CoreUObject/Public/Serialization/BulkData2.h
--- a/Engine/Source/Runtime/CoreUObject/Public/Serialization/BulkData2.h
+++ b/Engine/Source/Runtime/CoreUObject/Public/Serialization/BulkData2.h
@@ -197,7 +197,7 @@ public:
 	void RemoveBulkData();
 
 	/**
-	* Initiates a new asynchronous operation to load the dulkdata from disk assuming that it is not already
+	* Initiates a new asynchronous operation to load the bulkdata from disk assuming that it is not already
 	* loaded.
 	* Note that a new asynchronous loading operation will not be created if one is already in progress.
 	*
diff --git a/Engine/Source/Runtime/CoreUObject/Public/Serialization/BulkDataBuffer.h b/Engine/Source/Runtime/CoreUObject/Public/Serialization/BulkDataBuffer.h
--- a/Engine/Source/Runtime/CoreUObject/Public/Serialization/BulkDataBuffer.h
+++ b/Engine/Source/Runtime/CoreUObject/Public/Serialization/BulkDataBuffer.h
@@ -11,11 +11,11 @@
  * This object assumes that it owns the buffer that it has been given and will free the
  * memory when the object is destroyed.
  */
-template<typename TYPE>
+template<typename DataType>
 class FBulkDataBuffer
 {
 public:
-	using ViewType = TArrayView<TYPE>;
+	using ViewType = TArrayView<DataType>;
 
 	/** Constructor. */
 	FBulkDataBuffer() = default;
@@ -47,7 +47,7 @@ public:
 	* @param InBuffer Pointer to memory to take ownership of, must've have been allocated via FMemory::Malloc/Realloc.
 	* @param InNumberOfElements The number of elements in the buffer.
 	*/
-	FBulkDataBuffer(TYPE* InBuffer, uint64 InNumberOfElements)
+	FBulkDataBuffer(DataType* InBuffer, uint64 InNumberOfElements)
 		: View(InBuffer, InNumberOfElements)
 	{
 		check(InNumberOfElements <= TNumericLimits<uint32>::Max());
@@ -74,7 +74,7 @@ public:
 		{
 			const int32 BufferSize = Other.View.Num();
 
-			TYPE* BufferCopy = (TYPE*)FMemory::Malloc(BufferSize);
+			DataType* BufferCopy = (DataType*)FMemory::Malloc(BufferSize);
 			FMemory::Memcpy(BufferCopy, Other.View.GetData(), BufferSize);
 
 			View = ViewType(BufferCopy, BufferSize);
@@ -119,7 +119,7 @@ public:
 	*  @param InNumberOfElements The number of elements in the buffer.
 	 * buffers actual size if desired.
 	 */
-	void Reset(TYPE* InBuffer, uint64 InNumberOfElements)
+	void Reset(DataType* InBuffer, uint64 InNumberOfElements)
 	{
 		FreeBuffer();
 
diff --git a/Engine/Source/Runtime/Launch/Private/LaunchEngineLoop.cpp b/Engine/Source/Runtime/Launch/Private/LaunchEngineLoop.cpp
--- a/Engine/Source/Runtime/Launch/Private/LaunchEngineLoop.cpp
+++ b/Engine/Source/Runtime/Launch/Private/LaunchEngineLoop.cpp
@@ -1697,35 +1697,6 @@ int32 FEngineLoop::PreInitPreStartupScreen(const TCHAR* CmdLine)
 		BeginPreInitTextLocalization();
 	}
 
-#if !(IS_PROGRAM || WITH_EDITOR)
-	// Initialize I/O dispatcher if available
-	{
-		TRACE_CPUPROFILER_EVENT_SCOPE(InitIoDispatcher);
-		FIoStoreEnvironment IoStoreEnvironment;
-		IoStoreEnvironment.InitializeFileEnvironment(FPaths::ProjectDir() / TEXT("global"));
-		bool bEnableIoDispatcher = false;
-		if (FIoDispatcher::IsValidEnvironment(IoStoreEnvironment))
-		{
-			bEnableIoDispatcher = true;
-		}
-#if !UE_BUILD_SHIPPING
-		if (FParse::Param(CmdLine, TEXT("forceiodispatcher")))
-		{
-			bEnableIoDispatcher = true;
-		}
-		else if (FParse::Param(CmdLine, TEXT("noiodispatcher")))
-		{
-			bEnableIoDispatcher = false;
-		}
-#endif
-		if (bEnableIoDispatcher)
-		{
-			FIoStatus IoDispatcherInitStatus = FIoDispatcher::Initialize();
-			UE_CLOG(!IoDispatcherInitStatus.IsOk(), LogInit, Fatal, TEXT("Failed to initialize I/O dispatcher: '%s'"), *IoDispatcherInitStatus.ToString());
-		}
-	}
-#endif
-
 	// allow the command line to override the platform file singleton
 	bool bFileOverrideFound = false;
 	{
diff --git a/Engine/Source/Runtime/PakFile/Private/IPlatformFilePak.cpp b/Engine/Source/Runtime/PakFile/Private/IPlatformFilePak.cpp
--- a/Engine/Source/Runtime/PakFile/Private/IPlatformFilePak.cpp
+++ b/Engine/Source/Runtime/PakFile/Private/IPlatformFilePak.cpp
@@ -6704,6 +6704,29 @@ bool FPakPlatformFile::Initialize(IPlatformFile* Inner, const TCHAR* CmdLine)
 	FParse::Value(FCommandLine::Get(), TEXT("StartupPaksWildcard="), StartupPaksWildcard);
 #endif
 
+	FIoStoreEnvironment IoStoreGlobalEnvironment;
+	IoStoreGlobalEnvironment.InitializeFileEnvironment(FString::Printf(TEXT("%sPaks/global"), *FPaths::ProjectContentDir()));
+	if (FIoDispatcher::IsValidEnvironment(IoStoreGlobalEnvironment))
+	{
+		FIoStatus IoDispatcherInitStatus = FIoDispatcher::Initialize();
+		if (IoDispatcherInitStatus.IsOk())
+		{
+			FIoStatus IoDispatcherMountStatus = FIoDispatcher::Get().Mount(IoStoreGlobalEnvironment);
+			if (IoDispatcherMountStatus.IsOk())
+			{
+				UE_LOG(LogPakFile, Display, TEXT("Initialized I/O dispatcher"));
+			}
+			else
+			{
+				UE_LOG(LogPakFile, Error, TEXT("Failed to mount I/O dispatcher global environment: '%s'"), *IoDispatcherMountStatus.ToString());
+			}
+		}
+		else
+		{
+			UE_LOG(LogPakFile, Error, TEXT("Failed to initialize I/O dispatcher: '%s'"), *IoDispatcherInitStatus.ToString());
+		}
+	}
+
 	// Find and mount pak files from the specified directories.
 	TArray<FString> PakFolders;
 	GetPakFolders(FCommandLine::Get(), PakFolders);
@@ -6916,7 +6939,7 @@ bool FPakPlatformFile::Mount(const TCHAR* InPakFilename, uint32 PakOrder, const
 			if (FIoDispatcher::IsInitialized())
 			{
 				FIoStoreEnvironment IoStoreEnvironment;
-				IoStoreEnvironment.InitializeFileEnvironment(FPaths::ChangeExtension(InPakFilename, FString()));
+				IoStoreEnvironment.InitializeFileEnvironment(FPaths::ChangeExtension(InPakFilename, FString()), PakOrder);
 				FIoStatus IoStatus = FIoDispatcher::Get().Mount(IoStoreEnvironment);
 				if (IoStatus.IsOk())
 				{
@@ -6924,7 +6947,7 @@ bool FPakPlatformFile::Mount(const TCHAR* InPakFilename, uint32 PakOrder, const
 				}
 				else
 				{
-					UE_LOG(LogPakFile, Warning, TEXT("Failed to mount IoStore environment \"%s\""), *IoStoreEnvironment.GetPath());
+					UE_LOG(LogPakFile, Warning, TEXT("Failed to mount IoStore environment \"%s\" [%s]"), *IoStoreEnvironment.GetPath(), *IoStatus.ToString());
 				}
 			}
 
