IoStore: Gracefully handle missing uasset files
#rnx

diff --git a/Engine/Source/Developer/IoStoreUtilities/Private/IoStoreUtilities.cpp b/Engine/Source/Developer/IoStoreUtilities/Private/IoStoreUtilities.cpp
--- a/Engine/Source/Developer/IoStoreUtilities/Private/IoStoreUtilities.cpp
+++ b/Engine/Source/Developer/IoStoreUtilities/Private/IoStoreUtilities.cpp
@@ -1125,7 +1125,10 @@ static void AssignPackagesDiskOrder(
 				if (!AssignedPackages.Contains(PackageToProcess))
 				{
 					AssignedPackages.Add(PackageToProcess);
-					Cluster->Packages.Add(PackageToProcess);
+					if (PackageToProcess->ExportBundles.Num())
+					{
+						Cluster->Packages.Add(PackageToProcess);
+					}
 					for (FPackage* ImportedPackage : PackageToProcess->ImportedPackages)
 					{
 						ProcessStack.Push(ImportedPackage);
@@ -2725,10 +2728,16 @@ static void ParsePackageAssets(
 			FPackage* Package = Packages[Index];
 			uint8* Buffer = PackageAssetBuffers[Index];
 			IFileHandle* FileHandle = FPlatformFileManager::Get().GetPlatformFile().OpenRead(*Package->FileName);
-			check(FileHandle);
-			bool bSuccess = FileHandle->Read(Buffer, Package->UAssetSize);
-			check(bSuccess);
-			delete FileHandle;
+			if (FileHandle)
+			{
+				bool bSuccess = FileHandle->Read(Buffer, Package->UAssetSize);
+				UE_CLOG(!bSuccess, LogIoStore, Warning, TEXT("Failed reading file '%s'"), *Package->FileName);
+				delete FileHandle;
+			}
+			else
+			{
+				UE_LOG(LogIoStore, Warning, TEXT("Couldn't open file '%s'"), *Package->FileName);
+			}
 			uint64 LocalFileIndex = CurrentFileIndex.IncrementExchange() + 1;
 			UE_CLOG(LocalFileIndex % 1000 == 0, LogIoStore, Display, TEXT("Reading %d/%d: '%s'"), LocalFileIndex, Packages.Num(), *Package->FileName);
 		}, EParallelForFlags::Unbalanced);
@@ -2748,6 +2757,11 @@ static void ParsePackageAssets(
 			FPackage& Package = *Packages[Index];
 
 			TArrayView<const uint8> MemView(PackageBuffer, Package.UAssetSize);
+			if (!Package.UAssetSize)
+			{
+				return;
+			}
+
 			FMemoryReaderView Ar(MemView);
 			Ar << Summary;
 
@@ -3504,6 +3518,10 @@ void InitializeContainerTargetsAndPackages(
 						{
 							TargetFile.ChunkId = CreateChunkIdForBulkData(Package->GlobalPackageId, BulkdataTypeToChunkIdType(FPackageStoreBulkDataManifest::EBulkdataType::Normal), *TargetFile.TargetPath);
 						}
+						if (Package->FileName.IsEmpty())
+						{
+							Package->FileName = FPaths::ChangeExtension(SourceFile.NormalizedPath, TEXT(".uasset"));
+						}
 					}
 					else
 					{
