IoStore: Replace GlobalImports (only referenced script objects and package exports) with ScriptObjects (all script objects) and PublicExports (all public exports)
1) First, parse all import tables to find the set of all referenced script packages (as a good approximation of which modules/script packages that are actually used for the target platform)
2) Then find and build import name lookup tables from all public non editoronly objects in memory in these packages.
3) Parse all export tables, to find all exports, and build import name lookup tables from all public exports.
4) Parse and validate all individual import entries in the import tables using the import name lookup tables from 2 & 3.
Dump all import and export data to/from the release version meta.

diff --git a/Engine/Source/Developer/IoStoreUtilities/Private/IoStoreUtilities.cpp b/Engine/Source/Developer/IoStoreUtilities/Private/IoStoreUtilities.cpp
--- a/Engine/Source/Developer/IoStoreUtilities/Private/IoStoreUtilities.cpp
+++ b/Engine/Source/Developer/IoStoreUtilities/Private/IoStoreUtilities.cpp
@@ -27,6 +27,7 @@
 #include "UObject/PackageFileSummary.h"
 #include "UObject/ObjectResource.h"
 #include "UObject/Package.h"
+#include "UObject/UObjectHash.h"
 #include "Algo/Find.h"
 #include "Misc/FileHelper.h"
 #include "Misc/ScopeLock.h"
@@ -293,6 +294,7 @@ struct FIoStoreArguments
 {
 	FString GlobalContainerPath;
 	FString CookedDir;
+	ITargetPlatform* TargetPlatform = nullptr;
 	FString OutputReleaseVersionDir;
 	FString BasedOnReleaseVersionDir;
 	TArray<FContainerSourceSpec> Containers;
@@ -347,6 +349,7 @@ using FSourceToLocalizedPackageMultimap = TMultiMap<FPackage*, FPackage*>;
 using FLocalizedToSourceImportIndexMap = TMap<FPackageObjectIndex, FPackageObjectIndex>;
 
 static constexpr TCHAR L10NPrefix[] = TEXT("/Game/L10N/");
+static constexpr TCHAR ScriptPrefix[] = TEXT("/Script/");
 
 // modified copy from PakFileUtilities
 static FString RemapLocalizationPathIfNeeded(const FString& Path, FString* OutRegion)
@@ -522,9 +525,10 @@ struct FExportGraphNode
 class FExportGraph
 {
 public:
-	FExportGraph()
+	FExportGraph(int32 NumExports, int32 NumPreloadDependencies)
 	{
-
+		Nodes.Reserve(NumExports * 2);
+		Edges.Reserve(NumExports + NumPreloadDependencies);
 	}
 
 	~FExportGraph()
@@ -609,12 +613,9 @@ struct FPackage
 	bool bIsLocalizedAndConformed = false;
 	bool bHasCircularImportDependencies = false;
 
-	TArray<FString> ImportedFullNames;
-
 	TArray<FPackage*> ImportedPackages;
 	TArray<FPackage*> ImportedByPackages;
 	TSet<FPackage*> AllReachablePackages;
-	TSet<FPackage*> ImportedPreloadPackages;
 
 	TArray<FName> Names;
 	TArray<FNameEntryId> NameMap;
@@ -1751,232 +1752,228 @@ static bool MapAdditionalBulkDataChunks(FContainerTargetFile& TargetFile, const
 	return true;
 }
 
-struct FScriptImportData
+struct FScriptObjectData
 {
 	FName ObjectName;
 	FString FullName;
 	FPackageObjectIndex GlobalIndex;
 	FPackageObjectIndex OuterIndex;
-	FPackageObjectIndex OutermostIndex;
 	FPackageObjectIndex CDOClassIndex;
-	FString CDOClassFullName;
 	bool bInitialized = false;
+
+	friend FArchive& operator<<(FArchive& Ar, FScriptObjectData& Data)
+	{
+		FString ObjectName;
+		if (Ar.IsSaving())
+		{
+			ObjectName = Data.ObjectName.ToString();
+			Ar << ObjectName;
+		}
+		else
+		{
+			Ar << ObjectName;
+			Data.ObjectName = FName(*ObjectName);
+		}
+		Ar << Data.FullName;
+		Ar << Data.GlobalIndex;
+		Ar << Data.OuterIndex;
+		Ar << Data.CDOClassIndex;
+		return Ar;
+	}
 };
 
-struct FPackageImportData
+struct FExportObjectData
 {
 	FName ObjectName;
 	FString FullName;
-	FPackageObjectIndex GlobalIndex;
-	int32 GlobalExportIndex = -1;
-	FPackage* Package = nullptr;
-	bool bIsLocalized = false;
-	bool bIsMissingImport = false;
-	bool bInitialized = false;
-};
-
-struct FExportData
-{
 	int32 GlobalIndex = -1;
-	FName ObjectName;
 	int32 SourceIndex = -1;
 	FPackageObjectIndex GlobalImportIndex;
 	FPackageObjectIndex OuterIndex;
 	FPackageObjectIndex ClassIndex;
 	FPackageObjectIndex SuperIndex;
 	FPackageObjectIndex TemplateIndex;
-	FString FullName;
+	EObjectFlags ObjectFlags = RF_NoFlags;
 
 	FPackage* Package = nullptr;
 	FExportGraphNode* CreateNode = nullptr;
 	FExportGraphNode* SerializeNode = nullptr;
+
+	friend FArchive& operator<<(FArchive& Ar, FExportObjectData& Data)
+	{
+		FString ObjectName;
+		if (Ar.IsSaving())
+		{
+			ObjectName = Data.ObjectName.ToString();
+			Ar << ObjectName;
+		}
+		else
+		{
+			Ar << ObjectName;
+			Data.ObjectName = FName(*ObjectName);
+		}
+		Ar << Data.FullName;
+		Ar << Data.GlobalIndex;
+		Ar << Data.SourceIndex;
+		Ar << Data.SourceIndex;
+		Ar << Data.GlobalImportIndex;
+		Ar << Data.OuterIndex;
+		Ar << Data.ClassIndex;
+		Ar << Data.TemplateIndex;
+		Ar << (uint32&)Data.ObjectFlags;
+		return Ar;
+	}
 };
 
 using FImportObjectsByFullName = TMap<FString, FPackageObjectIndex>;
+using FExportObjectsByFullName = TMap<FString, int32>;
+
+using FGlobalScriptObjects = TArray<FScriptObjectData>;
+using FGlobalExportObjects = TArray<FExportObjectData>;
 
-template <typename T>
-struct FGlobalObjects
+struct FGlobalPackageData
 {
-	TArray<T> Objects;
-	TMap<FString, int32> ObjectsByFullName;
-};
+	FGlobalScriptObjects ScriptObjects;
+	FGlobalExportObjects ExportObjects;
+	TArray<int32> PublicExportIndices;
+	FImportObjectsByFullName ImportsByFullName;
+	FExportObjectsByFullName ExportsByFullName;
+
+	void Reserve(int32 TotalExportCount)
+	{
+		const int32 EstimatedPublicExportObjectCount = TotalExportCount / 10;
+		const int32 EstimatedScriptObjectCount = 64000;
+		ExportObjects.Reserve(TotalExportCount);
+		ExportsByFullName.Reserve(TotalExportCount);
+		PublicExportIndices.Reserve(EstimatedPublicExportObjectCount);
+		ScriptObjects.Reserve(EstimatedScriptObjectCount);
+		ImportsByFullName.Reserve(EstimatedScriptObjectCount + EstimatedPublicExportObjectCount);
+	}
 
-using FGlobalScriptImports = TArray<FScriptImportData>;
-using FGlobalPackageImports = TArray<FPackageImportData>;
-using FGlobalExports = FGlobalObjects<FExportData>;
+	FExportObjectData& GetPublicExport(FPackageObjectIndex Index)
+	{
+		check(Index.IsPackageImport());
+		int32 GlobalExportIndex = PublicExportIndices[Index.GetIndex()];
+		return ExportObjects[GlobalExportIndex];
+	}
 
-struct FGlobalImports
-{
-	FGlobalScriptImports ScriptImports;
-	FGlobalPackageImports PackageImports;
-	FImportObjectsByFullName ObjectsByFullName;
-};
+	const FExportObjectData& GetPublicExport(FPackageObjectIndex Index) const
+	{
+		check(Index.IsPackageImport());
+		int32 GlobalExportIndex = PublicExportIndices[Index.GetIndex()];
+		return ExportObjects[GlobalExportIndex];
+	}
 
-struct FGlobalPackageData
-{
-	FGlobalImports Imports;
-	FGlobalExports Exports;
+	FName GetObjectName(FPackageObjectIndex Index, const TArray<int32>* PackageExportIndices) const
+	{
+		if (Index.IsScriptImport())
+		{
+			return ScriptObjects[Index.GetIndex()].ObjectName;
+		}
+		if (Index.IsPackageImport())
+		{
+			int32 GlobalExportIndex = PublicExportIndices[Index.GetIndex()];
+			return ExportObjects[GlobalExportIndex].ObjectName;
+		}
+		if (Index.IsExport() && PackageExportIndices)
+		{
+			int32 GlobalExportIndex = (*PackageExportIndices)[Index.GetIndex()];
+			return ExportObjects[GlobalExportIndex].ObjectName;
+		}
+		return FName();
+	}
+
+	friend FArchive& operator<<(FArchive& Ar, FGlobalPackageData& Data)
+	{
+		Ar << Data.ScriptObjects;
+		Ar << Data.ExportObjects;
+		Ar << Data.PublicExportIndices;
+		Ar << Data.ImportsByFullName;
+		Ar << Data.ExportsByFullName;
+		return Ar;
+	}
 };
 
-static void FindImport(
-	FGlobalImports& GlobalImports,
-	TArray<FString>& TempFullNames,
+static void FindImportFullName(
+	TArray<FString>& ImportFullNames,
 	FObjectImport* ImportMap,
-	const int32 LocalImportIndex,
-	const FPackageMap& PackageMap,
-	FPackage*& CurrentPackage)
+	const int32 LocalImportIndex)
 {
-	FObjectImport* Import = &ImportMap[LocalImportIndex];
-	FString& FullName = TempFullNames[LocalImportIndex];
-
-	if (Import->OuterIndex.IsNull())
-	{
-		CurrentPackage = PackageMap.FindRef(Import->ObjectName);
-	}
+	FString& FullName = ImportFullNames[LocalImportIndex];
 
 	if (FullName.Len() == 0)
 	{
+		FullName.Reserve(256);
+
+		FObjectImport* Import = &ImportMap[LocalImportIndex];
 		if (Import->OuterIndex.IsNull())
 		{
 			Import->ObjectName.AppendString(FullName);
-
-			const bool bIsScript = FullName.StartsWith(TEXT("/Script/"));
-			if (bIsScript)
-			{
-				FScriptImportData* ScriptImport;
-				FPackageObjectIndex FindGlobalImportIndex = GlobalImports.ObjectsByFullName.FindRef(FullName);
-				if (FindGlobalImportIndex.IsNull())
-				{
-					// assign global index for this script UPackage
-					FPackageObjectIndex::Type IndexType = FPackageObjectIndex::ScriptImport;
-					TArray<FScriptImportData>& ScriptImports = GlobalImports.ScriptImports;
-					FPackageObjectIndex GlobalImportIndex(IndexType, ScriptImports.Num());
-					GlobalImports.ObjectsByFullName.Add(FullName, GlobalImportIndex);
-
-					ScriptImport = &ScriptImports.AddDefaulted_GetRef();
-					ScriptImport->GlobalIndex = GlobalImportIndex;
-					ScriptImport->FullName = FullName;
-				}
-				else
-				{
-					ScriptImport = &GlobalImports.ScriptImports[FindGlobalImportIndex.GetIndex()];
-				}
-				if (!ScriptImport->bInitialized)
-				{
-					ScriptImport->OuterIndex = FPackageObjectIndex();
-					ScriptImport->ObjectName = Import->ObjectName;
-					ScriptImport->bInitialized = true;
-				}
-			}
 		}
 		else
 		{
-			const int32 LocalOuterIndex = Import->OuterIndex.ToImport();
-			FindImport(GlobalImports, TempFullNames, ImportMap, LocalOuterIndex, PackageMap, CurrentPackage);
-			const FString& OuterName = TempFullNames[LocalOuterIndex];
+			const int32 OuterIndex = Import->OuterIndex.ToImport();
+			FindImportFullName(ImportFullNames, ImportMap, OuterIndex);
+			const FString& OuterName = ImportFullNames[OuterIndex];
 			check(OuterName.Len() > 0);
 
 			FullName.Append(OuterName);
 			FullName.AppendChar(TEXT('/'));
 			Import->ObjectName.AppendString(FullName);
+		}
+	}
+}
 
-			const bool bIsScript = FullName.StartsWith(TEXT("/Script/"));
-			FPackageObjectIndex FindGlobalImportIndex = GlobalImports.ObjectsByFullName.FindRef(FullName);
-			FPackageObjectIndex FindOuterGlobalImport = GlobalImports.ObjectsByFullName.FindRef(OuterName);
+static FPackageObjectIndex FindAndVerifyGlobalImport(
+	FGlobalPackageData& GlobalPackageData,
+	FObjectImport& Import,
+	FString& FullName)
+{
+	FPackageObjectIndex GlobalImportIndex = GlobalPackageData.ImportsByFullName.FindRef(FullName);
+	if (GlobalImportIndex.IsNull())
+	{
+		bool bIsPackage = Import.OuterIndex.IsNull();
+		bool bIsScript = FullName.StartsWith(ScriptPrefix);
+		if (bIsPackage)
+		{
 			if (bIsScript)
 			{
-				check(FindOuterGlobalImport.IsScriptImport());
-
-				FScriptImportData* ScriptImport;
-				if (FindGlobalImportIndex.IsNull())
-				{
-					// assign global index for this script UObject
-					TArray<FScriptImportData>& ScriptImports = GlobalImports.ScriptImports;
-					FPackageObjectIndex GlobalImportIndex(FPackageObjectIndex::ScriptImport, ScriptImports.Num());
-					GlobalImports.ObjectsByFullName.Add(FullName, GlobalImportIndex);
-
-					ScriptImport = &GlobalImports.ScriptImports.AddDefaulted_GetRef();
-					ScriptImport->GlobalIndex = GlobalImportIndex;
-					ScriptImport->FullName = FullName;
-				}
-				else
-				{
-					ScriptImport = &GlobalImports.ScriptImports[FindGlobalImportIndex.GetIndex()];
-				}
-				if (!ScriptImport->bInitialized)
-				{
-					ScriptImport->OuterIndex = FindOuterGlobalImport;
-					ScriptImport->ObjectName = Import->ObjectName;
-
-					const FScriptImportData& OuterScriptImport = GlobalImports.ScriptImports[FindOuterGlobalImport.GetIndex()];
-					if (OuterScriptImport.CDOClassFullName.Len() > 0)
-					{
-						ScriptImport->CDOClassFullName = OuterScriptImport.CDOClassFullName;
-					}
-					else
-					{
-						TCHAR NameBuffer[FName::StringBufferSize];
-						uint32 Len = Import->ObjectName.ToString(NameBuffer);
-						if (FCString::Strncmp(NameBuffer, TEXT("Default__"), 9) == 0)
-						{
-							ScriptImport->CDOClassFullName.Append(OuterName);
-							ScriptImport->CDOClassFullName.AppendChar(TEXT('/'));
-							ScriptImport->CDOClassFullName.AppendChars(NameBuffer + 9, Len - 9);
-						}
-					}
-					ScriptImport->bInitialized = true;
-				}
+				UE_LOG(LogIoStore, Warning, TEXT("Missing import script package '%s'. Editor only?"), *FullName);
 			}
 			else
 			{
-				FPackageImportData* PackageImport;
-				if (FindGlobalImportIndex.IsNull())
-				{
-					TArray<FPackageImportData>& PackageImports = GlobalImports.PackageImports;
-					FPackageObjectIndex GlobalImportIndex(FPackageObjectIndex::PackageImport, PackageImports.Num());
-					GlobalImports.ObjectsByFullName.Add(FullName, GlobalImportIndex);
-
-					PackageImport = &GlobalImports.PackageImports.AddDefaulted_GetRef();
-					PackageImport->GlobalIndex = GlobalImportIndex;
-					PackageImport->FullName = FullName;
-				}
-				else
-				{
-					PackageImport = &GlobalImports.PackageImports[FindGlobalImportIndex.GetIndex()];
-				}
-				if (!PackageImport->bInitialized)
-				{
-					PackageImport->ObjectName = Import->ObjectName;
-					PackageImport->bIsLocalized = FullName.Contains(L10NPrefix);
-
-					if (FindOuterGlobalImport.IsPackageImport())
-					{
-						const FPackageImportData& OuterPackageImport = GlobalImports.PackageImports[FindOuterGlobalImport.GetIndex()];
-						PackageImport->Package = OuterPackageImport.Package;
-					}
-					else
-					{
-						PackageImport->Package = CurrentPackage;
-					}
-					PackageImport->bInitialized = true;
-				}
+				// UPackages are never serialized, hence they are never imports
+			}
+		}
+		else
+		{
+			if (bIsScript)
+			{
+				UE_LOG(LogIoStore, Warning, TEXT("Missing import script object '%s'. Editor only?"), *FullName);
+			}
+			else
+			{
+				UE_LOG(LogIoStore, Warning, TEXT("Missing import object '%s' due to missing public export. Editor only?"), *FullName);
 			}
 		}
 	}
+	return GlobalImportIndex;
 }
 
-static void FindExport(
-	FGlobalExports& GlobalExports,
+static int32 FindExport(
+	FGlobalPackageData& GlobalPackageData,
 	TArray<FString>& TempFullNames,
 	const FObjectExport* ExportMap,
 	const int32 LocalExportIndex,
 	FPackage* Package)
 {
-	const FObjectExport* Export = ExportMap + LocalExportIndex;
 	FString& FullName = TempFullNames[LocalExportIndex];
 
 	if (FullName.Len() == 0)
 	{
+		FullName.Reserve(256);
+
+		const FObjectExport* Export = ExportMap + LocalExportIndex;
 		if (Export->OuterIndex.IsNull())
 		{
 			Package->Name.AppendString(FullName);
@@ -1987,22 +1984,32 @@ static void FindExport(
 		{
 			check(Export->OuterIndex.IsExport());
 
-			FindExport(GlobalExports, TempFullNames, ExportMap, Export->OuterIndex.ToExport(), Package);
+			FindExport(GlobalPackageData, TempFullNames, ExportMap, Export->OuterIndex.ToExport(), Package);
 			FString& OuterName = TempFullNames[Export->OuterIndex.ToExport()];
 			check(OuterName.Len() > 0);
 			FullName.Append(OuterName);
 			FullName.AppendChar(TEXT('/'));
 			Export->ObjectName.AppendString(FullName);
 		}
-		check(!GlobalExports.ObjectsByFullName.Contains(FullName));
-		const int32 GlobalExportIndex = GlobalExports.Objects.Num();
-		GlobalExports.ObjectsByFullName.Add(FullName, GlobalExportIndex);
-		FExportData& ExportData = GlobalExports.Objects.AddDefaulted_GetRef();
+
+		const int32 GlobalExportIndex = GlobalPackageData.ExportObjects.Num();
+		GlobalPackageData.ExportsByFullName.Add(FullName, GlobalExportIndex);
+
+		FExportObjectData& ExportData = GlobalPackageData.ExportObjects.AddDefaulted_GetRef();
 		ExportData.GlobalIndex = GlobalExportIndex;
 		ExportData.Package = Package;
 		ExportData.ObjectName = Export->ObjectName;
 		ExportData.SourceIndex = LocalExportIndex;
 		ExportData.FullName = FullName;
+		ExportData.ObjectFlags = Export->ObjectFlags;
+
+		return GlobalExportIndex;
+	}
+	else
+	{
+		int32* GlobalExportIndex = GlobalPackageData.ExportsByFullName.Find(FullName);
+		check(GlobalExportIndex);
+		return *GlobalExportIndex;
 	}
 }
 
@@ -2089,10 +2096,9 @@ FPackage* FindOrAddPackage(
 
 static bool ConformLocalizedPackage(
 	const FPackageMap& PackageMap,
-	const FGlobalImports& GlobalImports,
+	FGlobalPackageData& GlobalPackageData,
 	const FPackage& SourcePackage,
 	FPackage& LocalizedPackage,
-	TArray<FExportData>& GlobalExports,
 	FLocalizedToSourceImportIndexMap& LocalizedToSourceImportIndexMap)
 {
 	const int32 ExportCount =
@@ -2145,25 +2151,11 @@ static bool ConformLocalizedPackage(
 		return nullptr;
 	};
 
-	auto AppendMismatchMessage = [&GlobalImports, &GlobalExports, &LocalizedPackage, &SourcePackage](
+	auto AppendMismatchMessage = [&GlobalPackageData, &LocalizedPackage, &SourcePackage](
 		const TCHAR* Text, FName ExportName, FPackageObjectIndex LocIndex, FPackageObjectIndex SrcIndex, FString& FailReason)
 	{
-		FString LocString = 
-			LocIndex.IsPackageImport() ?
-			GlobalImports.PackageImports[LocIndex.GetIndex()].ObjectName.ToString() :
-			LocIndex.IsScriptImport() ?
-			GlobalImports.ScriptImports[LocIndex.GetIndex()].ObjectName.ToString() :
-			LocIndex.IsExport() ?
-			GlobalExports[LocalizedPackage.Exports[LocIndex.GetIndex()]].ObjectName.ToString() :
-			TEXT("");
-		FString SrcString = 
-			SrcIndex.IsPackageImport() ?
-			GlobalImports.PackageImports[SrcIndex.GetIndex()].ObjectName.ToString() :
-			SrcIndex.IsScriptImport() ?
-			GlobalImports.ScriptImports[SrcIndex.GetIndex()].ObjectName.ToString() :
-			SrcIndex.IsExport() ?
-			GlobalExports[LocalizedPackage.Exports[SrcIndex.GetIndex()]].ObjectName.ToString() :
-			TEXT("");
+		FString LocString = GlobalPackageData.GetObjectName(LocIndex, &LocalizedPackage.Exports).ToString();
+		FString SrcString = GlobalPackageData.GetObjectName(SrcIndex, &SourcePackage.Exports).ToString();
 
 		FailReason.Appendf(TEXT("Public export '%s' has %s %s (%d) vs. %s (%d)"),
 			*ExportName.ToString(),
@@ -2186,8 +2178,8 @@ static bool ConformLocalizedPackage(
 	while (LocalizedIndex < ExportCount && SourceIndex < ExportCount)
 	{
 		FString FailReason;
-		const FExportData& LocalizedExportData = GlobalExports[LocalizedPackage.Exports[LocalizedIndex]];
-		const FExportData& SourceExportData = GlobalExports[SourcePackage.Exports[SourceIndex]];
+		const FExportObjectData& LocalizedExportData = GlobalPackageData.ExportObjects[LocalizedPackage.Exports[LocalizedIndex]];
+		const FExportObjectData& SourceExportData = GlobalPackageData.ExportObjects[SourcePackage.Exports[SourceIndex]];
 
 		const TCHAR* LocalizedExportStr = GetExportNameSafe(
 			LocalizedExportData.FullName, LocalizedPackage.Name, LocalizedPackageNameLen);
@@ -2281,10 +2273,10 @@ static bool ConformLocalizedPackage(
 		LocalizedPackage.PublicExports.Reset();
 		for (TPair<int32, int32>& Pair : NewPublicExports)
 		{
-			FExportData& LocalizedExportData = GlobalExports[LocalizedPackage.Exports[Pair.Key]];
+			FExportObjectData& LocalizedExportData = GlobalPackageData.ExportObjects[LocalizedPackage.Exports[Pair.Key]];
 			if (Pair.Value != -1)
 			{
-				const FExportData& SourceExportData = GlobalExports[SourcePackage.Exports[Pair.Value]];
+				const FExportObjectData& SourceExportData = GlobalPackageData.ExportObjects[SourcePackage.Exports[Pair.Value]];
 
 				LocalizedToSourceImportIndexMap.Add(
 					LocalizedExportData.GlobalImportIndex,
@@ -2309,9 +2301,6 @@ static void AddPreloadDependencies(
 	IOSTORE_CPU_SCOPE(PreLoadDependencies);
 	UE_LOG(LogIoStore, Display, TEXT("Adding preload dependencies..."));
 
-	const FGlobalPackageImports& PackageImports = GlobalPackageData.Imports.PackageImports;
-	const TArray<FExportData>& GlobalExports = GlobalPackageData.Exports.Objects;
-
 	TArray<FPackage*> LocalizedPackages;
 	for (FPackage* Package : Packages)
 	{
@@ -2334,32 +2323,22 @@ static void AddPreloadDependencies(
 					FPackageObjectIndex ImportIndex = Package->Imports[Dep.ToImport()];
 					if (ImportIndex.IsPackageImport())
 					{
-						const FPackageImportData& Import = PackageImports[ImportIndex.GetIndex()];
-						check(Import.GlobalIndex == ImportIndex);
-						if (Import.GlobalExportIndex != -1)
-						{
-							const FExportData& Export = GlobalExports[Import.GlobalExportIndex];
-
-							AddExternalExportArc(ExportGraph, *Export.Package, Export.SourceIndex, PhaseFrom, *Package, I, PhaseTo);
-							Package->ImportedPreloadPackages.Add(Export.Package);
-
-							LocalizedPackages.Reset();
-							SourceToLocalizedPackageMap.MultiFind(Export.Package, LocalizedPackages);
-							for (FPackage* LocalizedPackage : LocalizedPackages)
-							{
-								UE_LOG(LogIoStore, Verbose, TEXT("For package '%s' (%d): Adding localized preload dependency '%s' in '%s'"),
-									*Package->Name.ToString(),
-									Package->GlobalPackageId.ToIndexForDebugging(),
-									*Export.ObjectName.ToString(),
-									*LocalizedPackage->Name.ToString());
-
-								AddExternalExportArc(ExportGraph, *LocalizedPackage, Export.SourceIndex, PhaseFrom, *Package, I, PhaseTo);
-								Package->ImportedPreloadPackages.Add(LocalizedPackage);
-							}
-						}
-						else
+						const FExportObjectData& Export = GlobalPackageData.GetPublicExport(ImportIndex);
+						check(Export.GlobalImportIndex == ImportIndex);
+
+						AddExternalExportArc(ExportGraph, *Export.Package, Export.SourceIndex, PhaseFrom, *Package, I, PhaseTo);
+
+						LocalizedPackages.Reset();
+						SourceToLocalizedPackageMap.MultiFind(Export.Package, LocalizedPackages);
+						for (FPackage* LocalizedPackage : LocalizedPackages)
 						{
-							UE_LOG(LogIoStore, Verbose, TEXT("Skipping export arc to '%s' due to missing import"), *Package->Name.ToString());
+							UE_LOG(LogIoStore, Verbose, TEXT("For package '%s' (%d): Adding localized preload dependency '%s' in '%s'"),
+								*Package->Name.ToString(),
+								Package->GlobalPackageId.ToIndexForDebugging(),
+								*Export.ObjectName.ToString(),
+								*LocalizedPackage->Name.ToString());
+
+							AddExternalExportArc(ExportGraph, *LocalizedPackage, Export.SourceIndex, PhaseFrom, *Package, I, PhaseTo);
 						}
 					}
 				}
@@ -2415,7 +2394,7 @@ void BuildContainerNameMap(FContainerTargetSpec& ContainerTarget)
 void FinalizePackageHeaders(
 	FContainerTargetSpec& ContainerTarget,
 	const TArray<FObjectExport>& ObjectExports,
-	const TArray<FExportData>& GlobalExports,
+	const TArray<FExportObjectData>& GlobalExports,
 	const FImportObjectsByFullName& GlobalImportsByFullName)
 {
 	const uint16 NameMapIndex = ContainerTarget.bUseLocalNameMap ? ContainerTarget.Header.ContainerId.ToIndex() : 0;
@@ -2460,7 +2439,7 @@ void FinalizePackageHeaders(
 		for (int32 I = 0; I < Package->ExportCount; ++I)
 		{
 			const FObjectExport& ObjectExport = ObjectExports[Package->ExportIndexOffset + I];
-			const FExportData& ExportData = GlobalExports[Package->Exports[I]];
+			const FExportObjectData& ExportData = GlobalExports[Package->Exports[I]];
 
 			FExportMapEntry ExportMapEntry;
 			ExportMapEntry.CookedSerialOffset = ObjectExport.SerialOffset;
@@ -2638,7 +2617,7 @@ void SerializePackageStoreEntries(
 
 static void SerializeInitialLoad(
 	FNameMapBuilder& GlobalNameMapBuilder,
-	const FGlobalScriptImports& GlobalScriptImports,
+	const FGlobalScriptObjects& GlobalScriptImports,
 	FArchive& InitialLoadArchive)
 {
 	IOSTORE_CPU_SCOPE(SerializeInitialLoad);
@@ -2647,7 +2626,7 @@ static void SerializeInitialLoad(
 	int32 NumScriptObjects = GlobalScriptImports.Num();
 	InitialLoadArchive << NumScriptObjects; 
 
-	for (const FScriptImportData& ImportData : GlobalScriptImports)
+	for (const FScriptObjectData& ImportData : GlobalScriptImports)
 	{
 		GlobalNameMapBuilder.MarkNameAsReferenced(ImportData.ObjectName);
 		FScriptObjectEntry Entry;
@@ -2687,8 +2666,7 @@ static FIoBuffer CreateExportBundleBuffer(const FContainerTargetFile& TargetFile
 
 static void ParsePackageAssets(
 	TArray<FPackage*>& Packages,
-	FPackageAssetData& PackageAssetData,
-	FExportGraph& ExportGraph)
+	FPackageAssetData& PackageAssetData)
 {
 	IOSTORE_CPU_SCOPE(ParsePackageAssets);
 	UE_LOG(LogIoStore, Display, TEXT("Parsing packages..."));
@@ -2837,8 +2815,7 @@ static void ParsePackageAssets(
 				&PackageAssetBuffers,
 				&PackageFileSummaries,
 				&Packages,
-				&PackageAssetData,
-				&ExportGraph](int32 Index)
+				&PackageAssetData](int32 Index)
 		{
 			uint8* PackageBuffer = PackageAssetBuffers[Index];
 			const FPackageFileSummary& Summary = PackageFileSummaries[Index];
@@ -2888,146 +2865,336 @@ static void ParsePackageAssets(
 	FMemory::Free(UAssetMemory);
 }
 
-static void CreateGlobalImportsAndExports(
-	TArray<FPackage*>& Packages,
-	const FPackageMap& PackageMap,
-	FPackageAssetData& PackageAssetData,
-	FGlobalPackageData& GlobalPackageData,
-	FExportGraph& ExportGraph)
+static void FindScriptPackages(TArray<FObjectImport>& ObjectImports, TArray<UPackage*>& ScriptPackages)
 {
-	IOSTORE_CPU_SCOPE(CreateGlobalImportsAndExports);
-	UE_LOG(LogIoStore, Display, TEXT("Creating global imports and exports..."));
-
-	FGlobalImports& GlobalImports = GlobalPackageData.Imports;
-	FImportObjectsByFullName& GlobalImportsByFullName = GlobalPackageData.Imports.ObjectsByFullName;
+	IOSTORE_CPU_SCOPE(FindScriptPackages);
+	UE_LOG(LogIoStore, Display, TEXT("Finding script packages..."));
 
-	TArray<FString> TmpFullNames;
+	TSet<FName> AllPackages;
+	AllPackages.Reserve(65536);
+	ScriptPackages.Reserve(16384);
 
-	for (FPackage* Package : Packages)
+	for (const FObjectImport& ObjectImport : ObjectImports)
 	{
-		if (Package->ImportCount > 0)
+		if (ObjectImport.OuterIndex.IsNull())
 		{
-			FPackage* CurrentImportPackage = nullptr;
-			Package->ImportedFullNames.SetNum(Package->ImportCount);
-			for (int32 ImportIndex = 0; ImportIndex < Package->ImportCount; ++ImportIndex)
-			{
-				FindImport(
-					GlobalPackageData.Imports,
-					Package->ImportedFullNames,
-					PackageAssetData.ObjectImports.GetData() + Package->ImportIndexOffset,
-					ImportIndex,
-					PackageMap,
-					CurrentImportPackage);
-			}
-		}
-
-		if (Package->ExportCount > 0)
-		{
-			TmpFullNames.Reset();
-			TmpFullNames.SetNum(Package->ExportCount, false);
-
-			for (int32 ExportIndex = 0; ExportIndex < Package->ExportCount; ExportIndex++)
-			{
-				FindExport(
-					GlobalPackageData.Exports,
-					TmpFullNames,
-					PackageAssetData.ObjectExports.GetData() + Package->ExportIndexOffset,
-					ExportIndex,
-					Package);
-				const int32* GlobalIndex = GlobalPackageData.Exports.ObjectsByFullName.Find(TmpFullNames[ExportIndex]);
-				FExportData& ExportData = GlobalPackageData.Exports.Objects[*GlobalIndex];
-				Package->Exports.Add(ExportData.GlobalIndex);
-				ExportData.CreateNode = ExportGraph.AddNode(Package, { uint32(ExportIndex), FExportBundleEntry::ExportCommandType_Create });
-				ExportData.SerializeNode = ExportGraph.AddNode(Package, { uint32(ExportIndex), FExportBundleEntry::ExportCommandType_Serialize });
-				Package->CreateExportNodes.Add(ExportData.CreateNode);
-				Package->SerializeExportNodes.Add(ExportData.SerializeNode);
-				ExportGraph.AddInternalDependency(ExportData.CreateNode, ExportData.SerializeNode);
-			}
+			AllPackages.Add(ObjectImport.ObjectName);
 		}
 	}
 
-	// lookup script CDO class indices for script imports
-	for (FScriptImportData& ImportData : GlobalImports.ScriptImports)
+	TCHAR NameBuffer[FName::StringBufferSize];
+	static constexpr int32 ScriptPrefixLength = sizeof(ScriptPrefix)/sizeof(TCHAR) - 1;
+	for (const FName& Name : AllPackages)
 	{
-		if (ImportData.bInitialized && ImportData.CDOClassFullName.Len() > 0)
+		Name.ToString(NameBuffer);
+		if (FCString::Strncmp(NameBuffer, ScriptPrefix, ScriptPrefixLength) == 0)
 		{
-			FPackageObjectIndex* CDOClassIndex = GlobalImportsByFullName.Find(ImportData.CDOClassFullName);
-			if (CDOClassIndex)
+			UPackage* Package = (UPackage*)StaticFindObjectFast(UPackage::StaticClass(), nullptr, Name, true);
+			if (Package)
 			{
-				ImportData.CDOClassIndex = *CDOClassIndex;
+				ScriptPackages.Add(Package);
 			}
 			else
 			{
-				UE_LOG(LogIoStore, Warning, TEXT("Missing script import class '%s' for CDO '%s'"),
-					*ImportData.CDOClassFullName, *ImportData.FullName);
+				UE_LOG(LogIoStore, Warning, TEXT("Failed to find imported script package: %s"), NameBuffer);
 			}
 		}
 	}
+}
+
+EObjectMark GetExcludedObjectMarksForTargetPlatform(const ITargetPlatform* TargetPlatform)
+{
+	EObjectMark Marks = OBJECTMARK_NOMARKS;
+	if (!TargetPlatform->HasEditorOnlyData())
+	{
+		Marks = (EObjectMark)(Marks | OBJECTMARK_EditorOnly);
+	}
+	if (TargetPlatform->IsServerOnly())
+	{
+		Marks = (EObjectMark)(Marks | OBJECTMARK_NotForServer);
+	}
+	if (TargetPlatform->IsClientOnly())
+	{
+		Marks = (EObjectMark)(Marks | OBJECTMARK_NotForClient);
+	}
+	return Marks;
+}
 
-	// lookup mappings between package imports and exports, and fill the package public exports
-	for (FPackageImportData& ImportData : GlobalImports.PackageImports)
+EObjectMark GetExcludedObjectMarksForObject(const UObject* Object, const ITargetPlatform* TargetPlatform)
+{
+	EObjectMark Marks = OBJECTMARK_NOMARKS;
+	if (!Object->NeedsLoadForClient())
+	{
+		Marks = (EObjectMark)(Marks | OBJECTMARK_NotForClient);
+	}
+	if (!Object->NeedsLoadForServer())
+	{
+		Marks = (EObjectMark)(Marks | OBJECTMARK_NotForServer);
+	}
+	if (!Object->NeedsLoadForTargetPlatform(TargetPlatform))
+	{
+		Marks = (EObjectMark)(Marks | OBJECTMARK_NotForClient | OBJECTMARK_NotForServer);
+	}
+	if (Object->IsEditorOnly())
+	{
+		Marks = (EObjectMark)(Marks | OBJECTMARK_EditorOnly);
+	}
+	if ((Marks & OBJECTMARK_NotForClient) && (Marks & OBJECTMARK_NotForServer))
+	{
+		Marks = (EObjectMark)(Marks | OBJECTMARK_EditorOnly);
+	}
+	return Marks;
+}
+
+static void FindScriptObjectsRecursive(
+	FGlobalPackageData& GlobalPackageData,
+	FPackageObjectIndex OuterIndex,
+	UObject* Object,
+	const ITargetPlatform* TargetPlatform,
+	const EObjectMark ExcludedObjectMarks)
+{
+	if (!Object->HasAllFlags(RF_Public))
+	{
+		UE_LOG(LogIoStore, Log, TEXT("Skipping script object: %s (!RF_Public)"), *Object->GetFullName());
+		return;
+	}
+
+	const UObject* ObjectForExclusion = Object->HasAnyFlags(RF_ClassDefaultObject) ? Object->GetClass() : Object;
+	const EObjectMark ObjectMarks = GetExcludedObjectMarksForObject(ObjectForExclusion, TargetPlatform);
+
+	if (ObjectMarks & ExcludedObjectMarks)
+	{
+		UE_LOG(LogIoStore, Log, TEXT("Skipping script object: %s (Excluded for target platform)"), *Object->GetFullName());
+		return;
+	}
+
+	FName ObjectName = Object->GetFName();
+
+	FString TempFullName = GlobalPackageData.ScriptObjects[OuterIndex.GetIndex()].FullName;
+	TempFullName.AppendChar(TEXT('/'));
+	ObjectName.AppendString(TempFullName);
+
+	FPackageObjectIndex FindGlobalImportIndex = GlobalPackageData.ImportsByFullName.FindRef(TempFullName);
+	FScriptObjectData* ScriptImport = nullptr;
+
+	if (FindGlobalImportIndex.IsNull())
+	{
+		TArray<FScriptObjectData>& ScriptObjects = GlobalPackageData.ScriptObjects;
+		FPackageObjectIndex GlobalImportIndex(FPackageObjectIndex::ScriptImport, ScriptObjects.Num());
+		GlobalPackageData.ImportsByFullName.Add(TempFullName, GlobalImportIndex);
+
+		ScriptImport = &GlobalPackageData.ScriptObjects.AddDefaulted_GetRef();
+		ScriptImport->GlobalIndex = GlobalImportIndex;
+		ScriptImport->FullName = MoveTemp(TempFullName);
+	}
+	else
 	{
-		if (ImportData.bInitialized)
+		ScriptImport = &GlobalPackageData.ScriptObjects[FindGlobalImportIndex.GetIndex()];
+	}
+	if (!ScriptImport->bInitialized)
+	{
+		FScriptObjectData& Outer = GlobalPackageData.ScriptObjects[OuterIndex.GetIndex()];
+		FPackageObjectIndex CDOClassIndex = Outer.CDOClassIndex;
+		if (CDOClassIndex.IsNull())
 		{
-			int32* FindGlobalExport = GlobalPackageData.Exports.ObjectsByFullName.Find(ImportData.FullName);
-			if (FindGlobalExport)
-			{
-				check(ImportData.Package);
-				ImportData.Package->PublicExports.Add(ImportData.GlobalIndex);
-				ImportData.GlobalExportIndex = *FindGlobalExport;
-				GlobalPackageData.Exports.Objects[*FindGlobalExport].GlobalImportIndex = ImportData.GlobalIndex;
-			}
-			else
+			TCHAR NameBuffer[FName::StringBufferSize];
+			uint32 Len = ObjectName.ToString(NameBuffer);
+			if (FCString::Strncmp(NameBuffer, TEXT("Default__"), 9) == 0)
 			{
-				UE_LOG(LogIoStore, Warning, TEXT("Missing import '%s' due to missing export"), *ImportData.FullName);
-				ImportData.bIsMissingImport = true;
+				FString CDOClassFullName = Outer.FullName;
+				CDOClassFullName.AppendChar(TEXT('/'));
+				CDOClassFullName.AppendChars(NameBuffer + 9, Len - 9);
+
+				CDOClassIndex = GlobalPackageData.ImportsByFullName.FindRef(CDOClassFullName);
+				check(CDOClassIndex.IsScriptImport());
 			}
 		}
+		ScriptImport->OuterIndex = Outer.GlobalIndex;
+		ScriptImport->ObjectName = ObjectName;
+		ScriptImport->CDOClassIndex = CDOClassIndex;
+		ScriptImport->bInitialized = true;
+	}
+
+	OuterIndex = ScriptImport->GlobalIndex;
+	TArray<UObject*> InnerObjects;
+	GetObjectsWithOuter(Object, InnerObjects, /*bIncludeNestedObjects*/false);
+	for (UObject* InnerObject : InnerObjects)
+	{
+		FindScriptObjectsRecursive(GlobalPackageData, OuterIndex, InnerObject, TargetPlatform, ExcludedObjectMarks);
+	}
+};
+
+static void CreateGlobalScriptObjects(
+	TArray<UPackage*> Packages,
+	FNameMapBuilder& NameMapBuilder,
+	FGlobalPackageData& GlobalPackageData,
+	const ITargetPlatform* TargetPlatform)
+{
+	IOSTORE_CPU_SCOPE(CreateGlobalScriptObjects);
+	UE_LOG(LogIoStore, Display, TEXT("Creating global script objects..."));
+
+	const EObjectMark ExcludedObjectMarks = GetExcludedObjectMarksForTargetPlatform(TargetPlatform);
+
+	TArray<UObject*> InnerObjects;
+	// GetObjectsOfClass(UPackage::StaticClass(), Packages, /*bIncludeDerivedClasses*/false);
+	for (UPackage* Package : Packages)
+	{
+		if (!Package->HasAnyPackageFlags(PKG_CompiledIn))
+		{
+			UE_LOG(LogIoStore, Warning, TEXT("Skipping script package: %s"), *Package->GetName());
+			continue;
+		}
+
+		if (Package->HasAnyPackageFlags(PKG_EditorOnly | PKG_Developer | PKG_UncookedOnly))
+		{
+			UE_LOG(LogIoStore, Warning, TEXT("Skipping script package: %s"), *Package->GetName());
+			continue;
+		}
+
+		FName ObjectName = Package->GetFName();
+		FString FullName = Package->GetName();
+
+		FPackageObjectIndex FindGlobalImportIndex = GlobalPackageData.ImportsByFullName.FindRef(FullName);
+		FScriptObjectData* ScriptImport = nullptr;
+
+		if (FindGlobalImportIndex.IsNull())
+		{
+			TArray<FScriptObjectData>& ScriptObjects = GlobalPackageData.ScriptObjects;
+			FPackageObjectIndex GlobalImportIndex(FPackageObjectIndex::ScriptImport, ScriptObjects.Num());
+			GlobalPackageData.ImportsByFullName.Add(FullName, GlobalImportIndex);
+
+			ScriptImport = &ScriptObjects.AddDefaulted_GetRef();
+			ScriptImport->GlobalIndex = GlobalImportIndex;
+			ScriptImport->FullName = FullName;
+		}
+		else
+		{
+			ScriptImport = &GlobalPackageData.ScriptObjects[FindGlobalImportIndex.GetIndex()];
+		}
+		if (!ScriptImport->bInitialized)
+		{
+			ScriptImport->OuterIndex = FPackageObjectIndex();
+			ScriptImport->ObjectName = ObjectName;
+			ScriptImport->bInitialized = true;
+		}
+
+		FPackageObjectIndex OuterIndex = ScriptImport->GlobalIndex;
+		InnerObjects.Reset();
+		GetObjectsWithOuter(Package, InnerObjects, /*bIncludeNestedObjects*/false);
+		for (UObject* InnerObject : InnerObjects)
+		{
+			FindScriptObjectsRecursive(GlobalPackageData, OuterIndex, InnerObject, TargetPlatform, ExcludedObjectMarks);
+		}
+	}
+
+	for (FScriptObjectData ScriptObject : GlobalPackageData.ScriptObjects)
+	{
+		NameMapBuilder.MarkNameAsReferenced(ScriptObject.ObjectName);
 	}
 }
 
-static void MapImportedPackages(
-	const FGlobalImports& GlobalImports,
-	TArray<FPackage*>& Packages)
+static void CreateGlobalImportsAndExports(
+	TArray<FPackage*>& Packages,
+	const FPackageMap& PackageMap,
+	FPackageAssetData& PackageAssetData,
+	FGlobalPackageData& GlobalPackageData,
+	FExportGraph& ExportGraph)
 {
-	IOSTORE_CPU_SCOPE(MapImportedPackages);
-	UE_LOG(LogIoStore, Display, TEXT("Mapping import packages..."));
+	IOSTORE_CPU_SCOPE(CreateGlobalImportsAndExports);
+	UE_LOG(LogIoStore, Display, TEXT("Creating global imports and exports..."));
+
+	TArray<FString> TempFullNames;
 
-	const FImportObjectsByFullName& ObjectsByFullName = GlobalImports.ObjectsByFullName;
-	TSet<FPackage*> ImportedPackages;
+	TSet<FPackage*> TempImportedPackages;
 	for (FPackage* Package : Packages)
 	{
-		bool bHasMissingImports = false;
+		if (Package->ExportCount == 0)
+		{
+			continue;
+		}
+
+		TempFullNames.Reset();
+		TempFullNames.SetNum(Package->ExportCount, false);
+		FObjectExport* ExportMap = PackageAssetData.ObjectExports.GetData() + Package->ExportIndexOffset;
+		for (int32 ExportIndex = 0; ExportIndex < Package->ExportCount; ExportIndex++)
+		{
+			int32 GlobalExportIndex = FindExport(
+				GlobalPackageData,
+				TempFullNames,
+				PackageAssetData.ObjectExports.GetData() + Package->ExportIndexOffset,
+				ExportIndex,
+				Package);
+
+			FExportObjectData& ExportData = GlobalPackageData.ExportObjects[GlobalExportIndex];
+			ExportData.CreateNode = ExportGraph.AddNode(Package, { uint32(ExportIndex), FExportBundleEntry::ExportCommandType_Create });
+			ExportData.SerializeNode = ExportGraph.AddNode(Package, { uint32(ExportIndex), FExportBundleEntry::ExportCommandType_Serialize });
+
+			Package->Exports.Add(GlobalExportIndex);
+			Package->CreateExportNodes.Add(ExportData.CreateNode);
+			Package->SerializeExportNodes.Add(ExportData.SerializeNode);
+			ExportGraph.AddInternalDependency(ExportData.CreateNode, ExportData.SerializeNode);
+		}
+	}
+
+	for (FExportObjectData& Export : GlobalPackageData.ExportObjects)
+	{
+		const bool bIsPublicExport = !!(Export.ObjectFlags & RF_Public);
+		if (bIsPublicExport)
+		{
+			TArray<int32>& PublicExports = GlobalPackageData.PublicExportIndices;
+			FPackageObjectIndex GlobalImportIndex = FPackageObjectIndex(FPackageObjectIndex::PackageImport, PublicExports.Num());
+			GlobalPackageData.ImportsByFullName.Add(Export.FullName, GlobalImportIndex);
+			PublicExports.Add(Export.GlobalIndex);
+			Export.GlobalImportIndex = GlobalImportIndex;
+			Export.Package->PublicExports.Add(GlobalImportIndex);
+		}
+	}
+
+	for (FPackage* Package : Packages)
+	{
+		if (Package->ImportCount == 0)
+		{
+			continue;
+		}
+
+		FObjectImport* ImportMap = PackageAssetData.ObjectImports.GetData() + Package->ImportIndexOffset;
+		TempFullNames.Reset();
+		TempFullNames.SetNum(Package->ImportCount, false);
 		Package->Imports.Reserve(Package->ImportCount);
 		Package->ImportedPackages.Reserve(Package->ImportCount / 2);
-		ImportedPackages.Reset();
-		for (int32 I = 0; I < Package->ImportCount; ++I)
+		TempImportedPackages.Reset();
+
+		for (int32 ImportIndex = 0; ImportIndex < Package->ImportCount; ++ImportIndex)
 		{
-			FPackageObjectIndex GlobalImportIndex = ObjectsByFullName.FindRef(Package->ImportedFullNames[I]);
+			FindImportFullName(TempFullNames, ImportMap, ImportIndex);
+
+			FPackageObjectIndex GlobalImportIndex = FindAndVerifyGlobalImport(
+				GlobalPackageData,
+				ImportMap[ImportIndex],
+				TempFullNames[ImportIndex]);
+
 			Package->Imports.Add(GlobalImportIndex);
+
 			if (GlobalImportIndex.IsPackageImport())
 			{
-				const FPackageImportData& ImportData = GlobalImports.PackageImports[GlobalImportIndex.GetIndex()];
-				bHasMissingImports |= ImportData.bIsMissingImport;
-				if (ImportData.Package)
+				const FExportObjectData& Export = GlobalPackageData.GetPublicExport(GlobalImportIndex);
+				check(Export.Package);
+
+				bool bAlreadyInSet;
+				TempImportedPackages.Add(Export.Package, &bAlreadyInSet);
+				if (!bAlreadyInSet)
 				{
-					bool bAlreadyInSet;
-					ImportedPackages.Add(ImportData.Package, &bAlreadyInSet);
-					if (!bAlreadyInSet)
-					{
-						Package->ImportedPackages.Add(ImportData.Package);
-						ImportData.Package->ImportedByPackages.Add(Package);
-					}
+					Package->ImportedPackages.Add(Export.Package);
+					Export.Package->ImportedByPackages.Add(Package);
 				}
 			}
 		}
-		UE_CLOG(bHasMissingImports, LogIoStore, Warning, TEXT("Missing imports for package '%s'"), *Package->Name.ToString());
 	}
-};
+}
 
 static void MapExportEntryIndices(
 	TArray<FObjectExport>& ObjectExports,
-	TArray<FExportData>& GlobalExports,
+	TArray<FExportObjectData>& GlobalExports,
 	TArray<FPackage*>& Packages)
 {
 	IOSTORE_CPU_SCOPE(ExportData);
@@ -3052,7 +3219,7 @@ static void MapExportEntryIndices(
 		for (int32 I = 0; I < Package->ExportCount; ++I)
 		{
 			const FObjectExport& ObjectExport = ObjectExports[Package->ExportIndexOffset + I];
-			FExportData& ExportData = GlobalExports[Package->Exports[I]];
+			FExportObjectData& ExportData = GlobalExports[Package->Exports[I]];
 			ExportData.OuterIndex = PackageObjectIndexFromPackageIndex(Package->Imports, ObjectExport.OuterIndex);
 			ExportData.ClassIndex = PackageObjectIndexFromPackageIndex(Package->Imports, ObjectExport.ClassIndex);
 			ExportData.SuperIndex = PackageObjectIndexFromPackageIndex(Package->Imports, ObjectExport.SuperIndex);
@@ -3063,8 +3230,7 @@ static void MapExportEntryIndices(
 static void ProcessLocalizedPackages(
 	const TArray<FPackage*>& Packages,
 	const FPackageMap& PackageMap,
-	const FGlobalImports& GlobalImports,
-	TArray<FExportData>& GlobalExports,
+	FGlobalPackageData& GlobalPackageData,
 	FCulturePackageMap& OutCulturePackageMap,
 	FSourceToLocalizedPackageMultimap& OutSourceToLocalizedPackageMap)
 {
@@ -3105,8 +3271,8 @@ static void ProcessLocalizedPackages(
 		Package->SourceGlobalPackageId = SourcePackage->GlobalPackageId;
 
 		Package->bIsLocalizedAndConformed = ConformLocalizedPackage(
-			PackageMap, GlobalImports, *SourcePackage,
-			*Package, GlobalExports, LocalizedToSourceImportIndexMap);
+			PackageMap, GlobalPackageData, *SourcePackage,
+			*Package, LocalizedToSourceImportIndexMap);
 
 		if (Package->bIsLocalizedAndConformed)
 		{
@@ -3164,21 +3330,21 @@ static void ProcessLocalizedPackages(
 		{
 			if (GlobalImportIndex.IsPackageImport())
 			{
-				const FPackageImportData& ImportData = GlobalImports.PackageImports[GlobalImportIndex.GetIndex()];
-				if (ImportData.bIsLocalized)
+				const FExportObjectData& ExportData = GlobalPackageData.GetPublicExport(GlobalImportIndex);
+				if (ExportData.Package->SourcePackageName != ExportData.Package->Name)
 				{
 					const FPackageObjectIndex* SourceGlobalImportIndex = LocalizedToSourceImportIndexMap.Find(GlobalImportIndex);
 					if (SourceGlobalImportIndex)
 					{
 						GlobalImportIndex = *SourceGlobalImportIndex;
 
-						const FPackageImportData& SourceImportData = GlobalImports.PackageImports[SourceGlobalImportIndex->GetIndex()];
+						const FExportObjectData& SourceExportData = GlobalPackageData.GetPublicExport(*SourceGlobalImportIndex);
 						UE_LOG(LogIoStore, Verbose,
 							TEXT("For package '%s' (%d): Remap localized import %s to source import %s (in a conformed localized package)"),
 							*Package->Name.ToString(),
 							Package->GlobalPackageId.ToIndexForDebugging(),
-							*ImportData.FullName,
-							*SourceImportData.FullName);
+							*ExportData.FullName,
+							*SourceExportData.FullName);
 					}
 					else
 					{
@@ -3187,7 +3353,7 @@ static void ProcessLocalizedPackages(
 							TEXT(", either there is no source package or the localized package did not conform to it."),
 							*Package->Name.ToString(),
 							Package->GlobalPackageId.ToIndexForDebugging(),
-							*ImportData.FullName);
+							*ExportData.FullName);
 					}
 				}
 			}
@@ -3195,7 +3361,7 @@ static void ProcessLocalizedPackages(
 	}
 }
 
-static void SaveReleaseVersionMeta(const TCHAR* ReleaseVersionOutputDir, const FNameMapBuilder& GlobalNameMap, const FPackageGlobalIdMap& PackageGlobalIdMap, const FGlobalImports& GlobalImports, TArray<FContainerTargetSpec*>& ContainerTargets)
+static void SaveReleaseVersionMeta(const TCHAR* ReleaseVersionOutputDir, const FNameMapBuilder& GlobalNameMap, const FPackageGlobalIdMap& PackageGlobalIdMap, FGlobalPackageData& GlobalPackageData, const TArray<FContainerTargetSpec*>& ContainerTargets)
 {
 	UE_LOG(LogIoStore, Display, TEXT("Saving release meta data to '%s'"), ReleaseVersionOutputDir);
 
@@ -3232,22 +3398,9 @@ static void SaveReleaseVersionMeta(const TCHAR* ReleaseVersionOutputDir, const F
 	}
 
 	{
-		FString ImportMapOutputPath = FPaths::Combine(ReleaseVersionOutputDir, TEXT("iodispatcher.uimportmap"));
-		TUniquePtr<FArchive> ImportMapArchive(IFileManager::Get().CreateFileWriter(*ImportMapOutputPath));
-
-		int32 ImportCount = GlobalImports.ObjectsByFullName.Num();
-		(*ImportMapArchive) << ImportCount;
-		for (const auto& KV : GlobalImports.ObjectsByFullName)
-		{
-			FString ObjectName = KV.Key;
-			FPackageObjectIndex ObjectIndex = KV.Value;
-			int32 Type = ObjectIndex.GetType();
-			int32 Index = ObjectIndex.GetIndex();
-
-			(*ImportMapArchive) << ObjectName;
-			(*ImportMapArchive) << Type;
-			(*ImportMapArchive) << Index;
-		}
+		FString ImportExportOutputPath = FPaths::Combine(ReleaseVersionOutputDir, TEXT("iodispatcher.uimportexport"));
+		TUniquePtr<FArchive> ImportExportArchive(IFileManager::Get().CreateFileWriter(*ImportExportOutputPath));
+		(*ImportExportArchive) << GlobalPackageData;
 	}
 
 	for (FContainerTargetSpec* ContainerTarget : ContainerTargets)
@@ -3275,7 +3428,7 @@ static void LoadReleaseVersionMeta(
 	const TCHAR* ReleaseVersionOutputDir,
 	FNameMapBuilder& GlobalNameMap,
 	FPackageGlobalIdMap& PackageGlobalIdMap,
-	FGlobalImports& GlobalImports,
+	FGlobalPackageData& GlobalPackageData,
 	TArray<FContainerTargetSpec*>& ContainerTargets,
 	TMap<FName, FContainerTargetSpec*>& ContainerTargetMap)
 {
@@ -3309,50 +3462,11 @@ static void LoadReleaseVersionMeta(
 			check(PackageGlobalIdMap.Num() == PackageCount);
 		}
 	}
-	{
-		FString ImportMapPath = FPaths::Combine(ReleaseVersionOutputDir, TEXT("iodispatcher.uimportmap"));
-		TUniquePtr<FArchive> ImportMapArchive(IFileManager::Get().CreateFileReader(*ImportMapPath));
-		if (ImportMapArchive)
-		{
-			int32 ImportCount = 0;
-			(*ImportMapArchive) << ImportCount;
-			for (int32 ImportIndex = 0; ImportIndex < ImportCount; ++ImportIndex)
-			{
-				FString FullName;
-				int32 Type;
-				int32 Index;
 
-				(*ImportMapArchive) << FullName;
-				(*ImportMapArchive) << Type;
-				(*ImportMapArchive) << Index;
-
-				FPackageObjectIndex ObjectIndex(static_cast<FPackageObjectIndex::Type>(Type), Index);
-				GlobalImports.ObjectsByFullName.Add(FullName, ObjectIndex);
-				if (Type == FPackageObjectIndex::ScriptImport)
-				{
-					TArray<FScriptImportData>& ScriptImports = GlobalImports.ScriptImports;
-					if (ScriptImports.Num() <= Index)
-					{
-						ScriptImports.AddDefaulted(Index - ScriptImports.Num() + 1);
-					}
-					FScriptImportData& ScriptImport = ScriptImports[Index];
-					ScriptImport.GlobalIndex = ObjectIndex;
-					ScriptImport.OutermostIndex = ObjectIndex;
-					ScriptImport.FullName = FullName;
-				}
-				else
-				{
-					TArray<FPackageImportData>& PackageImports = GlobalImports.PackageImports;
-					if (PackageImports.Num() <= Index)
-					{
-						PackageImports.AddDefaulted(Index - PackageImports.Num() + 1);
-					}
-					FPackageImportData& PackageImport = PackageImports[Index];
-					PackageImport.GlobalIndex = ObjectIndex;
-					PackageImport.FullName = FullName;
-				}
-			}
-		}
+	{
+		FString ImportExportOutputPath = FPaths::Combine(ReleaseVersionOutputDir, TEXT("iodispatcher.uimportexport"));
+		TUniquePtr<FArchive> ImportExportArchive(IFileManager::Get().CreateFileWriter(*ImportExportOutputPath));
+		(*ImportExportArchive) << GlobalPackageData;
 	}
 
 	FString ContainerMetaWildcard = FPaths::Combine(ReleaseVersionOutputDir, TEXT("*.ucontainermeta"));
@@ -3574,7 +3688,6 @@ int32 CreateTarget(const FIoStoreArguments& Arguments, const FIoStoreWriterSetti
 	FNameMapBuilder GlobalNameMapBuilder;
 	FPackageAssetData PackageAssetData;
 	FGlobalPackageData GlobalPackageData;
-	FExportGraph ExportGraph;
 
 	TArray<FPackage*> Packages;
 	FPackageMap PackageMap;
@@ -3592,7 +3705,7 @@ int32 CreateTarget(const FIoStoreArguments& Arguments, const FIoStoreWriterSetti
 
 		if (!Arguments.BasedOnReleaseVersionDir.IsEmpty())
 		{
-			LoadReleaseVersionMeta(*Arguments.BasedOnReleaseVersionDir, GlobalNameMapBuilder, PackageGlobalIdMap, GlobalPackageData.Imports, ContainerTargets, ContainerTargetMap);
+			LoadReleaseVersionMeta(*Arguments.BasedOnReleaseVersionDir, GlobalNameMapBuilder, PackageGlobalIdMap, GlobalPackageData, ContainerTargets, ContainerTargetMap);
 #if OUTPUT_DEBUG_PACKAGE_HASHES
 			FString PackageHashesOutputPath = FPaths::Combine(*Arguments.BasedOnReleaseVersionDir, TEXT("iodispatcher.upackagehashes"));
 			TUniquePtr<FArchive> PackageHashesArchive(IFileManager::Get().CreateFileReader(*PackageHashesOutputPath));
@@ -3619,22 +3732,24 @@ int32 CreateTarget(const FIoStoreArguments& Arguments, const FIoStoreWriterSetti
 		InitializeContainerTargetsAndPackages(Arguments, Packages, PackageMap, PackageGlobalIdMap, ContainerTargets, ContainerTargetMap, GlobalNameMapBuilder);
 	}
 
-	ParsePackageAssets(Packages, PackageAssetData, ExportGraph);
-	CreateGlobalImportsAndExports(Packages, PackageMap, PackageAssetData, GlobalPackageData, ExportGraph);
+	ParsePackageAssets(Packages, PackageAssetData);
 
-	FGlobalPackageImports& PackageImports = GlobalPackageData.Imports.PackageImports;
-	FGlobalScriptImports& ScriptImports = GlobalPackageData.Imports.ScriptImports;
-	TArray<FExportData>& GlobalExports = GlobalPackageData.Exports.Objects;
+	FExportGraph ExportGraph(PackageAssetData.ObjectExports.Num(), PackageAssetData.PreloadDependencies.Num());
+	GlobalPackageData.Reserve(PackageAssetData.ObjectExports.Num());
+
+	TArray<UPackage*> ScriptPackages;
+	FindScriptPackages(PackageAssetData.ObjectImports, ScriptPackages);
+	CreateGlobalScriptObjects(ScriptPackages, GlobalNameMapBuilder, GlobalPackageData, Arguments.TargetPlatform);
+	CreateGlobalImportsAndExports(Packages, PackageMap, PackageAssetData, GlobalPackageData, ExportGraph);
 
 	// Mapped import and exports are required before processing localization, and preload/postload arcs
-	MapImportedPackages(GlobalPackageData.Imports, Packages);
-	MapExportEntryIndices(PackageAssetData.ObjectExports, GlobalExports, Packages);
+	MapExportEntryIndices(PackageAssetData.ObjectExports, GlobalPackageData.ExportObjects, Packages);
 
 	FSourceToLocalizedPackageMultimap SourceToLocalizedPackageMap;
 	FCulturePackageMap CulturePackageMap;
 
 	ProcessLocalizedPackages(
-		Packages, PackageMap, GlobalPackageData.Imports, GlobalExports,
+		Packages, PackageMap, GlobalPackageData, 
 		CulturePackageMap, SourceToLocalizedPackageMap);
 
 	const int32 CircularChainCount = AddPostLoadDependencies(Packages);
@@ -3667,8 +3782,8 @@ int32 CreateTarget(const FIoStoreArguments& Arguments, const FIoStoreWriterSetti
 			FinalizePackageHeaders(
 				*ContainerTarget,
 				PackageAssetData.ObjectExports,
-				GlobalPackageData.Exports.Objects,
-				GlobalPackageData.Imports.ObjectsByFullName);
+				GlobalPackageData.ExportObjects,
+				GlobalPackageData.ImportsByFullName);
 
 			const FNameMapBuilder& NameMapBuilder = *ContainerTarget->NameMapBuilder;
 			SaveNameBatch(ContainerTarget->LocalNameMapBuilder.GetNameMap(), ContainerTarget->Header.Names, ContainerTarget->Header.NameHashes);
@@ -3705,18 +3820,18 @@ int32 CreateTarget(const FIoStoreArguments& Arguments, const FIoStoreWriterSetti
 
 		SerializeInitialLoad(
 			GlobalNameMapBuilder,
-			ScriptImports,
+			GlobalPackageData.ScriptObjects,
 			InitialLoadArchive);
 
 		StoreTocByteCount = StoreTocArchive.TotalSize();
 		StoreDataByteCount = StoreDataArchive.TotalSize();
 
 		int32 NumPackages = Packages.Num();
-		int32 PackageImportCount = PackageImports.Num();
+		int32 NumPublicExports = GlobalPackageData.PublicExportIndices.Num();
 		int32 StoreByteCount = StoreTocByteCount + StoreDataByteCount;
 
 		GlobalMetaArchive << NumPackages;
-		GlobalMetaArchive << PackageImportCount;
+		GlobalMetaArchive << NumPublicExports;
 		GlobalMetaArchive << StoreByteCount;
 		GlobalMetaArchive.Serialize(StoreTocArchive.GetData(), StoreTocByteCount);
 		GlobalMetaArchive.Serialize(StoreDataArchive.GetData(), StoreDataByteCount);
@@ -4071,7 +4186,7 @@ int32 CreateTarget(const FIoStoreArguments& Arguments, const FIoStoreWriterSetti
 
 	if (!Arguments.OutputReleaseVersionDir.IsEmpty())
 	{
-		SaveReleaseVersionMeta(*Arguments.OutputReleaseVersionDir, GlobalNameMapBuilder, PackageGlobalIdMap, GlobalPackageData.Imports, ContainerTargets);
+		SaveReleaseVersionMeta(*Arguments.OutputReleaseVersionDir, GlobalNameMapBuilder, PackageGlobalIdMap, GlobalPackageData, ContainerTargets);
 #if OUTPUT_DEBUG_PACKAGE_HASHES
 		FString PackageHashesOutputPath = FPaths::Combine(*Arguments.OutputReleaseVersionDir, TEXT("iodispatcher.upackagehashes"));
 		TUniquePtr<FArchive> PackageHashesArchive(IFileManager::Get().CreateFileWriter(*PackageHashesOutputPath));
@@ -4198,7 +4313,6 @@ int32 CreateTarget(const FIoStoreArguments& Arguments, const FIoStoreWriterSetti
 		PublicExportsCount += Package.PublicExports.Num();
 		ExportBundlesMetaCount += Package.ExportBundles.Num();
 		NameCount += Package.NameMap.Num();
-		PackagesWithoutPreloadDependenciesCount += Package.ImportedPreloadPackages.Num() == 0;
 		PackagesWithoutImportDependenciesCount += Package.ImportedPackages.Num() == 0;
 
 		for (auto& KV : Package.ExternalArcs)
@@ -4242,9 +4356,9 @@ int32 CreateTarget(const FIoStoreArguments& Arguments, const FIoStoreWriterSetti
 	}
 	UE_LOG(LogIoStore, Display, TEXT("Compression block padding: %8.2f MB"), TotalPaddingSize / 1024.0 / 1024.0);
 	UE_LOG(LogIoStore, Display, TEXT(""));
-	UE_LOG(LogIoStore, Display, TEXT("Packages: %8d total, %d circular dependencies, %d no preload dependencies, %d no import dependencies"),
-		PackageMap.Num(), CircularPackagesCount, PackagesWithoutPreloadDependenciesCount, PackagesWithoutImportDependenciesCount);
-	UE_LOG(LogIoStore, Display, TEXT("Bundles:  %8d total, %d entries, %d export objects"), BundleCount, BundleEntryCount, GlobalPackageData.Exports.Objects.Num());
+	UE_LOG(LogIoStore, Display, TEXT("Packages: %8d total, %d circular dependencies, %d no import dependencies"),
+		PackageMap.Num(), CircularPackagesCount, PackagesWithoutImportDependenciesCount);
+	UE_LOG(LogIoStore, Display, TEXT("Bundles:  %8d total, %d entries, %d export objects"), BundleCount, BundleEntryCount, GlobalPackageData.ExportObjects.Num());
 
 	UE_LOG(LogIoStore, Display, TEXT("IoStore: %8.2f MB GlobalNames, %d unique names"), (double)GlobalNamesMB, GlobalNameMapBuilder.GetNameMap().Num());
 	UE_LOG(LogIoStore, Display, TEXT("IoStore: %8.2f MB GlobalNameHashes"), (double)GlobalNameHashesMB);
@@ -4255,8 +4369,8 @@ int32 CreateTarget(const FIoStoreArguments& Arguments, const FIoStoreWriterSetti
 		PublicExportsCount,
 		ExportBundlesMetaCount);
 	UE_LOG(LogIoStore, Display, TEXT("IoStore: %8.2f MB PackageImportNames, %d imports"),
-		(double)GlobalImportNamesByteCount / 1024.0 / 1024.0, PackageImports.Num());
-	UE_LOG(LogIoStore, Display, TEXT("IoStore: %8.2f MB InitialLoadData, %d script imports"), (double)InitialLoadSize / 1024.0 / 1024.0, ScriptImports.Num());
+		(double)GlobalImportNamesByteCount / 1024.0 / 1024.0, GlobalPackageData.ImportsByFullName.Num());
+	UE_LOG(LogIoStore, Display, TEXT("IoStore: %8.2f MB InitialLoadData, %d script objects (%d packages)"), (double)InitialLoadSize / 1024.0 / 1024.0, GlobalPackageData.ScriptObjects.Num(), ScriptPackages.Num());
 	UE_LOG(LogIoStore, Display, TEXT("IoStore: %8.2f MB PackageHeader, %d packages"), (double)PackageHeaderSize / 1024.0 / 1024.0, Packages.Num());
 	UE_LOG(LogIoStore, Display, TEXT("IoStore: %8.2f MB PackageSummary"), (double)PackageSummarySize / 1024.0 / 1024.0);
 	UE_LOG(LogIoStore, Display, TEXT("IoStore: %8.2f MB PackageNameMap, %d indices"), (double)NameMapSize / 1024.0 / 1024.0, NameMapCount);
@@ -4541,6 +4655,24 @@ int32 CreateIoStoreContainerFiles(const TCHAR* CmdLine)
 	FParse::Value(CmdLine, TEXT("-CreateReleaseVersionDirectory="), Arguments.OutputReleaseVersionDir);
 	FParse::Value(CmdLine, TEXT("-BasedOnReleaseVersionDirectory="), Arguments.BasedOnReleaseVersionDir);
 
+	FString TargetPlatform;
+	if (FParse::Value(FCommandLine::Get(), TEXT("TargetPlatform="), TargetPlatform))
+	{
+		UE_LOG(LogIoStore, Display, TEXT("Using target platform '%s'"), *TargetPlatform);
+		ITargetPlatformManagerModule& TPM = GetTargetPlatformManagerRef();
+		Arguments.TargetPlatform = TPM.FindTargetPlatform(TargetPlatform);
+		if (!Arguments.TargetPlatform)
+		{
+			UE_LOG(LogIoStore, Error, TEXT("Invalid TargetPlatform: '%s'"), *TargetPlatform);
+			return 1;
+		}
+	}
+	else
+	{
+		UE_LOG(LogIoStore, Error, TEXT("TargetPlatform must be specified"));
+		return 1;
+	}
+
 	if (FParse::Value(FCommandLine::Get(), TEXT("CreateGlobalContainer="), Arguments.GlobalContainerPath))
 	{
 		Arguments.GlobalContainerPath = FPaths::ChangeExtension(Arguments.GlobalContainerPath, TEXT(""));
diff --git a/Engine/Source/Programs/AutomationTool/Scripts/CopyBuildToStagingDirectory.Automation.cs b/Engine/Source/Programs/AutomationTool/Scripts/CopyBuildToStagingDirectory.Automation.cs
--- a/Engine/Source/Programs/AutomationTool/Scripts/CopyBuildToStagingDirectory.Automation.cs
+++ b/Engine/Source/Programs/AutomationTool/Scripts/CopyBuildToStagingDirectory.Automation.cs
@@ -2629,6 +2629,7 @@ public partial class Project : CommandUtils
 		{
 			CommandletParams += String.Format(" -BasedOnReleaseVersionDirectory={0}", MakePathSafeToUseWithCommandLine(Params.GetBasedOnReleaseVersionPath(SC, Params.Client)));
 		}
+		CommandletParams += String.Format(" -TargetPlatform={0}", SC.StageTargetPlatform.GetCookPlatform(Params.DedicatedServer, Params.Client));
 
 		LogInformation("Running IoStore commandlet with arguments: {0}", CommandletParams);
 		RunCommandlet(SC.RawProjectPath, Params.UE4Exe, "IoStore", CommandletParams);
