IoStore: Add missing include

diff --git a/Engine/Source/Developer/IoStoreUtilities/Private/IoStoreUtilities.cpp b/Engine/Source/Developer/IoStoreUtilities/Private/IoStoreUtilities.cpp
--- a/Engine/Source/Developer/IoStoreUtilities/Private/IoStoreUtilities.cpp
+++ b/Engine/Source/Developer/IoStoreUtilities/Private/IoStoreUtilities.cpp
@@ -23,6 +23,7 @@
 #include "Serialization/LargeMemoryWriter.h"
 #include "Serialization/MemoryReader.h"
 #include "Serialization/AsyncLoading2.h"
+#include "UObject/Class.h"
 #include "UObject/NameBatchSerialization.h"
 #include "UObject/PackageFileSummary.h"
 #include "UObject/ObjectResource.h"
