IoStore: Fix linux compile error

diff --git a/Engine/Source/Developer/IoStoreUtilities/Private/IoStoreUtilities.cpp b/Engine/Source/Developer/IoStoreUtilities/Private/IoStoreUtilities.cpp
--- a/Engine/Source/Developer/IoStoreUtilities/Private/IoStoreUtilities.cpp
+++ b/Engine/Source/Developer/IoStoreUtilities/Private/IoStoreUtilities.cpp
@@ -1764,16 +1764,16 @@ struct FScriptObjectData
 
 	friend FArchive& operator<<(FArchive& Ar, FScriptObjectData& Data)
 	{
-		FString ObjectName;
+		FString ObjectNameStr;
 		if (Ar.IsSaving())
 		{
-			ObjectName = Data.ObjectName.ToString();
-			Ar << ObjectName;
+			ObjectNameStr = Data.ObjectName.ToString();
+			Ar << ObjectNameStr;
 		}
 		else
 		{
-			Ar << ObjectName;
-			Data.ObjectName = FName(*ObjectName);
+			Ar << ObjectNameStr;
+			Data.ObjectName = FName(*ObjectNameStr);
 		}
 		Ar << Data.FullName;
 		Ar << Data.GlobalIndex;
@@ -1802,16 +1802,16 @@ struct FExportObjectData
 
 	friend FArchive& operator<<(FArchive& Ar, FExportObjectData& Data)
 	{
-		FString ObjectName;
+		FString ObjectNameStr;
 		if (Ar.IsSaving())
 		{
-			ObjectName = Data.ObjectName.ToString();
-			Ar << ObjectName;
+			ObjectNameStr = Data.ObjectName.ToString();
+			Ar << ObjectNameStr;
 		}
 		else
 		{
-			Ar << ObjectName;
-			Data.ObjectName = FName(*ObjectName);
+			Ar << ObjectNameStr;
+			Data.ObjectName = FName(*ObjectNameStr);
 		}
 		Ar << Data.FullName;
 		Ar << Data.GlobalIndex;
