IoDispatcher: Encryption support.
- updated IoStoreWriter to encrypt blocks (not assets)
- updated automation scripts to pass relevant encryption commands to IoStore comandlet
- added encryption guid to container TOC header
- added encryption guid to container meta when loading/saving release information
- updated IoStore command to enable container local namemap when encryption is enabled
- updated IoDispatcher file backend to decrypt blocks
- added encryption key cache to FileIoStore

diff --git a/Engine/Source/Developer/IoStoreUtilities/Private/IoStoreUtilities.cpp b/Engine/Source/Developer/IoStoreUtilities/Private/IoStoreUtilities.cpp
--- a/Engine/Source/Developer/IoStoreUtilities/Private/IoStoreUtilities.cpp
+++ b/Engine/Source/Developer/IoStoreUtilities/Private/IoStoreUtilities.cpp
@@ -14,6 +14,9 @@
 #include "Misc/Paths.h"
 #include "Misc/PackageName.h"
 #include "Misc/Base64.h"
+#include "Misc/AES.h"
+#include "Misc/CoreDelegates.h"
+#include "Misc/IEngineCrypto.h"
 #include "Modules/ModuleManager.h"
 #include "Serialization/Archive.h"
 #include "Serialization/BulkDataManifest.h"
@@ -35,6 +38,7 @@
 #include "Async/ParallelFor.h"
 #include "Async/AsyncFileHandle.h"
 #include "Async/Async.h"
+#include "RSA.h"
 
 //PRAGMA_DISABLE_OPTIMIZATION
 
@@ -54,6 +58,317 @@ static const FName DefaultCompressionMethod = NAME_Zlib;
 static const int64 DefaultCompressionBlockSize = 64 << 10;
 const int64 DefaultMemoryMappingAlignment = 16 << 10;
 
+struct FNamedAESKey
+{
+	FString Name;
+	FGuid Guid;
+	FAES::FAESKey Key;
+
+	bool IsValid() const
+	{
+		return Key.IsValid();
+	}
+};
+
+struct FKeyChain
+{
+	FRSAKeyHandle SigningKey = InvalidRSAKeyHandle;
+	TMap<FGuid, FNamedAESKey> EncryptionKeys;
+	const FNamedAESKey* MasterEncryptionKey = nullptr;
+};
+
+static void ApplyEncryptionKeys(const FKeyChain& KeyChain)
+{
+	if (KeyChain.EncryptionKeys.Contains(FGuid()))
+	{
+		FAES::FAESKey DefaultKey = KeyChain.EncryptionKeys[FGuid()].Key;
+		FCoreDelegates::GetPakEncryptionKeyDelegate().BindLambda([DefaultKey](uint8 OutKey[32]) { FMemory::Memcpy(OutKey, DefaultKey.Key, sizeof(DefaultKey.Key)); });
+	}
+
+	for (const TMap<FGuid, FNamedAESKey>::ElementType& Key : KeyChain.EncryptionKeys)
+	{
+		if (Key.Key.IsValid())
+		{
+			FCoreDelegates::GetRegisterEncryptionKeyDelegate().ExecuteIfBound(Key.Key, Key.Value.Key);
+		}
+	}
+}
+
+static FRSAKeyHandle ParseRSAKeyFromJson(TSharedPtr<FJsonObject> InObj)
+{
+	TSharedPtr<FJsonObject> PublicKey = InObj->GetObjectField(TEXT("PublicKey"));
+	TSharedPtr<FJsonObject> PrivateKey = InObj->GetObjectField(TEXT("PrivateKey"));
+
+	FString PublicExponentBase64, PrivateExponentBase64, PublicModulusBase64, PrivateModulusBase64;
+
+	if (   PublicKey->TryGetStringField("Exponent", PublicExponentBase64)
+		&& PublicKey->TryGetStringField("Modulus", PublicModulusBase64)
+		&& PrivateKey->TryGetStringField("Exponent", PrivateExponentBase64)
+		&& PrivateKey->TryGetStringField("Modulus", PrivateModulusBase64))
+	{
+		check(PublicModulusBase64 == PrivateModulusBase64);
+
+		TArray<uint8> PublicExponent, PrivateExponent, Modulus;
+		FBase64::Decode(PublicExponentBase64, PublicExponent);
+		FBase64::Decode(PrivateExponentBase64, PrivateExponent);
+		FBase64::Decode(PublicModulusBase64, Modulus);
+
+		return FRSA::CreateKey(PublicExponent, PrivateExponent, Modulus);
+	}
+	else
+	{
+		return nullptr;
+	}
+}
+
+static void LoadKeyChainFromFile(const FString& InFilename, FKeyChain& OutCryptoSettings)
+{
+	TUniquePtr<FArchive> File;
+	File.Reset(IFileManager::Get().CreateFileReader(*InFilename));
+
+	UE_CLOG(File == nullptr, LogIoStore, Fatal, TEXT("Specified crypto keys cache '%s' does not exist!"), *InFilename);
+	TSharedPtr<FJsonObject> RootObject;
+	TSharedRef<TJsonReader<char>> Reader = TJsonReaderFactory<char>::Create(File.Get());
+	if (FJsonSerializer::Deserialize(Reader, RootObject))
+	{
+		const TSharedPtr<FJsonObject>* EncryptionKeyObject;
+		if (RootObject->TryGetObjectField(TEXT("EncryptionKey"), EncryptionKeyObject))
+		{
+			FString EncryptionKeyBase64;
+			if ((*EncryptionKeyObject)->TryGetStringField(TEXT("Key"), EncryptionKeyBase64))
+			{
+				if (EncryptionKeyBase64.Len() > 0)
+				{
+					TArray<uint8> Key;
+					FBase64::Decode(EncryptionKeyBase64, Key);
+					check(Key.Num() == sizeof(FAES::FAESKey::Key));
+					FNamedAESKey NewKey;
+					NewKey.Name = TEXT("Default");
+					NewKey.Guid = FGuid();
+					FMemory::Memcpy(NewKey.Key.Key, &Key[0], sizeof(FAES::FAESKey::Key));
+					OutCryptoSettings.EncryptionKeys.Add(NewKey.Guid, NewKey);
+				}
+			}
+		}
+
+		const TSharedPtr<FJsonObject>* SigningKey = nullptr;
+		if (RootObject->TryGetObjectField(TEXT("SigningKey"), SigningKey))
+		{
+			OutCryptoSettings.SigningKey = ParseRSAKeyFromJson(*SigningKey);
+		}
+
+		const TArray<TSharedPtr<FJsonValue>>* SecondaryEncryptionKeyArray = nullptr;
+		if (RootObject->TryGetArrayField(TEXT("SecondaryEncryptionKeys"), SecondaryEncryptionKeyArray))
+		{
+			for (TSharedPtr<FJsonValue> EncryptionKeyValue : *SecondaryEncryptionKeyArray)
+			{
+				FNamedAESKey NewKey;
+				TSharedPtr<FJsonObject> SecondaryEncryptionKeyObject = EncryptionKeyValue->AsObject();
+				FGuid::Parse(SecondaryEncryptionKeyObject->GetStringField(TEXT("Guid")), NewKey.Guid);
+				NewKey.Name = SecondaryEncryptionKeyObject->GetStringField(TEXT("Name"));
+				FString KeyBase64 = SecondaryEncryptionKeyObject->GetStringField(TEXT("Key"));
+
+				TArray<uint8> Key;
+				FBase64::Decode(KeyBase64, Key);
+				check(Key.Num() == sizeof(FAES::FAESKey::Key));
+				FMemory::Memcpy(NewKey.Key.Key, &Key[0], sizeof(FAES::FAESKey::Key));
+
+				check(!OutCryptoSettings.EncryptionKeys.Contains(NewKey.Guid) || OutCryptoSettings.EncryptionKeys[NewKey.Guid].Key == NewKey.Key);
+				OutCryptoSettings.EncryptionKeys.Add(NewKey.Guid, NewKey);
+			}
+		}
+		UE_LOG(LogIoStore, Display, TEXT("Parsed '%d' crypto keys from '%s'"), OutCryptoSettings.EncryptionKeys.Num(), *InFilename);
+	}
+	FGuid EncryptionKeyOverrideGuid;
+	OutCryptoSettings.MasterEncryptionKey = OutCryptoSettings.EncryptionKeys.Find(EncryptionKeyOverrideGuid);
+}
+
+static void LoadKeyChain(const TCHAR* CmdLine, FKeyChain& OutCryptoSettings)
+{
+	OutCryptoSettings.SigningKey = InvalidRSAKeyHandle;
+	OutCryptoSettings.EncryptionKeys.Empty();
+
+	// First, try and parse the keys from a supplied crypto key cache file
+	FString CryptoKeysCacheFilename;
+	if (FParse::Value(CmdLine, TEXT("cryptokeys="), CryptoKeysCacheFilename))
+	{
+		UE_LOG(LogIoStore, Display, TEXT("Parsing crypto keys from a crypto key cache file '%s'"), *CryptoKeysCacheFilename);
+		LoadKeyChainFromFile(CryptoKeysCacheFilename, OutCryptoSettings);
+	}
+	else if (FParse::Param(CmdLine, TEXT("encryptionini")))
+	{
+		FString ProjectDir, EngineDir, Platform;
+
+		if (FParse::Value(CmdLine, TEXT("projectdir="), ProjectDir, false)
+			&& FParse::Value(CmdLine, TEXT("enginedir="), EngineDir, false)
+			&& FParse::Value(CmdLine, TEXT("platform="), Platform, false))
+		{
+			UE_LOG(LogIoStore, Warning, TEXT("A legacy command line syntax is being used for crypto config. Please update to using the -cryptokey parameter as soon as possible as this mode is deprecated"));
+
+			FConfigFile EngineConfig;
+
+			FConfigCacheIni::LoadExternalIniFile(EngineConfig, TEXT("Engine"), *FPaths::Combine(EngineDir, TEXT("Config\\")), *FPaths::Combine(ProjectDir, TEXT("Config/")), true, *Platform);
+			bool bDataCryptoRequired = false;
+			EngineConfig.GetBool(TEXT("PlatformCrypto"), TEXT("PlatformRequiresDataCrypto"), bDataCryptoRequired);
+
+			if (!bDataCryptoRequired)
+			{
+				return;
+			}
+
+			FConfigFile ConfigFile;
+			FConfigCacheIni::LoadExternalIniFile(ConfigFile, TEXT("Crypto"), *FPaths::Combine(EngineDir, TEXT("Config\\")), *FPaths::Combine(ProjectDir, TEXT("Config/")), true, *Platform);
+			bool bSignPak = false;
+			bool bEncryptPakIniFiles = false;
+			bool bEncryptPakIndex = false;
+			bool bEncryptAssets = false;
+			bool bEncryptPak = false;
+
+			if (ConfigFile.Num())
+			{
+				UE_LOG(LogIoStore, Display, TEXT("Using new format crypto.ini files for crypto configuration"));
+
+				static const TCHAR* SectionName = TEXT("/Script/CryptoKeys.CryptoKeysSettings");
+
+				ConfigFile.GetBool(SectionName, TEXT("bEnablePakSigning"), bSignPak);
+				ConfigFile.GetBool(SectionName, TEXT("bEncryptPakIniFiles"), bEncryptPakIniFiles);
+				ConfigFile.GetBool(SectionName, TEXT("bEncryptPakIndex"), bEncryptPakIndex);
+				ConfigFile.GetBool(SectionName, TEXT("bEncryptAssets"), bEncryptAssets);
+				bEncryptPak = bEncryptPakIniFiles || bEncryptPakIndex || bEncryptAssets;
+
+				if (bSignPak)
+				{
+					FString PublicExpBase64, PrivateExpBase64, ModulusBase64;
+					ConfigFile.GetString(SectionName, TEXT("SigningPublicExponent"), PublicExpBase64);
+					ConfigFile.GetString(SectionName, TEXT("SigningPrivateExponent"), PrivateExpBase64);
+					ConfigFile.GetString(SectionName, TEXT("SigningModulus"), ModulusBase64);
+
+					TArray<uint8> PublicExp, PrivateExp, Modulus;
+					FBase64::Decode(PublicExpBase64, PublicExp);
+					FBase64::Decode(PrivateExpBase64, PrivateExp);
+					FBase64::Decode(ModulusBase64, Modulus);
+
+					OutCryptoSettings.SigningKey = FRSA::CreateKey(PublicExp, PrivateExp, Modulus);
+
+					UE_LOG(LogIoStore, Display, TEXT("Parsed signature keys from config files."));
+				}
+
+				if (bEncryptPak)
+				{
+					FString EncryptionKeyString;
+					ConfigFile.GetString(SectionName, TEXT("EncryptionKey"), EncryptionKeyString);
+
+					if (EncryptionKeyString.Len() > 0)
+					{
+						TArray<uint8> Key;
+						FBase64::Decode(EncryptionKeyString, Key);
+						check(Key.Num() == sizeof(FAES::FAESKey::Key));
+						FNamedAESKey NewKey;
+						NewKey.Name = TEXT("Default");
+						NewKey.Guid = FGuid();
+						FMemory::Memcpy(NewKey.Key.Key, &Key[0], sizeof(FAES::FAESKey::Key));
+						OutCryptoSettings.EncryptionKeys.Add(NewKey.Guid, NewKey);
+						UE_LOG(LogIoStore, Display, TEXT("Parsed AES encryption key from config files."));
+					}
+				}
+			}
+			else
+			{
+				static const TCHAR* SectionName = TEXT("Core.Encryption");
+
+				UE_LOG(LogIoStore, Display, TEXT("Using old format encryption.ini files for crypto configuration"));
+
+				FConfigCacheIni::LoadExternalIniFile(ConfigFile, TEXT("Encryption"), *FPaths::Combine(EngineDir, TEXT("Config\\")), *FPaths::Combine(ProjectDir, TEXT("Config/")), true, *Platform);
+				ConfigFile.GetBool(SectionName, TEXT("SignPak"), bSignPak);
+				ConfigFile.GetBool(SectionName, TEXT("EncryptPak"), bEncryptPak);
+
+				if (bSignPak)
+				{
+					FString RSAPublicExp, RSAPrivateExp, RSAModulus;
+					ConfigFile.GetString(SectionName, TEXT("rsa.publicexp"), RSAPublicExp);
+					ConfigFile.GetString(SectionName, TEXT("rsa.privateexp"), RSAPrivateExp);
+					ConfigFile.GetString(SectionName, TEXT("rsa.modulus"), RSAModulus);
+
+					//TODO: Fix me!
+					//OutSigningKey.PrivateKey.Exponent.Parse(RSAPrivateExp);
+					//OutSigningKey.PrivateKey.Modulus.Parse(RSAModulus);
+					//OutSigningKey.PublicKey.Exponent.Parse(RSAPublicExp);
+					//OutSigningKey.PublicKey.Modulus = OutSigningKey.PrivateKey.Modulus;
+
+					UE_LOG(LogIoStore, Display, TEXT("Parsed signature keys from config files."));
+				}
+
+				if (bEncryptPak)
+				{
+					FString EncryptionKeyString;
+					ConfigFile.GetString(SectionName, TEXT("aes.key"), EncryptionKeyString);
+					FNamedAESKey NewKey;
+					NewKey.Name = TEXT("Default");
+					NewKey.Guid = FGuid();
+					if (EncryptionKeyString.Len() == 32 && TCString<TCHAR>::IsPureAnsi(*EncryptionKeyString))
+					{
+						for (int32 Index = 0; Index < 32; ++Index)
+						{
+							NewKey.Key.Key[Index] = (uint8)EncryptionKeyString[Index];
+						}
+						OutCryptoSettings.EncryptionKeys.Add(NewKey.Guid, NewKey);
+						UE_LOG(LogIoStore, Display, TEXT("Parsed AES encryption key from config files."));
+					}
+				}
+			}
+		}
+	}
+	else
+	{
+		UE_LOG(LogIoStore, Display, TEXT("Using command line for crypto configuration"));
+
+		FString EncryptionKeyString;
+		FParse::Value(CmdLine, TEXT("aes="), EncryptionKeyString, false);
+
+		if (EncryptionKeyString.Len() > 0)
+		{
+			UE_LOG(LogIoStore, Warning, TEXT("A legacy command line syntax is being used for crypto config. Please update to using the -cryptokey parameter as soon as possible as this mode is deprecated"));
+
+			FNamedAESKey NewKey;
+			NewKey.Name = TEXT("Default");
+			NewKey.Guid = FGuid();
+			const uint32 RequiredKeyLength = sizeof(NewKey.Key);
+
+			// Error checking
+			if (EncryptionKeyString.Len() < RequiredKeyLength)
+			{
+				UE_LOG(LogIoStore, Fatal, TEXT("AES encryption key must be %d characters long"), RequiredKeyLength);
+			}
+
+			if (EncryptionKeyString.Len() > RequiredKeyLength)
+			{
+				UE_LOG(LogIoStore, Warning, TEXT("AES encryption key is more than %d characters long, so will be truncated!"), RequiredKeyLength);
+				EncryptionKeyString.LeftInline(RequiredKeyLength);
+			}
+
+			if (!FCString::IsPureAnsi(*EncryptionKeyString))
+			{
+				UE_LOG(LogIoStore, Fatal, TEXT("AES encryption key must be a pure ANSI string!"));
+			}
+
+			ANSICHAR* AsAnsi = TCHAR_TO_ANSI(*EncryptionKeyString);
+			check(TCString<ANSICHAR>::Strlen(AsAnsi) == RequiredKeyLength);
+			FMemory::Memcpy(NewKey.Key.Key, AsAnsi, RequiredKeyLength);
+			OutCryptoSettings.EncryptionKeys.Add(NewKey.Guid, NewKey);
+			UE_LOG(LogIoStore, Display, TEXT("Parsed AES encryption key from command line."));
+		}
+	}
+
+	FString EncryptionKeyOverrideGuidString;
+	FGuid EncryptionKeyOverrideGuid;
+	if (FParse::Value(CmdLine, TEXT("EncryptionKeyOverrideGuid="), EncryptionKeyOverrideGuidString))
+	{
+		UE_LOG(LogIoStore, Display, TEXT("Using encryption key override '%s'"), *EncryptionKeyOverrideGuidString);
+		FGuid::Parse(EncryptionKeyOverrideGuidString, EncryptionKeyOverrideGuid);
+	}
+	OutCryptoSettings.MasterEncryptionKey = OutCryptoSettings.EncryptionKeys.Find(EncryptionKeyOverrideGuid);
+}
+
 class FNameMapBuilder
 {
 public:
@@ -236,6 +551,7 @@ struct FContainerSourceFile
 {
 	FString NormalizedPath;
 	bool bNeedsCompression = false;
+	bool bNeedsEncryption = false;
 };
 
 struct FContainerSourceSpec
@@ -303,19 +619,25 @@ struct FIoStoreArguments
 	TMap<FName, uint64> GameOrderMap;
 	TMap<FName, uint64> CookerOrderMap;
 	int64 MemoryMappingAlignment = 0;
+	FKeyChain KeyChain;
+	FKeyChain PatchKeyChain;
 };
 
 struct FContainerMeta
 {
 	FString ContainerName;
 	FIoContainerId ContainerId;
+	FGuid EncryptionKeyGuid;
 	FNameMapBuilder NameMapBuilder;
+	bool bIsEncrypted;
 
 	friend FArchive& operator<<(FArchive& Ar, FContainerMeta& ContainerMeta)
 	{
 		Ar << ContainerMeta.ContainerName;
 		Ar << ContainerMeta.ContainerId;
+		Ar << ContainerMeta.EncryptionKeyGuid;
 		Ar << ContainerMeta.NameMapBuilder;
+		Ar << ContainerMeta.bIsEncrypted;
 		return Ar;
 	}
 };
@@ -324,6 +646,7 @@ struct FContainerTargetSpec
 {
 	FContainerHeader Header;
 	FName Name;
+	FGuid EncryptionKeyGuid;
 	FString OutputPath;
 	FIoStoreWriter* IoStoreWriter;
 	TArray<FContainerTargetFile> TargetFiles;
@@ -334,6 +657,7 @@ struct FContainerTargetSpec
 	bool bIsCompressed = false;
 	bool bUseLocalNameMap = false;
 	bool bGenerateDiffPatch = false;
+	bool bIsEncrypted = false;
 };
 
 struct FPackageAssetData
@@ -3409,9 +3733,10 @@ static void SaveReleaseVersionMeta(const TCHAR* ReleaseVersionOutputDir, const F
 		FContainerMeta ContainerMeta;
 		ContainerMeta.ContainerId = ContainerTarget->Header.ContainerId;
 		ContainerMeta.ContainerName = ContainerTarget->Name.ToString();
-
+		ContainerMeta.EncryptionKeyGuid = ContainerTarget->EncryptionKeyGuid;
 		// Should be safe now as all container(s) has been saved.
 		ContainerMeta.NameMapBuilder = MoveTemp(ContainerTarget->LocalNameMapBuilder);
+		ContainerMeta.bIsEncrypted = ContainerTarget->bIsEncrypted;
 
 		FString MetaOutputPath = FPaths::Combine(ReleaseVersionOutputDir, ContainerTarget->Name.ToString() + TEXT(".ucontainermeta"));
 		TUniquePtr<FArchive> Ar(IFileManager::Get().CreateFileWriter(*MetaOutputPath));
@@ -3483,7 +3808,9 @@ static void LoadReleaseVersionMeta(
 			*Ar << ContainerMeta;
 
 			FContainerTargetSpec* ContainerTarget = AddContainer(FName(ContainerMeta.ContainerName), ContainerMeta.ContainerId, ContainerTargets, ContainerTargetMap);
+			ContainerTarget->EncryptionKeyGuid = ContainerMeta.EncryptionKeyGuid;
 			ContainerTarget->LocalNameMapBuilder = MoveTemp(ContainerMeta.NameMapBuilder);
+			ContainerTarget->bIsEncrypted = ContainerMeta.bIsEncrypted;
 
 			UE_LOG(LogIoStore, Display,
 				TEXT("Loaded container release meta '%s' with container ID '%d' and '%d' names"),
@@ -3548,7 +3875,19 @@ void InitializeContainerTargetsAndPackages(
 			FIoStoreEnvironment PatchSourceEnvironment;
 			PatchSourceEnvironment.InitializeFileEnvironment(FPaths::ChangeExtension(PatchSourceContainerFile, TEXT("")));
 			TUniquePtr<FIoStoreReader> PatchSourceReader(new FIoStoreReader());
-			FIoStatus Status = PatchSourceReader->Initialize(PatchSourceEnvironment);
+
+			FIoEncryptionKey EncryptionKey;
+			if (ContainerTarget->bIsEncrypted)
+			{
+				const FNamedAESKey* PatchKey = Arguments.PatchKeyChain.EncryptionKeys.Find(ContainerTarget->EncryptionKeyGuid);
+				if (!PatchKey)
+				{
+					PatchKey = Arguments.KeyChain.EncryptionKeys.Find(ContainerTarget->EncryptionKeyGuid);
+				}
+				UE_CLOG(!PatchKey, LogIoStore, Fatal, TEXT("Failed to find encryption key '%s' for container '%s'"), *ContainerTarget->EncryptionKeyGuid.ToString(), *ContainerTarget->Name.ToString());
+				EncryptionKey = FIoEncryptionKey { ContainerTarget->EncryptionKeyGuid, PatchKey->Key };
+			}
+			FIoStatus Status = PatchSourceReader->Initialize(PatchSourceEnvironment, EncryptionKey);
 			if (Status.IsOk())
 			{
 				UE_LOG(LogIoStore, Display, TEXT("Loaded patch source container '%s'"), *PatchSourceContainerFile);
@@ -3614,6 +3953,11 @@ void InitializeContainerTargetsAndPackages(
 					{
 						TargetFile.bForceUncompressed = true;
 					}
+					if (SourceFile.bNeedsEncryption)
+					{
+						ContainerTarget->bIsEncrypted = true;
+						ContainerTarget->bUseLocalNameMap = true;
+					}
 
 					if (CookedFileStatData->FileType == FCookedFileStatData::BulkData)
 					{
@@ -3957,13 +4301,20 @@ int32 CreateTarget(const FIoStoreArguments& Arguments, const FIoStoreWriterSetti
 		}
 		FIoStatus IoStatus = IoStoreWriterContext->Initialize(GeneralIoWriterSettings);
 		check(IoStatus.IsOk());
-		IoStatus = GlobalIoStoreWriter->Initialize(*IoStoreWriterContext, /* bIsCompressed */ false);
+		IoStatus = GlobalIoStoreWriter->Initialize(*IoStoreWriterContext, /* bIsCompressed */ false, FIoEncryptionKey());
 		check(IoStatus.IsOk());
 		for (FContainerTargetSpec* ContainerTarget : ContainerTargets)
 		{
 			if (ContainerTarget->IoStoreWriter && ContainerTarget->IoStoreWriter != GlobalIoStoreWriter)
 			{
-				IoStatus = ContainerTarget->IoStoreWriter->Initialize(*IoStoreWriterContext, ContainerTarget->bIsCompressed);
+				FIoEncryptionKey EncryptionKey;
+				if (ContainerTarget->bIsEncrypted)
+				{
+					const FNamedAESKey* Key = Arguments.KeyChain.EncryptionKeys.Find(ContainerTarget->EncryptionKeyGuid);
+					check(Key);
+					EncryptionKey = FIoEncryptionKey { Key->Guid, Key->Key };
+				}
+				IoStatus = ContainerTarget->IoStoreWriter->Initialize(*IoStoreWriterContext, ContainerTarget->bIsCompressed, EncryptionKey);
 				check(IoStatus.IsOk());
 			}
 		}
@@ -4331,30 +4682,36 @@ int32 CreateTarget(const FIoStoreArguments& Arguments, const FIoStoreWriterSetti
 
 	PackageHeaderSize = PackageSummarySize + NameMapSize + ImportMapSize + ExportMapSize + UGraphSize;
 
-	UE_LOG(LogIoStore, Display, TEXT("--------------------------------------------- IoStore Summary -----------------------------------------------"));
+	UE_LOG(LogIoStore, Display, TEXT("--------------------------------------------------- IoStore Summary -----------------------------------------------------"));
 
 	UE_LOG(LogIoStore, Display, TEXT(""));
-	UE_LOG(LogIoStore, Display, TEXT("%-4s %-35s %15s %15s %15s %25s"), TEXT("ID"), TEXT("Container"), TEXT("TOC Size (KB)"), TEXT("TOC Entries"), TEXT("Size (MB)"), TEXT("Compressed (MB)"));
-	UE_LOG(LogIoStore, Display, TEXT("-------------------------------------------------------------------------------------------------------------"));
+	UE_LOG(LogIoStore, Display, TEXT("%-4s %-30s %10s %15s %15s %15s %25s"), TEXT("ID"), TEXT("Container"), TEXT("Encrypted"), TEXT("TOC Size (KB)"), TEXT("TOC Entries"), TEXT("Size (MB)"), TEXT("Compressed (MB)"));
+	UE_LOG(LogIoStore, Display, TEXT("-------------------------------------------------------------------------------------------------------------------------"));
 	int64 TotalPaddingSize = 0;
 	for (const FIoStoreWriterResult& Result : IoStoreWriterResults)
 	{
-		const double Compression = Result.CompressionMethod == NAME_None
-			? 0.0
-			: (double(Result.UncompressedContainerSize - Result.CompressedContainerSize) / double(Result.UncompressedContainerSize)) * 100.0;
+		FString CompressionInfo = TEXT("-");
+		
+		if (Result.CompressionMethod != NAME_None)
+		{
+			double Procentage = (double(Result.UncompressedContainerSize - Result.CompressedContainerSize) / double(Result.UncompressedContainerSize)) * 100.0;
+			CompressionInfo = FString::Printf(TEXT("%.2lf (%.2lf%% %s)"),
+				(double)Result.CompressedContainerSize / 1024.0 / 1024.0,
+				Procentage ,
+				*Result.CompressionMethod.ToString());
+		}
 
-		UE_LOG(LogIoStore, Display, TEXT("%-4d %-35s %15.2lf %15d %15.2lf %25s"),
+		UE_LOG(LogIoStore, Display, TEXT("%-4d %-30s %10s %15.2lf %15d %15.2lf %25s"),
 			Result.ContainerId.ToIndex(),
 			*Result.ContainerName,
+			Result.bIsEncrypted ? TEXT("Yes") : TEXT("No"),
 			(double)Result.TocSize / 1024.0,
 			Result.TocEntryCount,
 			(double)Result.UncompressedContainerSize / 1024.0 / 1024.0,
-			*FString::Printf(TEXT("%.2lf (%.2lf%% %s)"),
-				(double)Result.CompressedContainerSize / 1024.0 / 1024.0,
-				Compression,
-				*Result.CompressionMethod.ToString()));
+			*CompressionInfo);
 		TotalPaddingSize += Result.PaddingSize;
 	}
+	UE_LOG(LogIoStore, Display, TEXT(""));
 	UE_LOG(LogIoStore, Display, TEXT("Compression block padding: %8.2f MB"), TotalPaddingSize / 1024.0 / 1024.0);
 	UE_LOG(LogIoStore, Display, TEXT(""));
 	UE_LOG(LogIoStore, Display, TEXT("Packages: %8d total, %d circular dependencies, %d no import dependencies"),
@@ -4433,6 +4790,10 @@ static bool ParsePakResponseFile(const TCHAR* FilePath, TArray<FContainerSourceF
 			{
 				FileEntry.bNeedsCompression = true;
 			}
+			if (Switches[Index] == TEXT("encrypt"))
+			{
+				FileEntry.bNeedsEncryption = true;
+			}
 		}
 	}
 	return true;
@@ -4577,6 +4938,16 @@ int32 CreateIoStoreContainerFiles(const TCHAR* CmdLine)
 	UE_LOG(LogIoStore, Display, TEXT("==================== IoStore Utils ===================="));
 
 	FIoStoreArguments Arguments;
+
+	LoadKeyChain(FCommandLine::Get(), Arguments.KeyChain);
+
+	FString PatchReferenceCryptoKeysFilename;
+	FKeyChain PatchKeyChain;
+	if (FParse::Value(FCommandLine::Get(), TEXT("PatchCryptoKeys="), PatchReferenceCryptoKeysFilename))
+	{
+		LoadKeyChainFromFile(PatchReferenceCryptoKeysFilename, Arguments.PatchKeyChain);
+	}
+
 	FString GameOrderFilePath;
 	if (FParse::Value(FCommandLine::Get(), TEXT("GameOrder="), GameOrderFilePath))
 	{
diff --git a/Engine/Source/Programs/AutomationTool/Scripts/CopyBuildToStagingDirectory.Automation.cs b/Engine/Source/Programs/AutomationTool/Scripts/CopyBuildToStagingDirectory.Automation.cs
--- a/Engine/Source/Programs/AutomationTool/Scripts/CopyBuildToStagingDirectory.Automation.cs
+++ b/Engine/Source/Programs/AutomationTool/Scripts/CopyBuildToStagingDirectory.Automation.cs
@@ -192,6 +192,22 @@ public partial class Project : CommandUtils
 		WritePakResponseFile(UnrealPakResponseFileName, UnrealPakResponseFile, bCompressed, CryptoSettings, bForceEncryption);
 		CmdLine.AppendFormat(" -ResponseFile={0}", CommandUtils.MakePathSafeToUseWithCommandLine(UnrealPakResponseFileName));
 
+		if (CryptoSettings != null && CryptoSettings.bDataCryptoRequired)
+		{
+			if (CryptoSettings.bEnablePakIndexEncryption)
+			{
+				CmdLine.AppendFormat(" -encryptindex");
+			}
+			if (!string.IsNullOrEmpty(EncryptionKeyGuid))
+			{
+				CmdLine.AppendFormat(" -EncryptionKeyOverrideGuid={0}", EncryptionKeyGuid);
+			}
+			if (CryptoSettings.bDataCryptoRequired && CryptoSettings.bEnablePakSigning && CryptoSettings.SigningKey.IsValid())
+			{
+				CmdLine.AppendFormat(" -sign");
+			}
+		}
+
 		return CmdLine.ToString();
 	}
 
@@ -2436,6 +2452,11 @@ public partial class Project : CommandUtils
 				+ " " + AdditionalCompressionOptionsOnCommandLine
 				+ " " + Params.AdditionalPakOptions;
 
+			if (CryptoKeysCacheFilename != null)
+			{
+				AdditionalArgs += String.Format(" -cryptokeys={0}", CommandUtils.MakePathSafeToUseWithCommandLine(CryptoKeysCacheFilename.FullName));
+			}
+
 			RunIoStore(Params, SC, IoStoreCommandsFileName, GameOpenOrderFileLocation, CookerOpenOrderFileLocation, AdditionalArgs);
 		}
 
diff --git a/Engine/Source/Runtime/Core/Private/IO/IoDispatcherFileBackend.cpp b/Engine/Source/Runtime/Core/Private/IO/IoDispatcherFileBackend.cpp
--- a/Engine/Source/Runtime/Core/Private/IO/IoDispatcherFileBackend.cpp
+++ b/Engine/Source/Runtime/Core/Private/IO/IoDispatcherFileBackend.cpp
@@ -11,6 +11,7 @@
 #include "Async/MappedFileHandle.h"
 #include "HAL/RunnableThread.h"
 #include "Misc/ScopeLock.h"
+#include "Misc/CoreDelegates.h"
 
 TRACE_DECLARE_MEMORY_COUNTER(IoDispatcherTotalBytesRead, TEXT("IoDispatcher/TotalBytesRead"));
 TRACE_DECLARE_MEMORY_COUNTER(IoDispatcherTotalBytesScattered, TEXT("IoDispatcher/TotalBytesScattered"));
@@ -65,6 +66,43 @@ private:
 	IMappedFileHandle* SharedMappedFileHandle;
 };
 
+FEncryptionKeyCache::FEncryptionKeyCache()
+{
+	FCoreDelegates::GetRegisterEncryptionKeyDelegate().BindRaw(this, &FEncryptionKeyCache::RegisterEncryptionKey);
+}
+
+bool FEncryptionKeyCache::GetEncryptionKey(const FGuid& Guid, FAES::FAESKey& OutKey) const
+{
+	OutKey.Reset();
+
+	{
+		FScopeLock _(&EncryptionKeysCritical);
+		if (const FAES::FAESKey* ExistingKey = EncryptionKeysByGuid.Find(Guid))
+		{
+			OutKey = *ExistingKey;
+			return OutKey.IsValid();
+		}
+	}
+
+	if (!Guid.IsValid() && FCoreDelegates::GetPakEncryptionKeyDelegate().IsBound())
+	{
+		FCoreDelegates::GetPakEncryptionKeyDelegate().Execute(OutKey.Key);
+		return OutKey.IsValid();
+	}
+
+	return false;
+}
+
+void FEncryptionKeyCache::RegisterEncryptionKey(const FGuid& Guid, const FAES::FAESKey& Key)
+{
+	{
+		FScopeLock _(&EncryptionKeysCritical);
+		EncryptionKeysByGuid.Add(Guid, Key);
+	}
+
+	KeyRegisteredCallback(Guid, Key);
+}
+
 FFileIoStoreReader::FFileIoStoreReader(FFileIoStoreImpl& InPlatformImpl)
 	: PlatformImpl(InPlatformImpl)
 {
@@ -141,6 +179,8 @@ FIoStatus FFileIoStoreReader::Initialize(const FIoStoreEnvironment& Environment)
 
 	ContainerFile.CompressionBlockSize = TocReader.GetCompressionBlockSize();
 	ContainerFile.CompressionBlocks.Reserve(CompressedBlockEntryCount);
+	ContainerFile.bIsEncrypted = TocReader.IsContainerEncrypted();
+	ContainerFile.EncryptionKeyGuid = TocReader.GetEncryptionKeyGuid();
 
 	while (CompressedBlockEntryCount--)
 	{
@@ -220,6 +260,18 @@ FFileIoStore::FFileIoStore(FIoDispatcherEventQueue& InEventQueue)
 		FirstFreeCompressionContext = Context;
 	}
 
+	EncryptionKeys.SetKeyRegisteredCallback([this](const FGuid& Guid, const FAES::FAESKey& Key)
+	{
+		FReadScopeLock _(IoStoreReadersLock);
+		for (FFileIoStoreReader* Reader : IoStoreReaders)
+		{
+			if (Reader->IsEncrypted() && !Reader->GetEncryptionKey().IsValid() && Reader->GetEncryptionKeyGuid() == Guid)
+			{
+				Reader->SetEncryptionKey(Key);
+			}
+		}
+	});
+
 	Thread = FRunnableThread::Create(this, TEXT("IoService"), 0, TPri_AboveNormal);
 }
 
@@ -237,6 +289,20 @@ TIoStatusOr<FIoContainerId> FFileIoStore::Mount(const FIoStoreEnvironment& Envir
 	{
 		return IoStatus;
 	}
+
+	if (Reader->IsEncrypted())
+	{
+		FAES::FAESKey EncryptionKey;
+		if (EncryptionKeys.GetEncryptionKey(Reader->GetEncryptionKeyGuid(), EncryptionKey))
+		{
+			Reader->SetEncryptionKey(EncryptionKey);
+		}
+		else
+		{
+			UE_LOG(LogIoDispatcher, Warning, TEXT("Failed to find valid encryption key for encrypted container '%s'"), *FPaths::GetBaseFilename(Environment.GetPath()));
+		}
+	}
+
 	int32 InsertionIndex;
 	FIoContainerId ContainerId = Reader->GetContainerId();
 	{
@@ -365,6 +431,11 @@ void FFileIoStore::ScatterBlock(FFileIoStoreCompressedBlock* CompressedBlock, bo
 		uint64 OffsetInBuffer = CompressedBlock->RawOffset - RawBlock->Offset;
 		CompressedBuffer = RawBlock->Buffer->Memory + OffsetInBuffer;
 	}
+	if (CompressedBlock->EncryptionKey.IsValid())
+	{
+		const uint64 BufferSize = CompressedBlock->CompressionMethod.IsNone() ? CompressedBlock->Size : CompressedBlock->RawSize;
+		FAES::DecryptData(CompressedBuffer, static_cast<uint32>(BufferSize), CompressedBlock->EncryptionKey);
+	}
 	uint8* UncompressedBuffer;
 	if (CompressedBlock->CompressionMethod.IsNone())
 	{
@@ -610,7 +681,12 @@ void FFileIoStore::ReadBlocks(uint32 ReaderIndex, const FFileIoStoreResolvedRequ
 	ScopeName.Appendf(TEXT("ReadBlock %d"), BlockIndex);
 	TRACE_CPUPROFILER_EVENT_SCOPE_TEXT(*ScopeName);*/
 
-	const FFileIoStoreContainerFile& ContainerFile = IoStoreReaders[ReaderIndex]->GetContainerFile();
+	const FFileIoStoreReader& Reader = *IoStoreReaders[ReaderIndex];
+	UE_CLOG(Reader.IsEncrypted() && !Reader.GetEncryptionKey().IsValid(),
+		LogIoDispatcher, Fatal, TEXT("Reading from encrypted container (ID = '%d') with invalid encryption key (Guid = '%s')"),
+		Reader.GetContainerId().ToIndex(),
+		*Reader.GetEncryptionKeyGuid().ToString());
+	const FFileIoStoreContainerFile& ContainerFile = Reader.GetContainerFile();
 	const uint64 CompressionBlockSize = ContainerFile.CompressionBlockSize;
 	const uint64 RequestEndOffset = ResolvedRequest.ResolvedOffset + ResolvedRequest.ResolvedSize;
 	int32 RequestBeginBlockIndex = int32(ResolvedRequest.ResolvedOffset / CompressionBlockSize);
@@ -632,6 +708,7 @@ void FFileIoStore::ReadBlocks(uint32 ReaderIndex, const FFileIoStoreResolvedRequ
 		{
 			CompressedBlock = new FFileIoStoreCompressedBlock();
 			CompressedBlock->Key = CompressedBlockKey;
+			CompressedBlock->EncryptionKey = Reader.GetEncryptionKey();
 			CompressedBlocksMap.Add(CompressedBlockKey, CompressedBlock);
 
 			bool bCacheable = OffsetInRequest > 0 || RequestRemainingBytes < CompressionBlockSize;
diff --git a/Engine/Source/Runtime/Core/Private/IO/IoDispatcherFileBackend.h b/Engine/Source/Runtime/Core/Private/IO/IoDispatcherFileBackend.h
--- a/Engine/Source/Runtime/Core/Private/IO/IoDispatcherFileBackend.h
+++ b/Engine/Source/Runtime/Core/Private/IO/IoDispatcherFileBackend.h
@@ -9,6 +9,7 @@
 #include "Stats/Stats.h"
 #include "Async/TaskGraphInterfaces.h"
 #include "HAL/Runnable.h"
+#include "Misc/AES.h"
 
 class IMappedFileHandle;
 
@@ -21,6 +22,9 @@ struct FFileIoStoreContainerFile
 	TArray<FIoStoreTocCompressedBlockEntry> CompressionBlocks;
 	FString FilePath;
 	TUniquePtr<IMappedFileHandle> MappedFileHandle;
+	FGuid EncryptionKeyGuid;
+	FAES::FAESKey EncryptionKey;
+	bool bIsEncrypted = false;
 };
 
 struct FFileIoStoreBuffer
@@ -88,6 +92,7 @@ struct FFileIoStoreCompressedBlock
 	TArray<FFileIoStoreBlockScatter, TInlineAllocator<16>> ScatterList;
 	FFileIoStoreCompressionContext* CompressionContext = nullptr;
 	uint8* CompressedDataBuffer = nullptr;
+	FAES::FAESKey EncryptionKey;
 };
 
 struct FFileIoStoreRawBlock
@@ -117,6 +122,26 @@ struct FFileIoStoreResolvedRequest
 	uint64 ResolvedSize;
 };
 
+class FEncryptionKeyCache
+{
+public:
+	using FKeyRegisteredCallback = TFunction<void(const FGuid&, const FAES::FAESKey&)>;
+
+	FEncryptionKeyCache();
+	bool GetEncryptionKey(const FGuid& Guid, FAES::FAESKey& OutKey) const;
+	void SetKeyRegisteredCallback(FKeyRegisteredCallback&& Callback)
+	{
+		KeyRegisteredCallback = Callback;
+	}
+
+private:
+	void RegisterEncryptionKey(const FGuid& Guid, const FAES::FAESKey& Key);
+
+	TMap<FGuid, FAES::FAESKey> EncryptionKeysByGuid;
+	mutable FCriticalSection EncryptionKeysCritical;
+	FKeyRegisteredCallback KeyRegisteredCallback;
+};
+
 class FFileIoStoreReader
 {
 public:
@@ -125,10 +150,14 @@ public:
 	bool DoesChunkExist(const FIoChunkId& ChunkId) const;
 	TIoStatusOr<uint64> GetSizeForChunk(const FIoChunkId& ChunkId) const;
 	const FIoOffsetAndLength* Resolve(const FIoChunkId& ChunkId) const;
-	const FFileIoStoreContainerFile& GetContainerFile() { return ContainerFile; }
+	const FFileIoStoreContainerFile& GetContainerFile() const { return ContainerFile; }
 	IMappedFileHandle* GetMappedContainerFileHandle();
 	const FIoContainerId& GetContainerId() const { return ContainerId; }
 	int32 GetOrder() const { return Order; }
+	bool IsEncrypted() const { return ContainerFile.bIsEncrypted; }
+	const FGuid& GetEncryptionKeyGuid() const { return ContainerFile.EncryptionKeyGuid; }
+	void SetEncryptionKey(const FAES::FAESKey& Key) { ContainerFile.EncryptionKey = Key; }
+	const FAES::FAESKey& GetEncryptionKey() const { return ContainerFile.EncryptionKey; }
 
 private:
 	FFileIoStoreImpl& PlatformImpl;
@@ -231,4 +260,5 @@ private:
 	FFileIoStoreRawBlock* ScheduledBlocksTail = nullptr;
 	FCriticalSection DecompressedBlocksCritical;
 	FFileIoStoreCompressedBlock* FirstDecompressedBlock = nullptr;
+	FEncryptionKeyCache EncryptionKeys;
 };
diff --git a/Engine/Source/Runtime/Core/Private/IO/IoStore.cpp b/Engine/Source/Runtime/Core/Private/IO/IoStore.cpp
--- a/Engine/Source/Runtime/Core/Private/IO/IoStore.cpp
+++ b/Engine/Source/Runtime/Core/Private/IO/IoStore.cpp
@@ -82,7 +82,7 @@ struct FIoStoreWriterBlock
 	uint64 Alignment = 0;
 	FIoBuffer UncompressedData;
 	FIoBuffer CompressedData;
-	FGraphEventRef CompressionTask;
+	FGraphEventRef CompressionEncryptionTask;
 	bool bForceUncompressed = false;
 };
 
@@ -201,11 +201,13 @@ public:
 		}
 	}
 
-	FIoStatus Initialize(FIoStoreWriterContextImpl& InContext, const TCHAR* InFileName, bool bInIsContainerCompressed)
+	FIoStatus Initialize(FIoStoreWriterContextImpl& InContext, const TCHAR* InFileName, bool bInIsContainerCompressed, const FAES::FAESKey& InEncryptionKey)
 	{
 		TRACE_CPUPROFILER_EVENT_SCOPE(InitializeChunkWriter);
 		WriterContext = &InContext;
 		bIsContainerCompressed = bInIsContainerCompressed;
+		bIsContainerEncrypted = InEncryptionKey.IsValid();
+		EncryptionKey = InEncryptionKey;
 		check(WriterContext->Settings().CompressionBlockSize > 0);
 		CompressionInfo.BlockSize = WriterContext->Settings().CompressionBlockSize;
 
@@ -237,9 +239,10 @@ public:
 					FIoStoreWriterBlock* CurrentWriteQueueItem = NextWriteQueueItem;
 					NextWriteQueueItem = NextWriteQueueItem->Next;
 
-					if (bIsContainerCompressed && !CurrentWriteQueueItem->bForceUncompressed)
+					const bool bIsCompressed = bIsContainerCompressed && !CurrentWriteQueueItem->bForceUncompressed;
+					if (bIsCompressed || bIsContainerEncrypted)
 					{
-						FTaskGraphInterface::Get().WaitUntilTaskCompletes(CurrentWriteQueueItem->CompressionTask);
+						FTaskGraphInterface::Get().WaitUntilTaskCompletes(CurrentWriteQueueItem->CompressionEncryptionTask);
 					}
 					if (CurrentWriteQueueItem->Alignment > 0)
 					{
@@ -398,28 +401,56 @@ public:
 	}
 
 private:
-	void CompressBlock(FIoStoreWriterBlock* Block)
+	void CompressEncryptBlock(FIoStoreWriterBlock* Block)
 	{
-		TRACE_CPUPROFILER_EVENT_SCOPE(CompressBlock);
-		int32 CompressedSize = int32(Block->CompressedData.DataSize());
-		const bool bCompressed = FCompression::CompressMemory(
-			WriterContext->Settings().CompressionMethod,
-			Block->CompressedData.Data(),
-			CompressedSize,
-			Block->UncompressedData.Data(),
-			int32(Block->UncompressedData.DataSize()));
-		check(bCompressed);
-		check(CompressedSize > 0);
-
-		if (CompressedSize > Block->UncompressedData.DataSize())
+		if (bIsContainerCompressed)
 		{
-			Block->CompressionMethod = NAME_None;
-			Block->CompressedSize = Block->UncompressedData.DataSize();
+			TRACE_CPUPROFILER_EVENT_SCOPE(CompressBlock);
+			int32 CompressedSize = int32(Block->CompressedData.DataSize());
+			const bool bCompressed = FCompression::CompressMemory(
+				WriterContext->Settings().CompressionMethod,
+				Block->CompressedData.Data(),
+				CompressedSize,
+				Block->UncompressedData.Data(),
+				int32(Block->UncompressedData.DataSize()));
+			check(bCompressed);
+			check(CompressedSize > 0);
+
+			if (bIsContainerEncrypted)
+			{
+				// Fill the trailing buffer with bytes from the block. Note that this is now from a fixed location
+				// rather than a random one so that we produce deterministic results (PakFileUtilities.cpp)
+				const int32 EncryptedSize = Align(CompressedSize, FAES::AESBlockSize);
+				for (int32 FillIndex = CompressedSize; FillIndex < EncryptedSize; ++FillIndex)
+				{
+					Block->CompressedData.Data()[FillIndex] = Block->CompressedData.Data()[(FillIndex - CompressedSize) % CompressedSize];
+				}
+				CompressedSize = EncryptedSize;
+			}
+
+			if (CompressedSize > Block->UncompressedData.DataSize())
+			{
+				Block->CompressionMethod = NAME_None;
+				Block->CompressedSize = Block->UncompressedData.DataSize();
+			}
+			else
+			{
+				Block->CompressionMethod = WriterContext->Settings().CompressionMethod;
+				Block->CompressedSize = CompressedSize;
+			}
 		}
-		else
+
+		if (bIsContainerEncrypted)
 		{
-			Block->CompressionMethod = WriterContext->Settings().CompressionMethod;
-			Block->CompressedSize = CompressedSize;
+			TRACE_CPUPROFILER_EVENT_SCOPE(EncryptBlock);
+			if (Block->CompressionMethod == NAME_None)
+			{
+				FAES::EncryptData(Block->UncompressedData.Data(), static_cast<uint32>(Block->UncompressedData.DataSize()), EncryptionKey);
+			}
+			else
+			{
+				FAES::EncryptData(Block->CompressedData.Data(), static_cast<uint32>(Block->CompressedSize), EncryptionKey);
+			}
 		}
 	}
 
@@ -429,18 +460,33 @@ private:
 		{
 			return;
 		}
-		if (CurrentBlockOffset < CurrentBlock->UncompressedData.DataSize())
+		const bool bIsCompressed = bIsContainerCompressed && !CurrentBlock->bForceUncompressed;
+		uint8* UncompressedData = CurrentBlock->UncompressedData.Data();
+		const uint64 UncompressedDataSize = CurrentBlock->UncompressedData.DataSize();
+		if (CurrentBlockOffset < UncompressedDataSize)
 		{
-			uint64 PadCount = CurrentBlock->UncompressedData.DataSize() - CurrentBlockOffset;
-			FMemory::Memzero(CurrentBlock->UncompressedData.Data() + CurrentBlockOffset, PadCount);
+			uint64 PadCount = UncompressedDataSize - CurrentBlockOffset;
+			if (bIsContainerEncrypted && !bIsCompressed)
+			{
+				// When the block is encrypted but NOT compressed we fill the trailing bytes
+				// with data from the block. When the block IS compressed the trailing bytes
+				// will get filled after compression.
+				for (uint64 FillIndex = CurrentBlockOffset; FillIndex < UncompressedDataSize; ++FillIndex)
+				{
+					UncompressedData[FillIndex] = UncompressedData[FillIndex - CurrentBlockOffset];
+				}
+			}
+			else
+			{
+				FMemory::Memzero(UncompressedData + CurrentBlockOffset, PadCount);
+			}
 			UncompressedOffset += PadCount;
 		}
-		if (bIsContainerCompressed && !CurrentBlock->bForceUncompressed)
+		if (bIsCompressed || bIsContainerEncrypted)
 		{
-			FIoStoreWriterBlock* ReadyToCompressBlock = CurrentBlock;
-			CurrentBlock->CompressionTask = FFunctionGraphTask::CreateAndDispatchWhenReady([this, ReadyToCompressBlock]()
+			CurrentBlock->CompressionEncryptionTask = FFunctionGraphTask::CreateAndDispatchWhenReady([this, Block = CurrentBlock]()
 			{
-				CompressBlock(ReadyToCompressBlock);
+				CompressEncryptBlock(Block);
 			}, TStatId(), nullptr, ENamedThreads::AnyHiPriThreadHiPriTask);
 		}
 		else
@@ -480,7 +526,9 @@ private:
 	FEvent* WriteQueueEvent = nullptr;
 	TAtomic<bool> bAllScheduled{ false };
 	TArray<uint8> AlignmentPaddingBuffer;
+	FAES::FAESKey EncryptionKey;
 	bool bIsContainerCompressed = false;
+	bool bIsContainerEncrypted = false;
 };
 
 //////////////////////////////////////////////////////////////////////////
@@ -494,21 +542,24 @@ public:
 	{
 	}
 
-	UE_NODISCARD FIoStatus Initialize(FIoStoreWriterContextImpl& InContext, bool bInIsContainerCompressed)
+	UE_NODISCARD FIoStatus Initialize(FIoStoreWriterContextImpl& InContext, bool bInIsContainerCompressed, const FIoEncryptionKey& InEncryptionKey)
 	{
 		IPlatformFile& Ipf = IPlatformFile::GetPlatformPhysical();
 
+		EncryptionKey = InEncryptionKey;
+
 		FString TocFilePath = Environment.GetPath() + TEXT(".utoc");
 		FString ContainerFilePath = Environment.GetPath() + TEXT(".ucas");
 
 		Result.ContainerId = ContainerId;
 		Result.ContainerName = FPaths::GetBaseFilename(Environment.GetPath());
-		Result.CompressionMethod = InContext.Settings().CompressionMethod;
+		Result.CompressionMethod = bInIsContainerCompressed ? InContext.Settings().CompressionMethod : NAME_None;
+		Result.bIsEncrypted = InEncryptionKey.Key.IsValid();
 
 		Ipf.CreateDirectoryTree(*FPaths::GetPath(ContainerFilePath));
 
 		ChunkWriter = MakeUnique<FChunkWriter>();
-		FIoStatus Status = ChunkWriter->Initialize(InContext, *ContainerFilePath, bInIsContainerCompressed);
+		FIoStatus Status = ChunkWriter->Initialize(InContext, *ContainerFilePath, bInIsContainerCompressed, InEncryptionKey.Key);
 
 		if (!Status.IsOk())
 		{ 
@@ -664,6 +715,8 @@ public:
 		TocHeader.CompressionMethodNameCount = CompressionInfo.CompressionMethods.Num();
 		TocHeader.CompressionMethodNameLength = FIoStoreCompressionInfo::CompressionMethodNameLen;
 		TocHeader.ContainerId = ContainerId;
+		TocHeader.EncryptionKeyGuid = EncryptionKey.Guid;
+		TocHeader.bIsEncrypted = EncryptionKey.Key.IsValid();
 
 		TocFileHandle->Seek(0);
 		if (!TocFileHandle->Write(reinterpret_cast<const uint8*>(&TocHeader), sizeof(TocHeader)))
@@ -734,6 +787,7 @@ private:
 	TUniquePtr<FArchive>				CsvArchive;
 	FIoStoreWriterResult				Result;
 	FIoContainerId						ContainerId;
+	FIoEncryptionKey					EncryptionKey;
 	bool								IsMetadataDirty = true;
 };
 
@@ -747,9 +801,9 @@ FIoStoreWriter::~FIoStoreWriter()
 	(void)Impl->Flush();
 }
 
-FIoStatus FIoStoreWriter::Initialize(const FIoStoreWriterContext& Context, bool bIsContainerCompressed)
+FIoStatus FIoStoreWriter::Initialize(const FIoStoreWriterContext& Context, bool bIsContainerCompressed, const FIoEncryptionKey& EncryptionKey)
 {
-	return Impl->Initialize(*Context.Impl, bIsContainerCompressed);
+	return Impl->Initialize(*Context.Impl, bIsContainerCompressed, EncryptionKey);
 }
 
 FIoStatus FIoStoreWriter::Append(const FIoChunkId& ChunkId, FIoBuffer Chunk, const FIoWriteOptions& WriteOptions)
@@ -785,12 +839,14 @@ public:
 
 	}
 
-	UE_NODISCARD FIoStatus Initialize(const FIoStoreEnvironment& Environment)
+	UE_NODISCARD FIoStatus Initialize(const FIoStoreEnvironment& InEnvironment, const FIoEncryptionKey& InEncryptionKey)
 	{
+		EncryptionKey = InEncryptionKey;
+
 		IPlatformFile& Ipf = FPlatformFileManager::Get().GetPlatformFile();
 
 		TStringBuilder<256> ContainerFilePath;
-		ContainerFilePath.Append(Environment.GetPath());
+		ContainerFilePath.Append(InEnvironment.GetPath());
 
 		TStringBuilder<256> TocFilePath;
 		TocFilePath.Append(ContainerFilePath);
@@ -851,6 +907,9 @@ public:
 			++CompressedBlockEntry;
 		}
 
+		bIsEncrypted = TocReader.IsContainerEncrypted();
+		check(!bIsEncrypted || EncryptionKey.Guid == TocReader.GetEncryptionKeyGuid());
+
 		return FIoStatus::Ok;
 	}
 
@@ -923,6 +982,10 @@ public:
 			}
 			ContainerFileHandle->Seek(CompressionBlock.OffsetAndLength.GetOffset());
 			ContainerFileHandle->Read(CompressedBuffer.GetData(), CompressionBlock.OffsetAndLength.GetLength());
+			if (bIsEncrypted)
+			{
+				FAES::DecryptData(CompressedBuffer.GetData(), static_cast<uint32>(CompressionBlock.OffsetAndLength.GetLength()), EncryptionKey.Key);
+			}
 			if (CompressionBlock.CompressionMethod.IsNone())
 			{
 				Src = CompressedBuffer.GetData();
@@ -961,6 +1024,8 @@ private:
 	TUniquePtr<IFileHandle> ContainerFileHandle;
 	mutable TArray<uint8> CompressedBuffer;
 	mutable TArray<uint8> UncompressedBuffer;
+	FIoEncryptionKey EncryptionKey;
+	bool bIsEncrypted;
 };
 
 
@@ -974,9 +1039,9 @@ FIoStoreReader::~FIoStoreReader()
 	delete Impl;
 }
 
-FIoStatus FIoStoreReader::Initialize(const FIoStoreEnvironment& InEnvironment)
+FIoStatus FIoStoreReader::Initialize(const FIoStoreEnvironment& InEnvironment, const FIoEncryptionKey& EncryptionKey)
 {
-	return Impl->Initialize(InEnvironment);
+	return Impl->Initialize(InEnvironment, EncryptionKey);
 }
 
 void FIoStoreReader::EnumerateChunks(TFunction<bool(const FIoStoreTocChunkInfo&)>&& Callback) const
@@ -1062,4 +1127,4 @@ FIoStatus FIoStoreTocReader::Initialize(const TCHAR* TocFilePath)
 	}
 
 	return FIoStatus::Ok;
-}
\ No newline at end of file
+}
diff --git a/Engine/Source/Runtime/Core/Private/IO/IoStore.h b/Engine/Source/Runtime/Core/Private/IO/IoStore.h
--- a/Engine/Source/Runtime/Core/Private/IO/IoStore.h
+++ b/Engine/Source/Runtime/Core/Private/IO/IoStore.h
@@ -23,7 +23,9 @@ struct FIoStoreTocHeader
 	uint32	CompressionMethodNameLength;
 	uint32	CompressionBlockSize;
 	FIoContainerId ContainerId;
-	uint8 Pad[70];
+	FGuid	EncryptionKeyGuid;
+	uint8	bIsEncrypted;
+	uint8	Pad[53];
 
 	void MakeMagic()
 	{
@@ -193,6 +195,16 @@ public:
 		return Header->ContainerId;
 	}
 
+	const FGuid& GetEncryptionKeyGuid() const
+	{
+		return Header->EncryptionKeyGuid;
+	}
+
+	bool IsContainerEncrypted() const
+	{
+		return Header->bIsEncrypted > 0;
+	}
+
 private:
 	TArray<FName> CompressionMethodNames;
 	TUniquePtr<uint8[]> TocBuffer;
diff --git a/Engine/Source/Runtime/Core/Public/IO/IoDispatcher.h b/Engine/Source/Runtime/Core/Public/IO/IoDispatcher.h
--- a/Engine/Source/Runtime/Core/Public/IO/IoDispatcher.h
+++ b/Engine/Source/Runtime/Core/Public/IO/IoDispatcher.h
@@ -12,6 +12,7 @@
 #include "Templates/TypeCompatibleBytes.h"
 #include "HAL/PlatformAtomics.h"
 #include "Misc/SecureHash.h"
+#include "Misc/AES.h"
 
 #if __cplusplus >= 201703L
 #	define UE_NODISCARD		[[nodiscard]]
@@ -952,6 +953,7 @@ struct FIoStoreWriterResult
 	int64 UncompressedContainerSize = 0;
 	int64 CompressedContainerSize = 0;
 	FName CompressionMethod = NAME_None;
+	bool bIsEncrypted = false;
 };
 
 struct FIoWriteOptions
@@ -975,6 +977,12 @@ private:
 	FIoStoreWriterContextImpl* Impl;
 };
 
+struct FIoEncryptionKey
+{
+	FGuid Guid;
+	FAES::FAESKey Key;
+};
+
 class FIoStoreWriter
 {
 public:
@@ -984,7 +992,7 @@ public:
 	FIoStoreWriter(const FIoStoreWriter&) = delete;
 	FIoStoreWriter& operator=(const FIoStoreWriter&) = delete;
 
-	UE_NODISCARD CORE_API FIoStatus	Initialize(const FIoStoreWriterContext& Context, bool bIsContainerCompressed);
+	UE_NODISCARD CORE_API FIoStatus	Initialize(const FIoStoreWriterContext& Context, bool bIsContainerCompressed, const FIoEncryptionKey& EncryptionKey);
 	UE_NODISCARD CORE_API FIoStatus	Append(const FIoChunkId& ChunkId, const FIoChunkHash& ChunkHash, FIoBuffer Chunk, const FIoWriteOptions& WriteOptions);
 	UE_NODISCARD CORE_API FIoStatus	Append(const FIoChunkId& ChunkId, FIoBuffer Chunk, const FIoWriteOptions& WriteOptions);
 	UE_NODISCARD CORE_API FIoStatus	AppendPadding(uint64 Count);
@@ -1026,7 +1034,7 @@ public:
 	CORE_API FIoStoreReader();
 	CORE_API ~FIoStoreReader();
 
-	UE_NODISCARD CORE_API FIoStatus Initialize(const FIoStoreEnvironment& InEnvironment);
+	UE_NODISCARD CORE_API FIoStatus Initialize(const FIoStoreEnvironment& InEnvironment, const FIoEncryptionKey& EncryptionKey);
 	CORE_API void EnumerateChunks(TFunction<bool(const FIoStoreTocChunkInfo&)>&& Callback) const;
 	CORE_API void EnumeratePartialChunks(TFunction<bool(const FIoStoreTocPartialChunkInfo&)>&& Callback) const;
 	CORE_API TIoStatusOr<FIoBuffer> Read(const FIoChunkId& Chunk, const FIoReadOptions& Options) const;
