IoStore:
Fixed issue with localized packages not being handled correctly when creating export bundles causing too many bundles to be created in some cases
Fixed incorrect sorting criteria for export bundle arcs
Fixed issue with bundle arcs sometimes being incorrectly added to redirected packages
Ensure deterministic ordering of the bulkdata and optional bulkdata chunk of a package if they end up in the same container
Ensure deterministic ordering of the packages in each container header

#rnx

diff --git a/Engine/Source/Developer/IoStoreUtilities/Private/IoStoreUtilities.cpp b/Engine/Source/Developer/IoStoreUtilities/Private/IoStoreUtilities.cpp
--- a/Engine/Source/Developer/IoStoreUtilities/Private/IoStoreUtilities.cpp
+++ b/Engine/Source/Developer/IoStoreUtilities/Private/IoStoreUtilities.cpp
@@ -1232,14 +1232,25 @@ static void CreateDiskLayout(
 		{
 			if (A->bIsMemoryMappedBulkData != B->bIsMemoryMappedBulkData)
 			{
+				// Put all memory mapped bulkdata last
 				return B->bIsMemoryMappedBulkData;
 			}
-			if (A->bIsBulkData != B->bIsBulkData)
+			else if (A->bIsBulkData != B->bIsBulkData)
 			{
+				// Put packages before bulkdata
 				return B->bIsBulkData;
 			}
-
-			return A->Package->DiskLayoutOrder < B->Package->DiskLayoutOrder;
+			else if (A->Package != B->Package)
+			{
+				return A->Package->DiskLayoutOrder < B->Package->DiskLayoutOrder;
+			}
+			else if (A->bIsOptionalBulkData != B->bIsOptionalBulkData)
+			{
+				// Put each ubulk right before it's corresponding uptnl
+				return B->bIsOptionalBulkData;
+			}
+			check(A == B)
+			return false;
 		});
 		uint64 IdealOrder = 0;
 		for (FContainerTargetFile* TargetFile : SortedTargetFiles)
@@ -1846,6 +1857,7 @@ static void AddPreloadDependencies(
 	const FPackageAssetData& PackageAssetData,
 	const FGlobalPackageData& GlobalPackageData,
 	const FSourceToLocalizedPackageMultimap& SourceToLocalizedPackageMap,
+	const TMap<FPackageId, FPackageId>& SourceToRedirectedPackageMap,
 	FExportGraph& ExportGraph,
 	TArray<FPackage*>& Packages)
 {
@@ -1907,6 +1919,13 @@ static void AddPreloadDependencies(
 							}
 
 							FPackageId PackageId = FPackageId::FromName(Import->ObjectName);
+							if (SourceToRedirectedPackageMap.Contains(PackageId))
+							{
+								// Edge case with an import from a redirection source package but the corresponding
+								// export is missing from the redirection target package so don't add a base game
+								// package arc
+								return;
+							}
 							bool bIsAlreadyInSet = false;
 							ExternalPackageDependencies.Add(PackageId, &bIsAlreadyInSet);
 							if (!bIsAlreadyInSet)
@@ -2009,7 +2028,7 @@ void FinalizePackageHeaders(
 				{
 					return A.ToNodeIndex < B.ToNodeIndex;
 				}
-				return A.FromNodeIndex < B.ToNodeIndex;
+				return A.FromNodeIndex < B.FromNodeIndex;
 			});
 			SortedExternalArcs.Emplace(ImportedPackageId, MoveTemp(SortedArcs));
 		}
@@ -2181,14 +2200,23 @@ void FinalizePackageStoreContainerHeader(FContainerTargetSpec& ContainerTarget)
 	};
 
 	PackageIds.Reserve(ContainerTarget.PackageCount);
+	TArray<FContainerTargetFile*> SortedTargetFiles;
+	SortedTargetFiles.Reserve(ContainerTarget.TargetFiles.Num());
 	for (FContainerTargetFile& TargetFile : ContainerTarget.TargetFiles)
 	{
-		if (TargetFile.bIsBulkData)
+		if (!TargetFile.bIsBulkData)
 		{
-			continue;
+			SortedTargetFiles.Add(&TargetFile);
 		}
+	}
+	Algo::Sort(SortedTargetFiles, [](const FContainerTargetFile* A, const FContainerTargetFile* B)
+	{
+		return A->Package->GlobalPackageId < B->Package->GlobalPackageId;
+	});
 
-		FPackage* Package = TargetFile.Package;
+	for (FContainerTargetFile* TargetFile : SortedTargetFiles)
+	{
+		FPackage* Package = TargetFile->Package;
 
 		// PackageIds
 		{
@@ -2210,7 +2238,7 @@ void FinalizePackageStoreContainerHeader(FContainerTargetSpec& ContainerTarget)
 
 		// StoreEntries
 		{
-			uint64 ExportBundlesSize = TargetFile.HeaderSerialSize + Package->ExportsSerialSize;
+			uint64 ExportBundlesSize = TargetFile->HeaderSerialSize + Package->ExportsSerialSize;
 			int32 ExportBundleCount = Package->ExportBundles.Num();
 			uint32 LoadOrder = Package->ExportBundles.Num() > 0 ? Package->ExportBundles[0].LoadOrder : 0;
 			uint32 Pad = 0;
@@ -2701,7 +2729,8 @@ static void CreateGlobalImportsAndExports(
 	const FPackageIdMap& PackageIdMap,
 	FPackageAssetData& PackageAssetData,
 	FGlobalPackageData& GlobalPackageData,
-	FExportGraph& ExportGraph)
+	FExportGraph& ExportGraph,
+	TMap<FPackageId, FPackageId>& OutSourceToRedirectedPackageMap)
 {
 	IOSTORE_CPU_SCOPE(CreateGlobalImportsAndExports);
 	UE_LOG(LogIoStore, Display, TEXT("Creating global imports and exports..."));
@@ -2721,6 +2750,7 @@ static void CreateGlobalImportsAndExports(
 		if (Package->RedirectedPackageId.IsValid())
 		{
 			Redirects.Add(Package->Name, Package->SourcePackageName);
+			OutSourceToRedirectedPackageMap.Add(Package->RedirectedPackageId, Package->GlobalPackageId);
 		}
 
 		TempFullNames.Reset();
@@ -2933,16 +2963,15 @@ static void ProcessLocalizedPackages(
 		LocalizedPackages.Reset();
 		for (FPackage* ImportedPackage : Package->ImportedPackages)
 		{
-			LocalizedPackages.Reset();
 			OutSourceToLocalizedPackageMap.MultiFind(ImportedPackage, LocalizedPackages);
-			for (FPackage* LocalizedPackage : LocalizedPackages)
-			{
-				UE_LOG(LogIoStore, Verbose, TEXT("For package '%s' (0x%llX): Adding localized imported package '%s' (0x%llX)"),
-					*Package->Name.ToString(),
-					Package->GlobalPackageId.ValueForDebugging(),
-					*LocalizedPackage->Name.ToString(),
-					LocalizedPackage->GlobalPackageId.ValueForDebugging());
-			}
+		}
+		for (FPackage* LocalizedPackage : LocalizedPackages)
+		{
+			UE_LOG(LogIoStore, Verbose, TEXT("For package '%s' (0x%llX): Adding localized imported package '%s' (0x%llX)"),
+				*Package->Name.ToString(),
+				Package->GlobalPackageId.ValueForDebugging(),
+				*LocalizedPackage->Name.ToString(),
+				LocalizedPackage->GlobalPackageId.ValueForDebugging());
 		}
 		Package->ImportedPackages.Append(LocalizedPackages);
 	}
@@ -3836,7 +3865,8 @@ int32 CreateTarget(const FIoStoreArguments& Arguments, const FIoStoreWriterSetti
 	GlobalPackageData.Reserve(PackageAssetData.ObjectExports.Num());
 
 	CreateGlobalScriptObjects(GlobalPackageData, Arguments.TargetPlatform);
-	CreateGlobalImportsAndExports(Arguments, Packages, PackageIdMap, PackageAssetData, GlobalPackageData, ExportGraph);
+	TMap<FPackageId, FPackageId> SourceToRedirectedPackageMap;
+	CreateGlobalImportsAndExports(Arguments, Packages, PackageIdMap, PackageAssetData, GlobalPackageData, ExportGraph, SourceToRedirectedPackageMap);
 
 	// Mapped import and exports are required before processing localization, and preload/postload arcs
 	MapExportEntryIndices(PackageAssetData.ObjectExports, GlobalPackageData.ExportObjects, Packages);
@@ -3867,6 +3897,7 @@ int32 CreateTarget(const FIoStoreArguments& Arguments, const FIoStoreWriterSetti
 		PackageAssetData,
 		GlobalPackageData,
 		SourceToLocalizedPackageMap,
+		SourceToRedirectedPackageMap,
 		ExportGraph,
 		Packages);
 
