IoStore:
Refactor IoStore staging process so that package store/new loader related code can be shared with the cooker
#rnx

diff --git a/Engine/Source/Developer/IoStoreUtilities/Private/IoStoreUtilities.cpp b/Engine/Source/Developer/IoStoreUtilities/Private/IoStoreUtilities.cpp
--- a/Engine/Source/Developer/IoStoreUtilities/Private/IoStoreUtilities.cpp
+++ b/Engine/Source/Developer/IoStoreUtilities/Private/IoStoreUtilities.cpp
@@ -50,6 +50,7 @@
 #include "Misc/StringBuilder.h"
 #include "Async/Future.h"
 #include "Algo/MaxElement.h"
+#include "PackageStoreOptimizer.h"
 
 //PRAGMA_DISABLE_OPTIMIZATION
 
@@ -61,8 +62,6 @@ DEFINE_LOG_CATEGORY_STATIC(LogIoStore, Log, All);
 #define IOSTORE_CPU_SCOPE_DATA(NAME, DATA) TRACE_CPUPROFILER_EVENT_SCOPE(IoStore##NAME);
 
 #define OUTPUT_CHUNKID_DIRECTORY 0
-#define OUTPUT_NAMEMAP_CSV 0
-#define OUTPUT_DEBUG_PACKAGE_EXPORT_BUNDLES 0
 
 static const FName DefaultCompressionMethod = NAME_Zlib;
 static const uint64 DefaultCompressionBlockSize = 64 << 10;
@@ -261,194 +260,6 @@ static void LoadKeyChain(const TCHAR* CmdLine, FKeyChain& OutCryptoSettings)
 	OutCryptoSettings.MasterEncryptionKey = OutCryptoSettings.EncryptionKeys.Find(EncryptionKeyOverrideGuid);
 }
 
-class FNameMapBuilder
-{
-public:
-	void SetNameMapType(FMappedName::EType InNameMapType)
-	{
-		NameMapType = InNameMapType;
-	}
-
-	void AddName(const FName& Name)
-	{
-		const FNameEntryId ComparisonIndex = Name.GetComparisonIndex();
-		const FNameEntryId DisplayIndex = Name.GetDisplayIndex();
-		NameMap.Add(DisplayIndex);
-		int32 Index = NameMap.Num();
-		NameIndices.Add(ComparisonIndex, Index);
-	}
-
-	void MarkNamesAsReferenced(const TArray<FName>& Names, TArray<int32>& OutNameIndices)
-	{
-		for (const FName& Name : Names)
-		{
-			const FNameEntryId ComparisonIndex = Name.GetComparisonIndex();
-			const FNameEntryId DisplayIndex = Name.GetDisplayIndex();
-			int32& Index = NameIndices.FindOrAdd(ComparisonIndex);
-			if (Index == 0)
-			{
-				NameMap.Add(DisplayIndex);
-				Index = NameMap.Num();
-			}
-
-			OutNameIndices.Add(Index - 1);
-		}
-	}
-
-	void MarkNameAsReferenced(const FName& Name)
-	{
-		const FNameEntryId ComparisonIndex = Name.GetComparisonIndex();
-		const FNameEntryId DisplayIndex = Name.GetDisplayIndex();
-		int32& Index = NameIndices.FindOrAdd(ComparisonIndex);
-		if (Index == 0)
-		{
-			NameMap.Add(DisplayIndex);
-			Index = NameMap.Num();
-		}
-#if OUTPUT_NAMEMAP_CSV
-		// debug counts
-		{
-			const int32 Number = Name.GetNumber();
-			TTuple<int32, int32, int32>& Counts = DebugNameCounts.FindOrAdd(ComparisonIndex);
-
-			if (Number == 0)
-			{
-				++Counts.Get<0>();
-			}
-			else
-			{
-				++Counts.Get<1>();
-				if (Number > Counts.Get<2>())
-				{
-					Counts.Get<2>() = Number;
-				}
-			}
-		}
-#endif
-	}
-
-	FMappedName MapName(const FName& Name) const
-	{
-		const FNameEntryId Id = Name.GetComparisonIndex();
-		const int32* Index = NameIndices.Find(Id);
-		check(Index);
-		return FMappedName::Create(*Index - 1, Name.GetNumber(), NameMapType);
-	}
-
-	const TArray<FNameEntryId>& GetNameMap() const
-	{
-		return NameMap;
-	}
-
-	friend FArchive& operator<<(FArchive& Ar, FNameMapBuilder& NameMapBuilder)
-	{
-		if (Ar.IsSaving())
-		{
-			int32 NameCount = NameMapBuilder.NameMap.Num();
-			Ar << NameCount;
-			for (FNameEntryId NameEntryId : NameMapBuilder.NameMap)
-			{
-				const FNameEntry* NameEntry = FName::GetEntry(NameEntryId);
-				NameEntry->Write(Ar);
-			}
-		}
-		else
-		{
-			int32 NameCount;
-			Ar << NameCount;
-			for (int32 NameIndex = 0; NameIndex < NameCount; ++NameIndex)
-			{
-				FNameEntrySerialized NameEntrySerialized(ENAME_LinkerConstructor);
-				Ar << NameEntrySerialized;
-				FName Name(NameEntrySerialized);
-				NameMapBuilder.NameMap.Add(Name.GetDisplayIndex());
-				NameMapBuilder.NameIndices.Add(Name.GetComparisonIndex(), NameIndex + 1);
-			}
-		}
-		return Ar;
-	}
-
-#if OUTPUT_NAMEMAP_CSV
-	void SaveCsv(const FString& CsvFilePath)
-	{
-		{
-			TUniquePtr<FArchive> CsvArchive(IFileManager::Get().CreateFileWriter(*CsvFilePath));
-			if (CsvArchive)
-			{
-				TCHAR Name[FName::StringBufferSize];
-				ANSICHAR Line[MAX_SPRINTF + FName::StringBufferSize];
-				ANSICHAR Header[] = "Length\tMaxNumber\tNumberCount\tBaseCount\tTotalCount\tFName\n";
-				CsvArchive->Serialize(Header, sizeof(Header) - 1);
-				for (auto& Counts : DebugNameCounts)
-				{
-					const int32 NameLen = FName::CreateFromDisplayId(Counts.Key, 0).ToString(Name);
-					FCStringAnsi::Sprintf(Line, "%d\t%d\t%d\t%d\t%d\t",
-						NameLen, Counts.Value.Get<2>(), Counts.Value.Get<1>(), Counts.Value.Get<0>(), Counts.Value.Get<0>() + Counts.Value.Get<1>());
-					ANSICHAR* L = Line + FCStringAnsi::Strlen(Line);
-					const TCHAR* N = Name;
-					while (*N)
-					{
-						*L++ = CharCast<ANSICHAR,TCHAR>(*N++);
-					}
-					*L++ = '\n';
-					CsvArchive.Get()->Serialize(Line, L - Line);
-				}
-			}
-		}
-	}
-#endif
-
-	void Empty()
-	{
-		NameIndices.Empty();
-		NameMap.Empty();
-#if OUTPUT_NAMEMAP_CSV
-		DebugNameCounts.Empty();
-#endif
-	}
-private:
-	TMap<FNameEntryId, int32> NameIndices;
-	TArray<FNameEntryId> NameMap;
-#if OUTPUT_NAMEMAP_CSV
-	TMap<FNameEntryId, TTuple<int32,int32,int32>> DebugNameCounts; // <Number0Count,OtherNumberCount,MaxNumber>
-#endif
-	FMappedName::EType NameMapType = FMappedName::EType::Package;
-};
-
-class FNameReaderProxyArchive
-	: public FArchiveProxy
-{
-	const TArray<FNameEntryId>& NameMap;
-
-public:
-	using FArchiveProxy::FArchiveProxy;
-
-	FNameReaderProxyArchive(FArchive& InAr, const TArray<FNameEntryId>& InNameMap)
-		: FArchiveProxy(InAr)
-		, NameMap(InNameMap)
-	{ 
-		// Replicate the filter editor only state of the InnerArchive as FArchiveProxy will
-		// not intercept it.
-		FArchive::SetFilterEditorOnly(InAr.IsFilterEditorOnly());
-	}
-
-	FArchive& operator<<(FName& Name)
-	{
-		int32 NameIndex, Number;
-		InnerArchive << NameIndex << Number;
-
-		if (!NameMap.IsValidIndex(NameIndex))
-		{
-			UE_LOG(LogStreaming, Fatal, TEXT("Bad name index %i/%i"), NameIndex, NameMap.Num());
-		}
-
-		const FNameEntryId MappedName = NameMap[NameIndex];
-		Name = FName::CreateFromDisplayId(MappedName, Number);
-
-		return *this;
-	}
-};
-
 struct FContainerSourceFile 
 {
 	FString NormalizedPath;
@@ -482,13 +293,24 @@ struct FCookedFileStatData
 
 using FCookedFileStatMap = TMap<FString, FCookedFileStatData>;
 
-struct FPackage;
 struct FContainerTargetSpec;
 
+struct FLegacyCookedPackage
+{
+	FPackageId GlobalPackageId;
+	FName PackageName;
+	FName RedirectFromPackageName;
+	FString FileName;
+	uint64 UAssetSize = 0;
+	uint64 UExpSize = 0;
+	uint64 DiskLayoutOrder = -1;
+	FPackageStorePackage* OptimizedPackage = nullptr;
+};
+
 struct FContainerTargetFile
 {
 	FContainerTargetSpec* ContainerTarget = nullptr;
-	FPackage* Package = nullptr;
+	FLegacyCookedPackage* Package = nullptr;
 	FString NormalizedSourcePath;
 	FString TargetPath;
 	FString DestinationPath;
@@ -496,21 +318,11 @@ struct FContainerTargetFile
 	uint64 IdealOrder = 0;
 	FIoChunkId ChunkId;
 	TArray<uint8> PackageHeaderData;
-	TArray<int32> NameIndices;
-	FNameMapBuilder* NameMapBuilder = nullptr;
 	bool bIsBulkData = false;
 	bool bIsOptionalBulkData = false;
 	bool bIsMemoryMappedBulkData = false;
 	bool bForceUncompressed = false;
 
-	int64 UGraphSize = 0;
-	int64 NameMapSize = 0;
-	int64 ImportMapSize = 0;
-	int64 ExportMapSize = 0;
-	int64 ExportBundlesHeaderSize = 0;
-
-	uint64 HeaderSerialSize = 0;
-
 	TArray<FFileRegion> FileRegions;
 };
 
@@ -549,6 +361,7 @@ struct FIoStoreArguments
 
 struct FContainerTargetSpec
 {
+	FIoContainerId ContainerId;
 	FContainerHeader Header;
 	FName Name;
 	FGuid EncryptionKeyGuid;
@@ -556,54 +369,13 @@ struct FContainerTargetSpec
 	FIoStoreWriter* IoStoreWriter;
 	TArray<FContainerTargetFile> TargetFiles;
 	TArray<TUniquePtr<FIoStoreReader>> PatchSourceReaders;
-	FNameMapBuilder LocalNameMapBuilder;
-	FNameMapBuilder* NameMapBuilder = &LocalNameMapBuilder;
 	EIoContainerFlags ContainerFlags = EIoContainerFlags::None;
 	uint32 PackageCount = 0;
-	bool bUseLocalNameMap = false;
 	bool bGenerateDiffPatch = false;
 };
 
-struct FPackageAssetData
-{
-	TArray<FObjectImport> ObjectImports;
-	TArray<FObjectExport> ObjectExports;
-	TArray<FPackageIndex> PreloadDependencies;
-};
-
-struct FPackage;
-using FPackageNameMap = TMap<FName, FPackage*>;
-using FPackageIdMap = TMap<FPackageId, FPackage*>;
-using FSourceToLocalizedPackageMultimap = TMultiMap<FPackage*, FPackage*>;
-using FLocalizedToSourceImportIndexMap = TMap<FPackageObjectIndex, FPackageObjectIndex>;
-
-static constexpr TCHAR L10NString[] = TEXT("/L10N/");
-static constexpr TCHAR ScriptPrefix[] = TEXT("/Script/");
-
-// modified copy from PakFileUtilities
-static FString RemapLocalizationPathIfNeeded(const FString& Path, FString* OutRegion)
-{
-	static constexpr int32 L10NPrefixLength = sizeof(L10NString) / sizeof(TCHAR) - 1;
-
-	int32 BeginL10NOffset = Path.Find(L10NString, ESearchCase::IgnoreCase);
-	if (BeginL10NOffset >= 0)
-	{
-		int32 EndL10NOffset = BeginL10NOffset + L10NPrefixLength;
-		int32 NextSlashIndex = Path.Find(TEXT("/"), ESearchCase::IgnoreCase, ESearchDir::FromStart, EndL10NOffset);
-		int32 RegionLength = NextSlashIndex - EndL10NOffset;
-		if (RegionLength >= 2)
-		{
-			FString NonLocalizedPath = Path.Mid(0, BeginL10NOffset) + Path.Mid(NextSlashIndex);
-			if (OutRegion)
-			{
-				*OutRegion = Path.Mid(EndL10NOffset, RegionLength);
-				OutRegion->ToLowerInline();
-			}
-			return NonLocalizedPath;
-		}
-	}
-	return Path;
-}
+using FPackageNameMap = TMap<FName, FLegacyCookedPackage*>;
+using FPackageIdMap = TMap<FPackageId, FLegacyCookedPackage*>;
 
 #if OUTPUT_CHUNKID_DIRECTORY
 class FChunkIdCsv
@@ -652,2369 +424,340 @@ static FIoChunkId CreateChunkId(FPackageId GlobalPackageId, uint16 ChunkIndex, E
 	return ChunkId;
 }
 
-enum EPreloadDependencyType
-{
-	PreloadDependencyType_Create,
-	PreloadDependencyType_Serialize,
-};
-
-struct FArc
+static void AssignPackagesDiskOrder(
+	const TArray<FLegacyCookedPackage*>& Packages,
+	const TMap<FName, uint64> GameOrderMap,
+	const TMap<FName, uint64>& CookerOrderMap,
+	const FPackageIdMap& PackageIdMap)
 {
-	uint32 FromNodeIndex;
-	uint32 ToNodeIndex;
-
-	bool operator==(const FArc& Other) const
+	struct FCluster
 	{
-		return ToNodeIndex == Other.ToNodeIndex && FromNodeIndex == Other.FromNodeIndex;
-	}
-};
+		TArray<FLegacyCookedPackage*> Packages;
+	};
 
-struct FExportGraphNode;
+	TArray<FCluster*> Clusters;
+	TSet<FLegacyCookedPackage*> AssignedPackages;
+	TArray<FLegacyCookedPackage*> ProcessStack;
 
-struct FExportBundle
-{
-	TArray<FExportGraphNode*> Nodes;
-	uint32 LoadOrder;
-};
+	struct FPackageAndOrder
+	{
+		FLegacyCookedPackage* Package;
+		uint64 GameOpenOrder;
+		uint64 CookerOpenOrder;
 
-struct FPackageGraphNode
-{
-	FPackage* Package = nullptr;
-	mutable bool bTemporaryMark = false;
-	mutable bool bPermanentMark = false;
-};
+		bool operator<(const FPackageAndOrder& Other) const
+		{
+			if (GameOpenOrder != Other.GameOpenOrder)
+			{
+				return GameOpenOrder < Other.GameOpenOrder;
+			}
+			if (CookerOpenOrder != Other.CookerOpenOrder)
+			{
+				return CookerOpenOrder < Other.CookerOpenOrder;
+			}
+			// Fallback to reverse bundle order (so that packages are considered before their imports)
+			return Package->OptimizedPackage->GetLoadOrder() > Other.Package->OptimizedPackage->GetLoadOrder();
+		}
+	};
 
-class FPackageGraph
-{
-public:
-	FPackageGraph()
+	TArray<FPackageAndOrder> SortedPackages;
+	SortedPackages.Reserve(Packages.Num());
+	for (FLegacyCookedPackage* Package : Packages)
 	{
-
+		FPackageAndOrder& Entry = SortedPackages.AddDefaulted_GetRef();
+		Entry.Package = Package;
+		const uint64* FindGameOpenOrder = GameOrderMap.Find(Package->PackageName);
+		Entry.GameOpenOrder = FindGameOpenOrder ? *FindGameOpenOrder : MAX_uint64;
+		const uint64* FindCookerOpenOrder = CookerOrderMap.Find(Package->PackageName);
+		/*if (!FindCookerOpenOrder)
+		{
+			UE_LOG(LogIoStore, Warning, TEXT("Missing cooker order for package: %s"), *Package->Name.ToString());
+		}*/
+		Entry.CookerOpenOrder = FindCookerOpenOrder ? *FindCookerOpenOrder : MAX_uint64;
 	}
-
-	~FPackageGraph()
+	bool bHasGameOrder = true;
+	bool bHasCookerOrder = true;
+	int32 LastAssignedCount = 0;
+	Algo::Sort(SortedPackages);
+	for (FPackageAndOrder& Entry : SortedPackages)
 	{
-		for (FPackageGraphNode* Node : Nodes)
+		if (bHasGameOrder && Entry.GameOpenOrder == MAX_uint64)
+		{
+			UE_LOG(LogIoStore, Display, TEXT("Ordered %d/%d packages using game open order"), AssignedPackages.Num(), Packages.Num());
+			LastAssignedCount = AssignedPackages.Num();
+			bHasGameOrder = false;
+		}
+		if (!bHasGameOrder && bHasCookerOrder && Entry.CookerOpenOrder == MAX_uint64)
+		{
+			UE_LOG(LogIoStore, Display, TEXT("Ordered %d/%d packages using cooker open order"), AssignedPackages.Num() - LastAssignedCount, Packages.Num() - LastAssignedCount);
+			LastAssignedCount = AssignedPackages.Num();
+			bHasCookerOrder = false;
+		}
+		if (!AssignedPackages.Contains(Entry.Package))
 		{
-			delete Node;
+			FCluster* Cluster = new FCluster();
+			Clusters.Add(Cluster);
+			ProcessStack.Push(Entry.Package);
+
+			while (ProcessStack.Num())
+			{
+				FLegacyCookedPackage* PackageToProcess = ProcessStack.Pop(false);
+				if (!AssignedPackages.Contains(PackageToProcess))
+				{
+					AssignedPackages.Add(PackageToProcess);
+					if (PackageToProcess->OptimizedPackage->GetExportBundleCount())
+					{
+						Cluster->Packages.Add(PackageToProcess);
+					}
+					TArray<FPackageId> AllReferencedPackageIds;
+					AllReferencedPackageIds.Append(PackageToProcess->OptimizedPackage->GetImportedPackageIds());
+					AllReferencedPackageIds.Append(PackageToProcess->OptimizedPackage->GetImportedRedirectedPackageIds().Array());
+					for (const FPackageId& ReferencedPackageId : AllReferencedPackageIds)
+					{
+						FLegacyCookedPackage* FindReferencedPackage = PackageIdMap.FindRef(ReferencedPackageId);
+						if (FindReferencedPackage)
+						{
+							ProcessStack.Push(FindReferencedPackage);
+						}
+					}
+				}
+			}
 		}
 	}
+	UE_LOG(LogIoStore, Display, TEXT("Ordered %d packages using fallback bundle order"), AssignedPackages.Num() - LastAssignedCount);
 
-	FPackageGraphNode* AddNode(FPackage* Package)
+	check(AssignedPackages.Num() == Packages.Num());
+	
+	for (FCluster* Cluster : Clusters)
 	{
-		FPackageGraphNode* Node = new FPackageGraphNode();
-		Node->Package = Package;
-		Nodes.Add(Node);
-		return Node;
+		Algo::Sort(Cluster->Packages, [](const FLegacyCookedPackage* A, const FLegacyCookedPackage* B)
+		{
+			return A->OptimizedPackage->GetLoadOrder() < B->OptimizedPackage->GetLoadOrder();
+		});
 	}
 
-	void AddImportDependency(FPackageGraphNode* FromNode, FPackageGraphNode* ToNode)
+	uint64 LayoutIndex = 0;
+	for (FCluster* Cluster : Clusters)
 	{
-		Edges.Add(FromNode, ToNode);
+		for (FLegacyCookedPackage* Package : Cluster->Packages)
+		{
+			Package->DiskLayoutOrder = LayoutIndex++;
+		}
+		delete Cluster;
 	}
+}
 
-	TArray<FPackage*> TopologicalSort() const;
-
-private:
-	TArray<FPackageGraphNode*> Nodes;
-	TMultiMap<FPackageGraphNode*, FPackageGraphNode*> Edges;
-};
-
-struct FExportGraphNode
+static void CreateDiskLayout(
+	const TArray<FContainerTargetSpec*>& ContainerTargets,
+	const TArray<FLegacyCookedPackage*>& Packages,
+	const TMap<FName, uint64>& PackageOrderMap,
+	const TMap<FName, uint64>& CookerOrderMap,
+	const FPackageIdMap& PackageIdMap)
 {
-	FPackage* Package;
-	FExportBundleEntry BundleEntry;
-	TSet<FExportGraphNode*> ExternalDependencies;
-	TSet<FPackageId> BaseGamePackageDependencies;
-	uint64 NodeIndex;
-};
+	IOSTORE_CPU_SCOPE(CreateDiskLayout);
 
-class FExportGraph
-{
-public:
-	FExportGraph(int32 NumExports, int32 NumPreloadDependencies)
-	{
-		Nodes.Reserve(NumExports * 2);
-		Edges.Reserve(NumExports + NumPreloadDependencies);
-	}
+	AssignPackagesDiskOrder(Packages, PackageOrderMap, CookerOrderMap, PackageIdMap);
 
-	~FExportGraph()
+	for (FContainerTargetSpec* ContainerTarget : ContainerTargets)
 	{
-		for (FExportGraphNode* Node : Nodes)
+		TArray<FContainerTargetFile*> SortedTargetFiles;
+		SortedTargetFiles.Reserve(ContainerTarget->TargetFiles.Num());
+		for (FContainerTargetFile& TargetFile : ContainerTarget->TargetFiles)
 		{
-			delete Node;
+			SortedTargetFiles.Add(&TargetFile);
+		}
+		Algo::Sort(SortedTargetFiles, [](const FContainerTargetFile* A, const FContainerTargetFile* B)
+		{
+			if (A->bIsMemoryMappedBulkData != B->bIsMemoryMappedBulkData)
+			{
+				// Put all memory mapped bulkdata last
+				return B->bIsMemoryMappedBulkData;
+			}
+			else if (A->bIsBulkData != B->bIsBulkData)
+			{
+				// Put packages before bulkdata
+				return B->bIsBulkData;
+			}
+			else if (A->Package != B->Package)
+			{
+				return A->Package->DiskLayoutOrder < B->Package->DiskLayoutOrder;
+			}
+			else if (A->bIsOptionalBulkData != B->bIsOptionalBulkData)
+			{
+				// Put each ubulk right before it's corresponding uptnl
+				return B->bIsOptionalBulkData;
+			}
+			check(A == B)
+			return false;
+		});
+		uint64 IdealOrder = 0;
+		for (FContainerTargetFile* TargetFile : SortedTargetFiles)
+		{
+			TargetFile->IdealOrder = IdealOrder++;
 		}
 	}
+}
 
-	FExportGraphNode* AddNode(FPackage* Package, const FExportBundleEntry& BundleEntry)
-	{
-		FExportGraphNode* Node = new FExportGraphNode();
-		Node->Package = Package;
-		Node->BundleEntry = BundleEntry;
-		Node->NodeIndex = Nodes.Num();
-		Nodes.Add(Node);
-		return Node;
-	}
-
-	void AddInternalDependency(FExportGraphNode* FromNode, FExportGraphNode* ToNode)
+static EIoChunkType BulkdataTypeToChunkIdType(FPackageStoreBulkDataManifest::EBulkdataType Type)
+{
+	switch (Type)
 	{
-		AddEdge(FromNode, ToNode);
-	}
-
-	void AddExternalDependency(FExportGraphNode* FromNode, FExportGraphNode* ToNode)
-	{
-		AddEdge(FromNode, ToNode);
-		ToNode->ExternalDependencies.Add(FromNode);
-	}
-
-	TArray<FExportGraphNode*> ComputeLoadOrder(const TArray<FPackage*>& Packages) const;
-
-private:
-	void AddEdge(FExportGraphNode* FromNode, FExportGraphNode* ToNode)
-	{
-		Edges.Add(FromNode, ToNode);
-	}
-
-	TArray<FExportGraphNode*> TopologicalSort() const;
-
-	TArray<FExportGraphNode*> Nodes;
-	TMultiMap<FExportGraphNode*, FExportGraphNode*> Edges;
-};
-
-struct FPackage
-{
-	FName Name;
-	FName SourcePackageName; // for localized packages
-	FString FileName;
-	FPackageId GlobalPackageId;
-	FString Region; // for localized packages
-	FPackageId SourceGlobalPackageId; // for localized packages
-	FPackageId RedirectedPackageId;
-	uint32 PackageFlags = 0;
-	uint32 CookedHeaderSize = 0;
-	int32 NameCount = 0;
-	int32 ImportCount = 0;
-	int32 PreloadDependencyCount = 0;
-	int32 ExportCount = 0;
-	int32 ImportIndexOffset = -1;
-	int32 ExportIndexOffset = -1;
-	int32 PreloadIndexOffset = -1;
-	int64 UExpSize = 0;
-	int64 UAssetSize = 0;
-	int64 SummarySize = 0;
-	uint64 ExportsSerialSize = 0;
-	bool bIsLocalizedAndConformed = false;
-
-	TArray<FPackage*> ImportedPackages;
-	TArray<FPackageId> ImportedPackageIds;
-
-	TArray<FName> SummaryNames;
-	FNameMapBuilder LocalNameMapBuilder;
-	
-	TArray<FPackageObjectIndex> Imports;
-	TArray<int32> Exports;
-	TMap<FPackageId, TArray<FArc>> ExternalArcs;
-	
-	TArray<FExportBundle> ExportBundles;
-	TMap<FExportGraphNode*, uint32> ExportBundleMap;
-
-	TArray<FExportGraphNode*> CreateExportNodes;
-	TArray<FExportGraphNode*> SerializeExportNodes;
-
-	TArray<FExportGraphNode*> NodesWithNoIncomingEdges;
-	FPackageGraphNode* Node = nullptr;
-
-	uint64 DiskLayoutOrder = MAX_uint64;
-};
-
-struct FCircularImportChain
-{
-	TArray<FName> SortedNames;
-	TArray<FPackage*> Packages;
-	uint32 Hash = 0;
-
-	FCircularImportChain()
-	{
-		Packages.Reserve(128);
-	}
-
-	void Add(FPackage* Package)
-	{
-		Packages.Add(Package);
-	}
-
-	void Pop()
-	{
-		Packages.Pop();
-	}
-
-	int32 Num()
-	{
-		return Packages.Num();
-	}
-
-	void SortAndGenerateHash()
-	{
-		SortedNames.Empty(Packages.Num());
-		for (FPackage* Package : Packages)
-		{
-			SortedNames.Emplace(Package->Name);
-		}
-		SortedNames.Sort(FNameLexicalLess());
-		Hash = CityHash32((char*)SortedNames.GetData(), SortedNames.Num() * SortedNames.GetTypeSize());
-	}
-
-	FString ToString()
-	{
-		FString Result = FString::Printf(TEXT("%d:%u: "), SortedNames.Num(), Hash);
-		for (const FName& Name : SortedNames)
-		{
-			Result.Append(Name.ToString());
-			Result.Append(TEXT(" -> "));
-		}
-		Result.Append(SortedNames[0].ToString());
-		return Result;
-	}
-
-	bool operator==(const FCircularImportChain& Other) const
-	{
-		return Hash == Other.Hash && SortedNames == Other.SortedNames;
-	}
-
-	friend FORCEINLINE uint32 GetTypeHash(const FCircularImportChain& In)
-	{
-		return In.Hash;
-	}
-};
-
-TArray<FPackage*> FPackageGraph::TopologicalSort() const
-{
-	TMap<const FPackageGraphNode*, TArray<const FPackageGraphNode*>> SortedEdges;
-	for (const auto& KV : Edges)
-	{
-		const FPackageGraphNode* Source = KV.Key;
-		const FPackageGraphNode* Target = KV.Value;
-		TArray<const FPackageGraphNode*>& SourceArray = SortedEdges.FindOrAdd(Source);
-		SourceArray.Add(Target);
-	}
-	for (auto& KV : SortedEdges)
-	{
-		TArray<const FPackageGraphNode*>& SourceArray = KV.Value;
-		Algo::Sort(SourceArray, [](const FPackageGraphNode* A, const FPackageGraphNode* B)
-		{
-			return A->Package->GlobalPackageId < B->Package->GlobalPackageId;
-		});
-	}
-
-	TArray<FPackage*> Result;
-	Result.Reserve(Nodes.Num());
-	
-	struct
-	{
-		void Visit(const FPackageGraphNode* Node)
-		{
-			if (Node->bPermanentMark || Node->bTemporaryMark)
-			{
-				return;
-			}
-			Node->bTemporaryMark = true;
-			TArray<const FPackageGraphNode*>* TargetNodes = Edges.Find(Node);
-			if(TargetNodes)
-			{
-				for (const FPackageGraphNode* ToNode : *TargetNodes)
-				{
-					Visit(ToNode);
-				}
-			}
-			Node->bTemporaryMark = false;
-			Node->bPermanentMark = true;
-			Result.Add(Node->Package);
-		}
-
-		TMap<const FPackageGraphNode*, TArray<const FPackageGraphNode*>>& Edges;
-		TArray<FPackage*>& Result;
-
-	} Visitor{ SortedEdges, Result };
-
-	for (FPackageGraphNode* Node : Nodes)
-	{
-		Visitor.Visit(Node);
-	}
-	check(Result.Num() == Nodes.Num());
-
-	Algo::Reverse(Result);
-	return Result;
-}
-
-TArray<FExportGraphNode*> FExportGraph::ComputeLoadOrder(const TArray<FPackage*>& Packages) const
-{
-	IOSTORE_CPU_SCOPE(ComputeLoadOrder);
-	FPackageGraph PackageGraph;
-	{
-		for (FPackage* Package : Packages)
-		{
-			Package->Node = PackageGraph.AddNode(Package);
-		}
-		for (FPackage* Package : Packages)
-		{
-			for (FPackage* ImportedPackage : Package->ImportedPackages)
-			{
-				PackageGraph.AddImportDependency(ImportedPackage->Node, Package->Node);
-			}
-		}
-	}
-
-	TArray<FPackage*> SortedPackages = PackageGraph.TopologicalSort();
-	
-	int32 NodeCount = Nodes.Num();
-	TArray<uint32> NodesIncomingEdgeCount;
-	NodesIncomingEdgeCount.AddZeroed(NodeCount);
-
-	TMultiMap<FExportGraphNode*, FExportGraphNode*> EdgesCopy = Edges;
-	for (auto& KV : Edges)
-	{
-		FExportGraphNode* ToNode = KV.Value;
-		++NodesIncomingEdgeCount[ToNode->NodeIndex];
-	}
-
-	TArray<FExportGraphNode*> LoadOrder;
-	LoadOrder.Reserve(NodeCount);
-	
-	auto NodeSorter = [](const FExportGraphNode& A, const FExportGraphNode& B)
-	{
-		if (A.BundleEntry.LocalExportIndex == B.BundleEntry.LocalExportIndex)
-		{
-			return A.BundleEntry.CommandType < B.BundleEntry.CommandType;
-		}
-		return A.BundleEntry.LocalExportIndex < B.BundleEntry.LocalExportIndex;
-	};
-
-	for (int32 NodeIndex = 0; NodeIndex < NodeCount; ++NodeIndex)
-	{
-		if (NodesIncomingEdgeCount[NodeIndex] == 0)
-		{
-			FExportGraphNode* Node = Nodes[NodeIndex];
-			Node->Package->NodesWithNoIncomingEdges.HeapPush(Node, NodeSorter);
-		}
-	}
-	while (LoadOrder.Num() < NodeCount)
-	{
-		for (FPackage* Package : SortedPackages)
-		{
-			while (Package->NodesWithNoIncomingEdges.Num() > 0)
-			{
-				FExportGraphNode* RemovedNode;
-				Package->NodesWithNoIncomingEdges.HeapPop(RemovedNode, NodeSorter, false);
-				LoadOrder.Add(RemovedNode);
-				for (auto EdgeIt = EdgesCopy.CreateKeyIterator(RemovedNode); EdgeIt; ++EdgeIt)
-				{
-					FExportGraphNode* ToNode = EdgeIt.Value();
-					if (--NodesIncomingEdgeCount[ToNode->NodeIndex] == 0)
-					{
-						ToNode->Package->NodesWithNoIncomingEdges.HeapPush(ToNode, NodeSorter);
-					}
-					EdgeIt.RemoveCurrent();
-				}
-			}
-		}
-	}
-
-	return LoadOrder;
-}
-
-static void AddInternalExportArc(FExportGraph& ExportGraph, FPackage& Package, uint32 FromExportIndex, EPreloadDependencyType FromPhase, uint32 ToExportIndex, EPreloadDependencyType ToPhase)
-{
-	FExportGraphNode* FromNode = FromPhase == PreloadDependencyType_Create ? Package.CreateExportNodes[FromExportIndex] : Package.SerializeExportNodes[FromExportIndex];
-	FExportGraphNode* ToNode = ToPhase == PreloadDependencyType_Create ? Package.CreateExportNodes[ToExportIndex] : Package.SerializeExportNodes[ToExportIndex];
-	ExportGraph.AddInternalDependency(FromNode, ToNode);
-}
-
-static void AddExternalExportArc(FExportGraph& ExportGraph, FPackage& FromPackage, uint32 FromExportIndex, EPreloadDependencyType FromPhase, FPackage& ToPackage, uint32 ToExportIndex, EPreloadDependencyType ToPhase)
-{
-	FExportGraphNode* FromNode = FromPhase == PreloadDependencyType_Create ? FromPackage.CreateExportNodes[FromExportIndex] : FromPackage.SerializeExportNodes[FromExportIndex];
-	FExportGraphNode* ToNode = ToPhase == PreloadDependencyType_Create ? ToPackage.CreateExportNodes[ToExportIndex] : ToPackage.SerializeExportNodes[ToExportIndex];
-	ExportGraph.AddExternalDependency(FromNode, ToNode);
-}
-
-static void AddBaseGamePackageArc(FExportGraph& ExportGraph, FPackageId FromPackageId, FPackage& ToPackage, uint32 ToExportIndex, EPreloadDependencyType ToPhase)
-{
-	FExportGraphNode* ToNode = ToPhase == PreloadDependencyType_Create ? ToPackage.CreateExportNodes[ToExportIndex] : ToPackage.SerializeExportNodes[ToExportIndex];
-	ToNode->BaseGamePackageDependencies.Add(FromPackageId);
-}
-
-static void AddUniqueExternalBundleArc(FPackageId FromPackageId, uint32 FromBundleIndex, FPackage& ToPackage, uint32 ToBundleIndex)
-{
-	TArray<FArc>& ExternalArcs = ToPackage.ExternalArcs.FindOrAdd(FromPackageId);
-	ExternalArcs.AddUnique({ FromBundleIndex, ToBundleIndex });
-}
-
-static void BuildBundles(FExportGraph& ExportGraph, const TArray<FPackage*>& Packages)
-{
-	IOSTORE_CPU_SCOPE(BuildBundles)
-	UE_LOG(LogIoStore, Display, TEXT("Building bundles..."));
-
-	TArray<FExportGraphNode*> ExportLoadOrder = ExportGraph.ComputeLoadOrder(Packages);
-	FPackage* LastPackage = nullptr;
-	uint32 BundleLoadOrder = 0;
-	for (FExportGraphNode* Node : ExportLoadOrder)
-	{
-		FPackage* Package = Node->Package;
-		check(Package);
-		if (!Package)
-		{
-			continue;
-		}
-
-		uint32 BundleIndex;
-		FExportBundle* Bundle;
-		if (Package != LastPackage)
-		{
-			BundleIndex = Package->ExportBundles.Num();
-			Bundle = &Package->ExportBundles.AddDefaulted_GetRef();
-			Bundle->LoadOrder = BundleLoadOrder++;
-			LastPackage = Package;
-		}
-		else
-		{
-			BundleIndex = Package->ExportBundles.Num() - 1;
-			Bundle = &Package->ExportBundles[BundleIndex];
-		}
-		for (FExportGraphNode* ExternalDependency : Node->ExternalDependencies)
-		{
-			uint32* FindDependentBundleIndex = ExternalDependency->Package->ExportBundleMap.Find(ExternalDependency);
-			check(FindDependentBundleIndex);
-			check(*FindDependentBundleIndex < uint32(ExternalDependency->Package->ExportBundles.Num()));
-			AddUniqueExternalBundleArc(ExternalDependency->Package->GlobalPackageId, *FindDependentBundleIndex, *Package, BundleIndex);
-		}
-		for (FPackageId FromPackageId : Node->BaseGamePackageDependencies)
-		{
-			AddUniqueExternalBundleArc(FromPackageId, MAX_uint32, *Package, BundleIndex);
-		}
-		Bundle->Nodes.Add(Node);
-		Package->ExportBundleMap.Add(Node, BundleIndex);
-	}
-}
-
-static void AssignPackagesDiskOrder(
-	const TArray<FPackage*>& Packages,
-	const TMap<FName, uint64> GameOrderMap,
-	const TMap<FName, uint64>& CookerOrderMap)
-{
-	struct FCluster
-	{
-		TArray<FPackage*> Packages;
-	};
-
-	TArray<FCluster*> Clusters;
-	TSet<FPackage*> AssignedPackages;
-	TArray<FPackage*> ProcessStack;
-
-	struct FPackageAndOrder
-	{
-		FPackage* Package;
-		uint64 GameOpenOrder;
-		uint64 CookerOpenOrder;
-
-		bool operator<(const FPackageAndOrder& Other) const
-		{
-			if (GameOpenOrder != Other.GameOpenOrder)
-			{
-				return GameOpenOrder < Other.GameOpenOrder;
-			}
-			if (CookerOpenOrder != Other.CookerOpenOrder)
-			{
-				return CookerOpenOrder < Other.CookerOpenOrder;
-			}
-			// Fallback to reverse bundle order (so that packages are considered before their imports)
-			return Package->ExportBundles[0].LoadOrder > Other.Package->ExportBundles[0].LoadOrder;
-		}
-	};
-
-	TArray<FPackageAndOrder> SortedPackages;
-	SortedPackages.Reserve(Packages.Num());
-	for (FPackage* Package : Packages)
-	{
-		FPackageAndOrder& Entry = SortedPackages.AddDefaulted_GetRef();
-		Entry.Package = Package;
-		const uint64* FindGameOpenOrder = GameOrderMap.Find(Package->Name);
-		Entry.GameOpenOrder = FindGameOpenOrder ? *FindGameOpenOrder : MAX_uint64;
-		const uint64* FindCookerOpenOrder = CookerOrderMap.Find(Package->Name);
-		/*if (!FindCookerOpenOrder)
-		{
-			UE_LOG(LogIoStore, Warning, TEXT("Missing cooker order for package: %s"), *Package->Name.ToString());
-		}*/
-		Entry.CookerOpenOrder = FindCookerOpenOrder ? *FindCookerOpenOrder : MAX_uint64;
-	}
-	bool bHasGameOrder = true;
-	bool bHasCookerOrder = true;
-	int32 LastAssignedCount = 0;
-	Algo::Sort(SortedPackages);
-	for (FPackageAndOrder& Entry : SortedPackages)
-	{
-		if (bHasGameOrder && Entry.GameOpenOrder == MAX_uint64)
-		{
-			UE_LOG(LogIoStore, Display, TEXT("Ordered %d/%d packages using game open order"), AssignedPackages.Num(), Packages.Num());
-			LastAssignedCount = AssignedPackages.Num();
-			bHasGameOrder = false;
-		}
-		if (!bHasGameOrder && bHasCookerOrder && Entry.CookerOpenOrder == MAX_uint64)
-		{
-			UE_LOG(LogIoStore, Display, TEXT("Ordered %d/%d packages using cooker open order"), AssignedPackages.Num() - LastAssignedCount, Packages.Num() - LastAssignedCount);
-			LastAssignedCount = AssignedPackages.Num();
-			bHasCookerOrder = false;
-		}
-		if (!AssignedPackages.Contains(Entry.Package))
-		{
-			FCluster* Cluster = new FCluster();
-			Clusters.Add(Cluster);
-			ProcessStack.Push(Entry.Package);
-
-			while (ProcessStack.Num())
-			{
-				FPackage* PackageToProcess = ProcessStack.Pop(false);
-				if (!AssignedPackages.Contains(PackageToProcess))
-				{
-					AssignedPackages.Add(PackageToProcess);
-					if (PackageToProcess->ExportBundles.Num())
-					{
-						Cluster->Packages.Add(PackageToProcess);
-					}
-					for (FPackage* ImportedPackage : PackageToProcess->ImportedPackages)
-					{
-						ProcessStack.Push(ImportedPackage);
-					}
-				}
-			}
-		}
-	}
-	UE_LOG(LogIoStore, Display, TEXT("Ordered %d packages using fallback bundle order"), AssignedPackages.Num() - LastAssignedCount);
-
-	check(AssignedPackages.Num() == Packages.Num());
-	
-	for (FCluster* Cluster : Clusters)
-	{
-		Algo::Sort(Cluster->Packages, [](const FPackage* A, const FPackage* B)
-		{
-			return A->ExportBundles[0].LoadOrder < B->ExportBundles[0].LoadOrder;
-		});
-	}
-
-	uint64 LayoutIndex = 0;
-	for (FCluster* Cluster : Clusters)
-	{
-		for (FPackage* Package : Cluster->Packages)
-		{
-			Package->DiskLayoutOrder = LayoutIndex++;
-		}
-		delete Cluster;
-	}
-}
-
-static void CreateDiskLayout(
-	const TArray<FContainerTargetSpec*>& ContainerTargets,
-	const TArray<FPackage*>& Packages,
-	const TMap<FName, uint64> PackageOrderMap,
-	const TMap<FName, uint64>& CookerOrderMap)
-{
-	IOSTORE_CPU_SCOPE(CreateDiskLayout);
-
-	AssignPackagesDiskOrder(Packages, PackageOrderMap, CookerOrderMap);
-
-	for (FContainerTargetSpec* ContainerTarget : ContainerTargets)
-	{
-		TArray<FContainerTargetFile*> SortedTargetFiles;
-		SortedTargetFiles.Reserve(ContainerTarget->TargetFiles.Num());
-		for (FContainerTargetFile& TargetFile : ContainerTarget->TargetFiles)
-		{
-			SortedTargetFiles.Add(&TargetFile);
-		}
-		Algo::Sort(SortedTargetFiles, [](const FContainerTargetFile* A, const FContainerTargetFile* B)
-		{
-			if (A->bIsMemoryMappedBulkData != B->bIsMemoryMappedBulkData)
-			{
-				// Put all memory mapped bulkdata last
-				return B->bIsMemoryMappedBulkData;
-			}
-			else if (A->bIsBulkData != B->bIsBulkData)
-			{
-				// Put packages before bulkdata
-				return B->bIsBulkData;
-			}
-			else if (A->Package != B->Package)
-			{
-				return A->Package->DiskLayoutOrder < B->Package->DiskLayoutOrder;
-			}
-			else if (A->bIsOptionalBulkData != B->bIsOptionalBulkData)
-			{
-				// Put each ubulk right before it's corresponding uptnl
-				return B->bIsOptionalBulkData;
-			}
-			check(A == B)
-			return false;
-		});
-		uint64 IdealOrder = 0;
-		for (FContainerTargetFile* TargetFile : SortedTargetFiles)
-		{
-			TargetFile->IdealOrder = IdealOrder++;
-		}
-	}
-}
-
-static EIoChunkType BulkdataTypeToChunkIdType(FPackageStoreBulkDataManifest::EBulkdataType Type)
-{
-	switch (Type)
-	{
-	case FPackageStoreBulkDataManifest::EBulkdataType::Normal:
-		return EIoChunkType::BulkData;
-	case FPackageStoreBulkDataManifest::EBulkdataType::Optional:
-		return EIoChunkType::OptionalBulkData;
-	case FPackageStoreBulkDataManifest::EBulkdataType::MemoryMapped:
-		return EIoChunkType::MemoryMappedBulkData;
-	default:
-		UE_LOG(LogIoStore, Error, TEXT("Invalid EBulkdataType (%d) found!"), Type);
-		return EIoChunkType::Invalid;
-	}
-}
-
-struct FScriptObjectData
-{
-	FName ObjectName;
-	FString FullName;
-	FPackageObjectIndex GlobalIndex;
-	FPackageObjectIndex OuterIndex;
-	FPackageObjectIndex CDOClassIndex;
-
-	friend FArchive& operator<<(FArchive& Ar, FScriptObjectData& Data)
-	{
-		FString ObjectNameStr;
-		if (Ar.IsSaving())
-		{
-			ObjectNameStr = Data.ObjectName.ToString();
-			Ar << ObjectNameStr;
-		}
-		else
-		{
-			Ar << ObjectNameStr;
-			Data.ObjectName = FName(*ObjectNameStr);
-		}
-		Ar << Data.FullName;
-		Ar << Data.GlobalIndex;
-		Ar << Data.OuterIndex;
-		Ar << Data.CDOClassIndex;
-		return Ar;
-	}
-};
-
-struct FExportObjectData
-{
-	FName ObjectName;
-	FString FullName;
-	int32 GlobalIndex = -1;
-	int32 SourceIndex = -1;
-	FPackageObjectIndex GlobalImportIndex;
-	FPackageObjectIndex OuterIndex;
-	FPackageObjectIndex ClassIndex;
-	FPackageObjectIndex SuperIndex;
-	FPackageObjectIndex TemplateIndex;
-	EObjectFlags ObjectFlags = RF_NoFlags;
-
-	FPackage* Package = nullptr;
-	FExportGraphNode* CreateNode = nullptr;
-	FExportGraphNode* SerializeNode = nullptr;
-
-	bool IsPublicExport() const
-	{
-		return !!(ObjectFlags & RF_Public);
-	}
-
-	friend FArchive& operator<<(FArchive& Ar, FExportObjectData& Data)
-	{
-		FString ObjectNameStr;
-		if (Ar.IsSaving())
-		{
-			ObjectNameStr = Data.ObjectName.ToString();
-			Ar << ObjectNameStr;
-		}
-		else
-		{
-			Ar << ObjectNameStr;
-			Data.ObjectName = FName(*ObjectNameStr);
-		}
-		Ar << Data.FullName;
-		Ar << Data.GlobalIndex;
-		Ar << Data.SourceIndex;
-		Ar << Data.SourceIndex;
-		Ar << Data.GlobalImportIndex;
-		Ar << Data.OuterIndex;
-		Ar << Data.ClassIndex;
-		Ar << Data.TemplateIndex;
-		Ar << (uint32&)Data.ObjectFlags;
-		return Ar;
-	}
-};
-
-using FImportObjectsByFullName = TMap<FString, FPackageObjectIndex>;
-using FImportObjectsById = TMap<FPackageObjectIndex, FPackageObjectIndex>;
-using FExportObjectsByFullName = TMap<FString, int32>;
-
-using FGlobalScriptObjects = TMap<FPackageObjectIndex, FScriptObjectData>;
-using FGlobalExportObjects = TArray<FExportObjectData>;
-
-struct FGlobalPackageData
-{
-	FGlobalScriptObjects ScriptObjects;
-	FGlobalExportObjects ExportObjects;
-	TMap<FPackageObjectIndex, int32> PublicExportIndices;
-	FImportObjectsByFullName ImportsByFullName;
-	FExportObjectsByFullName ExportsByFullName;
-
-	void Reserve(int32 TotalExportCount)
-	{
-		const int32 EstimatedPublicExportObjectCount = TotalExportCount / 10;
-		const int32 EstimatedScriptObjectCount = 64000;
-		ExportObjects.Reserve(TotalExportCount);
-		ExportsByFullName.Reserve(TotalExportCount);
-		PublicExportIndices.Reserve(EstimatedPublicExportObjectCount);
-		ScriptObjects.Reserve(EstimatedScriptObjectCount);
-		ImportsByFullName.Reserve(EstimatedScriptObjectCount + EstimatedPublicExportObjectCount);
-	}
-
-	const FExportObjectData* FindPublicExport(FPackageObjectIndex Index) const
-	{
-		check(Index.IsPackageImport());
-		if (const int32* GlobalExportIndex = PublicExportIndices.Find(Index))
-		{
-			return &ExportObjects[*GlobalExportIndex];
-		}
-		return nullptr;
-	}
-
-	FName GetObjectName(FPackageObjectIndex Index, const TArray<int32>* PackageExportIndices) const
-	{
-		if (Index.IsScriptImport())
-		{
-			const FScriptObjectData* ScriptObjectData = ScriptObjects.Find(Index);
-			check(ScriptObjectData);
-			return ScriptObjectData->ObjectName;
-		}
-		if (Index.IsPackageImport())
-		{
-			const int32* GlobalExportIndex = PublicExportIndices.Find(Index);
-			check(GlobalExportIndex);
-			return ExportObjects[*GlobalExportIndex].ObjectName;
-		}
-		if (Index.IsExport() && PackageExportIndices)
-		{
-			int32 GlobalExportIndex = (*PackageExportIndices)[Index.ToExport()];
-			return ExportObjects[GlobalExportIndex].ObjectName;
-		}
-		return FName();
-	}
-};
-
-static void FindImportFullName(
-	const TMap<FName, FName>& Redirects,
-	TArray<FString>& ImportFullNames,
-	FObjectImport* ImportMap,
-	const int32 LocalImportIndex)
-{
-	FString& FullName = ImportFullNames[LocalImportIndex];
-
-	if (FullName.Len() == 0)
-	{
-		FullName.Reserve(256);
-
-		FObjectImport* Import = &ImportMap[LocalImportIndex];
-		if (Import->OuterIndex.IsNull())
-		{
-			FName PackageName = Import->ObjectName;
-			if (const FName* RedirectedName = Redirects.Find(PackageName))
-			{
-				PackageName = *RedirectedName;
-			}
-			PackageName.AppendString(FullName);
-			FullName.ToLowerInline();
-		}
-		else
-		{
-			const int32 OuterIndex = Import->OuterIndex.ToImport();
-			FindImportFullName(Redirects, ImportFullNames, ImportMap, OuterIndex);
-			const FString& OuterName = ImportFullNames[OuterIndex];
-			check(OuterName.Len() > 0);
-
-			FullName.Append(OuterName);
-			FullName.AppendChar(TEXT('/'));
-			Import->ObjectName.AppendString(FullName);
-			FullName.ToLowerInline();
-		}
-	}
-}
-
-static FPackageObjectIndex FindAndVerifyGlobalImport(
-	const FPackage* Package,
-	FGlobalPackageData& GlobalPackageData,
-	FObjectImport& Import,
-	const FString& FullName,
-	const FString& DLCPrefix)
-{
-	FPackageObjectIndex GlobalImportIndex = GlobalPackageData.ImportsByFullName.FindRef(FullName);
-	if (GlobalImportIndex.IsNull())
-	{
-		bool bIsPackage = Import.OuterIndex.IsNull();
-		bool bIsScript = FullName.StartsWith(ScriptPrefix);
-		if (bIsPackage)
-		{
-			if (bIsScript)
-			{
-				UE_LOG(LogIoStore, Display, TEXT("For package '%s' (0x%llX): Missing import script package '%s'. Editor only?"),
-					*Package->Name.ToString(),
-					Package->GlobalPackageId.ValueForDebugging(),
-					*FullName);
-			}
-			else
-			{
-				// UPackages are never serialized, hence they are never imports
-			}
-		}
-		else
-		{
-			if (bIsScript)
-			{
-				UE_LOG(LogIoStore, Display, TEXT("For package '%s' (0x%llX): Missing import script object '%s'. Editor only?"),
-					*Package->Name.ToString(),
-					Package->GlobalPackageId.ValueForDebugging(),
-					*FullName);
-			}
-			else if (!DLCPrefix.Len() || FullName.StartsWith(DLCPrefix))
-			{
-				UE_LOG(LogIoStore, Display, TEXT("For package '%s' (0x%llX): Missing import object '%s' due to missing public export. Editor only?"),
-					*Package->Name.ToString(),
-					Package->GlobalPackageId.ValueForDebugging(),
-					*FullName);
-			}
-		}
-	}
-	return GlobalImportIndex;
-}
-
-static int32 FindExport(
-	FGlobalPackageData& GlobalPackageData,
-	TArray<FString>& TempFullNames,
-	const FObjectExport* ExportMap,
-	const int32 LocalExportIndex,
-	FPackage* Package)
-{
-	FString& FullName = TempFullNames[LocalExportIndex];
-
-	if (FullName.Len() == 0)
-	{
-		FullName.Reserve(256);
-
-		const FObjectExport* Export = ExportMap + LocalExportIndex;
-		if (Export->OuterIndex.IsNull())
-		{
-			if (Package->RedirectedPackageId.IsValid())
-			{
-				Package->SourcePackageName.AppendString(FullName);
-			}
-			else
-			{
-				Package->Name.AppendString(FullName);
-			}
-			FullName.AppendChar(TEXT('/'));
-			Export->ObjectName.AppendString(FullName);
-			FullName.ToLowerInline();
-		}
-		else
-		{
-			check(Export->OuterIndex.IsExport());
-
-			FindExport(GlobalPackageData, TempFullNames, ExportMap, Export->OuterIndex.ToExport(), Package);
-			FString& OuterName = TempFullNames[Export->OuterIndex.ToExport()];
-			check(OuterName.Len() > 0);
-			FullName.Append(OuterName);
-			FullName.AppendChar(TEXT('/'));
-			Export->ObjectName.AppendString(FullName);
-			FullName.ToLowerInline();
-		}
-
-		int32 GlobalExportIndex = -1;
-		int32* FindGlobalExportIndex = GlobalPackageData.ExportsByFullName.Find(FullName);
-		if (!FindGlobalExportIndex)
-		{
-			GlobalExportIndex = GlobalPackageData.ExportObjects.Num();
-			GlobalPackageData.ExportsByFullName.Add(FullName, GlobalExportIndex);
-			GlobalPackageData.ExportObjects.AddDefaulted();
-		}
-		else
-		{
-			GlobalExportIndex = *FindGlobalExportIndex;
-		}
-		FExportObjectData& ExportData = GlobalPackageData.ExportObjects[GlobalExportIndex];
-		
-		ExportData.GlobalIndex = GlobalExportIndex;
-		ExportData.Package = Package;
-		ExportData.ObjectName = Export->ObjectName;
-		ExportData.SourceIndex = LocalExportIndex;
-		ExportData.FullName = FullName;
-		ExportData.ObjectFlags = Export->ObjectFlags;
-
-		return GlobalExportIndex;
-	}
-	else
-	{
-		int32* GlobalExportIndex = GlobalPackageData.ExportsByFullName.Find(FullName);
-		check(GlobalExportIndex);
-		return *GlobalExportIndex;
-	}
-}
-
-FContainerTargetSpec* AddContainer(
-	FName Name,
-	TArray<FContainerTargetSpec*>& Containers)
-{
-	FIoContainerId ContainerId = FIoContainerId::FromName(Name);
-	for (FContainerTargetSpec* ExistingContainer : Containers)
-	{
-		if (ExistingContainer->Name == Name)
-		{
-			UE_LOG(LogIoStore, Fatal, TEXT("Duplicate container name: '%s'"), *Name.ToString());
-			return nullptr;
-		}
-		if (ExistingContainer->Header.ContainerId == ContainerId)
-		{
-			UE_LOG(LogIoStore, Fatal, TEXT("Hash collision for container names: '%s' and '%s'"), *Name.ToString(), *ExistingContainer->Name.ToString());
-			return nullptr;
-		}
-	}
-	
-	FContainerTargetSpec* ContainerTargetSpec = new FContainerTargetSpec();
-	ContainerTargetSpec->Name = Name;
-	ContainerTargetSpec->Header.ContainerId = ContainerId;
-	Containers.Add(ContainerTargetSpec);
-	return ContainerTargetSpec;
-}
-
-FPackage* FindOrAddPackage(
-	const FIoStoreArguments& Arguments,
-	const TCHAR* RelativeFileName,
-	TArray<FPackage*>& Packages,
-	FPackageNameMap& PackageNameMap,
-	FPackageIdMap& PackageIdMap)
-{
-	FString PackageName;
-	FString ErrorMessage;
-	if (!FPackageName::TryConvertFilenameToLongPackageName(RelativeFileName, PackageName, &ErrorMessage))
-	{
-		UE_LOG(LogIoStore, Warning, TEXT("Failed to obtain package name from file name '%s'"), *ErrorMessage);
-		return nullptr;
-	}
-
-	FName PackageFName = *PackageName;
-
-	FPackage* Package = PackageNameMap.FindRef(PackageFName);
-	if (!Package)
-	{
-		FPackageId PackageId = FPackageId::FromName(PackageFName);
-		if (FPackage* FindById = PackageIdMap.FindRef(PackageId))
-		{
-			UE_LOG(LogIoStore, Fatal, TEXT("Package name hash collision \"%s\" and \"%s"), *FindById->Name.ToString(), *PackageFName.ToString());
-		}
-
-		if (const FName* ReleasedPackageName = Arguments.ReleasedPackages.PackageIdToName.Find(PackageId))
-		{
-			UE_LOG(LogIoStore, Fatal, TEXT("Package name hash collision \"%s\" and \"%s"), *ReleasedPackageName->ToString(), *PackageFName.ToString());
-		}
-
-		Package = new FPackage();
-		Package->Name = PackageFName;
-		Package->GlobalPackageId = PackageId;
-
-		if (Arguments.IsDLC() && Arguments.bRemapPluginContentToGame)
-		{
-			const int32 DLCNameLen = Arguments.DLCName.Len() + 1;
-			FString RedirectedPackageNameStr = TEXT("/Game");
-			RedirectedPackageNameStr.AppendChars(*PackageName + DLCNameLen, PackageName.Len() - DLCNameLen);
-			FName RedirectedPackageName = FName(*RedirectedPackageNameStr);
-
-			if (Arguments.ReleasedPackages.PackageNames.Contains(RedirectedPackageName))
-			{
-				Package->SourcePackageName = RedirectedPackageName;
-				Package->RedirectedPackageId = FPackageId::FromName(RedirectedPackageName);
-			}
-		}
-		else
-		{
-			Package->SourcePackageName = *RemapLocalizationPathIfNeeded(PackageName, &Package->Region);
-		}
-	
-		Packages.Add(Package);
-		PackageNameMap.Add(PackageFName, Package);
-		PackageIdMap.Add(PackageId, Package);
-	}
-
-	return Package;
-}
-
-static bool ConformLocalizedPackage(
-	const FPackageNameMap& PackageMap,
-	FGlobalPackageData& GlobalPackageData,
-	const FPackage& SourcePackage,
-	FPackage& LocalizedPackage,
-	FLocalizedToSourceImportIndexMap& LocalizedToSourceImportIndexMap)
-{
-	const int32 ExportCount =
-		SourcePackage.ExportCount < LocalizedPackage.ExportCount ?
-		SourcePackage.ExportCount :
-		LocalizedPackage.ExportCount;
-
-	UE_CLOG(SourcePackage.ExportCount != LocalizedPackage.ExportCount, LogIoStore, Verbose,
-		TEXT("For culture '%s': Localized package '%s' (0x%llX) for source package '%s' (0x%llX)  - Has ExportCount %d vs. %d"),
-			*LocalizedPackage.Region,
-			*LocalizedPackage.Name.ToString(),
-			LocalizedPackage.GlobalPackageId.ValueForDebugging(),
-			*LocalizedPackage.SourcePackageName.ToString(),
-			SourcePackage.GlobalPackageId.ValueForDebugging(),
-			LocalizedPackage.ExportCount,
-			SourcePackage.ExportCount);
-
-	auto GetExportNameSafe = [](
-		const FString& ExportFullName,
-		const FName& PackageName,
-		int32 PackageNameLen) -> const TCHAR*
-	{
-		const bool bValidNameLen = ExportFullName.Len() > PackageNameLen + 1;
-		if (bValidNameLen)
-		{
-			const TCHAR* ExportNameStr = *ExportFullName + PackageNameLen;
-			const bool bValidNameFormat = *ExportNameStr == '/';
-			if (bValidNameFormat)
-			{
-				return ExportNameStr + 1; // skip verified '/'
-			}
-			else
-			{
-				UE_CLOG(!bValidNameFormat, LogIoStore, Warning,
-					TEXT("Export name '%s' should start with '/' at position %d, i.e. right after package prefix '%s'"),
-					*ExportFullName,
-					PackageNameLen,
-					*PackageName.ToString());
-			}
-		}
-		else
-		{
-			UE_CLOG(!bValidNameLen, LogIoStore, Warning,
-				TEXT("Export name '%s' with length %d should be longer than package name '%s' with length %d"),
-				*ExportFullName,
-				PackageNameLen,
-				*PackageName.ToString());
-		}
-
-		return nullptr;
-	};
-
-	auto AppendMismatchMessage = [&GlobalPackageData, &LocalizedPackage, &SourcePackage](
-		const TCHAR* Text, FName ExportName, FPackageObjectIndex LocIndex, FPackageObjectIndex SrcIndex, FString& FailReason)
-	{
-		FString LocString = GlobalPackageData.GetObjectName(LocIndex, &LocalizedPackage.Exports).ToString();
-		FString SrcString = GlobalPackageData.GetObjectName(SrcIndex, &SourcePackage.Exports).ToString();
-
-		FailReason.Appendf(TEXT("Public export '%s' has %s %s vs. %s"),
-			*ExportName.ToString(),
-			Text,
-			*LocString,
-			*SrcString);
-	};
-
-	const int32 LocalizedPackageNameLen = LocalizedPackage.Name.GetStringLength();
-	const int32 SourcePackageNameLen = SourcePackage.Name.GetStringLength();
-
-	TArray <TPair<int32, int32>, TInlineAllocator<64>> NewPublicExports;
-	NewPublicExports.Reserve(ExportCount);
-
-	bool bSuccess = true;
-	int32 LocalizedIndex = 0;
-	int32 SourceIndex = 0;
-	while (LocalizedIndex < ExportCount && SourceIndex < ExportCount)
-	{
-		FString FailReason;
-		const FExportObjectData& LocalizedExportData = GlobalPackageData.ExportObjects[LocalizedPackage.Exports[LocalizedIndex]];
-		const FExportObjectData& SourceExportData = GlobalPackageData.ExportObjects[SourcePackage.Exports[SourceIndex]];
-
-		const TCHAR* LocalizedExportStr = GetExportNameSafe(
-			LocalizedExportData.FullName, LocalizedPackage.Name, LocalizedPackageNameLen);
-		const TCHAR* SourceExportStr = GetExportNameSafe(
-			SourceExportData.FullName, SourcePackage.Name, SourcePackageNameLen);
-
-		if (!LocalizedExportStr || !SourceExportStr)
-		{
-			UE_LOG(LogIoStore, Error,
-				TEXT("Culture '%s': Localized package '%s' (0x%llX) for source package '%s' (0x%llX) - Has some bad data from an earlier phase."),
-				*LocalizedPackage.Region,
-				*LocalizedPackage.Name.ToString(),
-				LocalizedPackage.GlobalPackageId.ValueForDebugging(),
-				*LocalizedPackage.SourcePackageName.ToString(),
-				SourcePackage.GlobalPackageId.ValueForDebugging())
-			return false;
-		}
-
-		int32 CompareResult = FCString::Stricmp(LocalizedExportStr, SourceExportStr);
-		if (CompareResult < 0)
-		{
-			++LocalizedIndex;
-
-			if (LocalizedExportData.IsPublicExport())
-			{
-				// public localized export is missing in the source package, so just keep it as it is
-				NewPublicExports.Emplace(LocalizedIndex - 1, 1);
-			}
-		}
-		else if (CompareResult > 0)
-		{
-			++SourceIndex;
-
-			if (SourceExportData.IsPublicExport())
-			{
-				FailReason.Appendf(TEXT("Public source export '%s' is missing in the localized package"),
-					*SourceExportData.ObjectName.ToString());
-			}
-		}
-		else
-		{
-			++LocalizedIndex;
-			++SourceIndex;
-
-			if (SourceExportData.IsPublicExport())
-			{
-				if (!LocalizedExportData.IsPublicExport())
-				{
-					FailReason.Appendf(TEXT("Public source export '%s' exists in the localized package")
-						TEXT(", but is not a public localized export."),
-						*SourceExportData.ObjectName.ToString());
-				}
-				else if (LocalizedExportData.ClassIndex != SourceExportData.ClassIndex)
-				{
-					AppendMismatchMessage(TEXT("class"), LocalizedExportData.ObjectName,
-						LocalizedExportData.ClassIndex, SourceExportData.ClassIndex, FailReason);
-				}
-				else if (LocalizedExportData.TemplateIndex != SourceExportData.TemplateIndex)
-				{
-					AppendMismatchMessage(TEXT("template"), LocalizedExportData.ObjectName,
-						LocalizedExportData.TemplateIndex, SourceExportData.TemplateIndex, FailReason);
-				}
-				else if (LocalizedExportData.SuperIndex != SourceExportData.SuperIndex)
-				{
-					AppendMismatchMessage(TEXT("super"), LocalizedExportData.ObjectName,
-						LocalizedExportData.SuperIndex, SourceExportData.SuperIndex, FailReason);
-				}
-				else
-				{
-					NewPublicExports.Emplace(LocalizedIndex - 1, SourceIndex - 1);
-				}
-			}
-			else if (LocalizedExportData.IsPublicExport())
-			{
-				FailReason.Appendf(TEXT("Public localized export '%s' exists in the source package")
-					TEXT(", but is not a public source export."),
-					*LocalizedExportData.ObjectName.ToString());
-			}
-		}
-
-		if (FailReason.Len() > 0)
-		{
-			UE_LOG(LogIoStore, Warning,
-				TEXT("Culture '%s': Localized package '%s' (0x%llX) for '%s' (0x%llX) - %s"),
-				*LocalizedPackage.Region,
-				*LocalizedPackage.Name.ToString(),
-				LocalizedPackage.GlobalPackageId.ValueForDebugging(),
-				*LocalizedPackage.SourcePackageName.ToString(),
-				SourcePackage.GlobalPackageId.ValueForDebugging(),
-				*FailReason);
-			bSuccess = false;
-		}
-	}
-
-	if (bSuccess)
-	{
-		for (TPair<int32, int32>& Pair : NewPublicExports)
-		{
-			FExportObjectData& LocalizedExportData = GlobalPackageData.ExportObjects[LocalizedPackage.Exports[Pair.Key]];
-			if (Pair.Value != -1)
-			{
-				const FExportObjectData& SourceExportData = GlobalPackageData.ExportObjects[SourcePackage.Exports[Pair.Value]];
-
-				LocalizedToSourceImportIndexMap.Add(LocalizedExportData.GlobalImportIndex, SourceExportData.GlobalImportIndex);
-
-				LocalizedExportData.GlobalImportIndex = SourceExportData.GlobalImportIndex;
-			}
-		}
-	}
-
-	return bSuccess;
-}
-
-static void AddPreloadDependencies(
-	const FPackageAssetData& PackageAssetData,
-	const FGlobalPackageData& GlobalPackageData,
-	const FSourceToLocalizedPackageMultimap& SourceToLocalizedPackageMap,
-	const TMap<FPackageId, FPackageId>& SourceToRedirectedPackageMap,
-	FExportGraph& ExportGraph,
-	TArray<FPackage*>& Packages)
-{
-	IOSTORE_CPU_SCOPE(PreLoadDependencies);
-	UE_LOG(LogIoStore, Display, TEXT("Adding preload dependencies..."));
-
-	TSet<FPackageId> ExternalPackageDependencies;
-	TArray<FPackage*> LocalizedPackages;
-	for (FPackage* Package : Packages)
-	{
-		ExternalPackageDependencies.Reset();
-
-		// Convert PreloadDependencies to arcs
-		for (int32 I = 0; I < Package->ExportCount; ++I)
-		{
-			const FObjectExport& ObjectExport = PackageAssetData.ObjectExports[Package->ExportIndexOffset + I];
-			int32 PreloadDependenciesBaseIndex = Package->PreloadIndexOffset;
-
-			FPackageIndex ExportPackageIndex = FPackageIndex::FromExport(I);
-
-			auto AddPreloadArc = [&](FPackageIndex Dep, EPreloadDependencyType PhaseFrom, EPreloadDependencyType PhaseTo)
-			{
-				if (Dep.IsExport())
-				{
-					AddInternalExportArc(ExportGraph, *Package, Dep.ToExport(), PhaseFrom, I, PhaseTo);
-				}
-				else
-				{
-					FPackageObjectIndex ImportIndex = Package->Imports[Dep.ToImport()];
-					if (ImportIndex.IsPackageImport())
-					{
-						// When building DLC's exports can be missing
-						if (const FExportObjectData* Export = GlobalPackageData.FindPublicExport(ImportIndex))
-						{
-							check(Export->GlobalImportIndex == ImportIndex);
-
-							AddExternalExportArc(ExportGraph, *Export->Package, Export->SourceIndex, PhaseFrom, *Package, I, PhaseTo);
-
-							LocalizedPackages.Reset();
-							SourceToLocalizedPackageMap.MultiFind(Export->Package, LocalizedPackages);
-							for (FPackage* LocalizedPackage : LocalizedPackages)
-							{
-								UE_LOG(LogIoStore, Verbose, TEXT("For package '%s' (0x%llX): Adding localized preload dependency '%s' in '%s'"),
-									*Package->Name.ToString(),
-									Package->GlobalPackageId.ValueForDebugging(),
-									*Export->ObjectName.ToString(),
-									*LocalizedPackage->Name.ToString());
-
-								AddExternalExportArc(ExportGraph, *LocalizedPackage, Export->SourceIndex, PhaseFrom, *Package, I, PhaseTo);
-							}
-						}
-						else
-						{
-							const FObjectImport* ImportMap = PackageAssetData.ObjectImports.GetData() + Package->ImportIndexOffset;
-							const FObjectImport* Import = ImportMap + Dep.ToImport();
-							while (!Import->OuterIndex.IsNull())
-							{
-								Import = ImportMap + Import->OuterIndex.ToImport();
-							}
-
-							FPackageId PackageId = FPackageId::FromName(Import->ObjectName);
-							if (SourceToRedirectedPackageMap.Contains(PackageId))
-							{
-								// Edge case with an import from a redirection source package but the corresponding
-								// export is missing from the redirection target package so don't add a base game
-								// package arc
-								return;
-							}
-							bool bIsAlreadyInSet = false;
-							ExternalPackageDependencies.Add(PackageId, &bIsAlreadyInSet);
-							if (!bIsAlreadyInSet)
-							{
-								AddBaseGamePackageArc(ExportGraph, PackageId, *Package, I, PhaseTo);
-							}
-						}
-					}
-				}
-			};
-
-			if (PreloadDependenciesBaseIndex >= 0 && ObjectExport.FirstExportDependency >= 0)
-			{
-				int32 RunningIndex = PreloadDependenciesBaseIndex + ObjectExport.FirstExportDependency;
-				for (int32 Index = ObjectExport.SerializationBeforeSerializationDependencies; Index > 0; Index--)
-				{
-					FPackageIndex Dep = PackageAssetData.PreloadDependencies[RunningIndex++];
-					check(!Dep.IsNull());
-					AddPreloadArc(Dep, PreloadDependencyType_Serialize, PreloadDependencyType_Serialize);
-				}
-
-				for (int32 Index = ObjectExport.CreateBeforeSerializationDependencies; Index > 0; Index--)
-				{
-					FPackageIndex Dep = PackageAssetData.PreloadDependencies[RunningIndex++];
-					check(!Dep.IsNull());
-					AddPreloadArc(Dep, PreloadDependencyType_Create, PreloadDependencyType_Serialize);
-				}
-
-				for (int32 Index = ObjectExport.SerializationBeforeCreateDependencies; Index > 0; Index--)
-				{
-					FPackageIndex Dep = PackageAssetData.PreloadDependencies[RunningIndex++];
-					check(!Dep.IsNull());
-					AddPreloadArc(Dep, PreloadDependencyType_Serialize, PreloadDependencyType_Create);
-				}
-
-				for (int32 Index = ObjectExport.CreateBeforeCreateDependencies; Index > 0; Index--)
-				{
-					FPackageIndex Dep = PackageAssetData.PreloadDependencies[RunningIndex++];
-					check(!Dep.IsNull());
-					// can't create this export until these things are created
-					AddPreloadArc(Dep, PreloadDependencyType_Create, PreloadDependencyType_Create);
-				}
-			}
-		}
-	}
-};
-
-void FinalizeNameMaps(FContainerTargetSpec& ContainerTarget)
-{
-	for (FContainerTargetFile& TargetFile : ContainerTarget.TargetFiles)
-	{
-		if (TargetFile.bIsBulkData)
-		{
-			continue;
-		}
-		FPackage* Package = TargetFile.Package;
-		TargetFile.NameMapBuilder->MarkNameAsReferenced(Package->Name);
-		TargetFile.NameMapBuilder->MarkNameAsReferenced(Package->SourcePackageName);
-		TargetFile.NameMapBuilder->MarkNamesAsReferenced(Package->SummaryNames, TargetFile.NameIndices);
-	}
-}
-
-void FinalizePackageHeaders(
-	FContainerTargetSpec& ContainerTarget,
-	const TArray<FObjectExport>& ObjectExports,
-	const TArray<FExportObjectData>& GlobalExports,
-	const FImportObjectsByFullName& GlobalImportsByFullName)
-{
-	for (FContainerTargetFile& TargetFile : ContainerTarget.TargetFiles)
-	{
-		if (TargetFile.bIsBulkData)
-		{
-			continue;
-		}
-
-		FPackage* Package = TargetFile.Package;
-
-		// Temporary Archive for serializing ImportMap
-		FBufferWriter ImportMapArchive(nullptr, 0, EBufferWriterFlags::AllowResize | EBufferWriterFlags::TakeOwnership);
-		for (FPackageObjectIndex& GlobalImportIndex : Package->Imports)
-		{
-			ImportMapArchive << GlobalImportIndex;
-		}
-		TargetFile.ImportMapSize = ImportMapArchive.Tell();
-
-		// Temporary Archive for serializing EDL graph data
-		FBufferWriter GraphArchive(nullptr, 0, EBufferWriterFlags::AllowResize | EBufferWriterFlags::TakeOwnership);
-
-		int32 ReferencedPackagesCount = Package->ExternalArcs.Num();
-		GraphArchive << ReferencedPackagesCount;
-		TArray<TTuple<FPackageId, TArray<FArc>>> SortedExternalArcs;
-		SortedExternalArcs.Reserve(Package->ExternalArcs.Num());
-		for (auto& KV : Package->ExternalArcs)
-		{
-			FPackageId ImportedPackageId = KV.Key;
-			TArray<FArc> SortedArcs = KV.Value;
-			Algo::Sort(SortedArcs, [](const FArc& A, const FArc& B)
-			{
-				if (A.FromNodeIndex == B.FromNodeIndex)
-				{
-					return A.ToNodeIndex < B.ToNodeIndex;
-				}
-				return A.FromNodeIndex < B.FromNodeIndex;
-			});
-			SortedExternalArcs.Emplace(ImportedPackageId, MoveTemp(SortedArcs));
-		}
-		Algo::Sort(SortedExternalArcs, [](const TTuple<FPackageId, TArray<FArc>>& A, const TTuple<FPackageId, TArray<FArc>>& B)
-		{
-			return A.Key < B.Key;
-		});
-		for (auto& KV : SortedExternalArcs)
-		{
-			FPackageId ImportedPackageId = KV.Key;
-			TArray<FArc>& Arcs = KV.Value;
-			int32 ExternalArcCount = Arcs.Num();
-
-			GraphArchive << ImportedPackageId;
-			GraphArchive << ExternalArcCount;
-			GraphArchive.Serialize(Arcs.GetData(), ExternalArcCount * sizeof(FArc));
-		}
-		TargetFile.UGraphSize = GraphArchive.Tell();
-
-		// Temporary Archive for serializing export map data
-		FBufferWriter ExportMapArchive(nullptr, 0, EBufferWriterFlags::AllowResize | EBufferWriterFlags::TakeOwnership);
-		for (int32 I = 0; I < Package->ExportCount; ++I)
-		{
-			const FObjectExport& ObjectExport = ObjectExports[Package->ExportIndexOffset + I];
-			const FExportObjectData& ExportData = GlobalExports[Package->Exports[I]];
-
-			FExportMapEntry ExportMapEntry;
-			ExportMapEntry.CookedSerialOffset = ObjectExport.SerialOffset;
-			ExportMapEntry.CookedSerialSize = ObjectExport.SerialSize;
-			ExportMapEntry.ObjectName = TargetFile.NameMapBuilder->MapName(ObjectExport.ObjectName);
-			ExportMapEntry.OuterIndex = ExportData.OuterIndex;
-			ExportMapEntry.ClassIndex = ExportData.ClassIndex;
-			ExportMapEntry.SuperIndex = ExportData.SuperIndex;
-			ExportMapEntry.TemplateIndex = ExportData.TemplateIndex;
-			ExportMapEntry.GlobalImportIndex = ExportData.GlobalImportIndex;
-			ExportMapEntry.ObjectFlags = ObjectExport.ObjectFlags;
-			ExportMapEntry.FilterFlags = EExportFilterFlags::None;
-			if (ObjectExport.bNotForClient)
-			{
-				ExportMapEntry.FilterFlags = EExportFilterFlags::NotForClient;
-			}
-			else if (ObjectExport.bNotForServer)
-			{
-				ExportMapEntry.FilterFlags = EExportFilterFlags::NotForServer;
-			}
-
-			ExportMapArchive << ExportMapEntry;
-		}
-		TargetFile.ExportMapSize = ExportMapArchive.Tell();
-
-		// Temporary archive for serializing export bundle data
-		FBufferWriter ExportBundlesArchive(nullptr, 0, EBufferWriterFlags::AllowResize | EBufferWriterFlags::TakeOwnership);
-		uint32 ExportBundleEntryIndex = 0;
-		for (FExportBundle& ExportBundle : Package->ExportBundles)
-		{
-			const uint32 EntryCount = ExportBundle.Nodes.Num();
-			FExportBundleHeader ExportBundleHeader { ExportBundleEntryIndex, EntryCount };
-			ExportBundlesArchive << ExportBundleHeader ;
-
-			ExportBundleEntryIndex += EntryCount; 
-		}
-		for (FExportBundle& ExportBundle : Package->ExportBundles)
-		{
-			for (FExportGraphNode* ExportNode : ExportBundle.Nodes)
-			{
-				ExportBundlesArchive << ExportNode->BundleEntry;
-			}
-		}
-		TargetFile.ExportBundlesHeaderSize = ExportBundlesArchive.Tell();
-
-		FMappedName MappedPackageName = TargetFile.NameMapBuilder->MapName(Package->Name);
-		FMappedName MappedPackageSourceName = TargetFile.NameMapBuilder->MapName(Package->SourcePackageName);
-
-		TArray<uint8> NamesBuffer;
-		TArray<uint8> NameHashesBuffer;
-		SaveNameBatch(Package->LocalNameMapBuilder.GetNameMap(), NamesBuffer, NameHashesBuffer);
-		TargetFile.NameMapSize = Align(NamesBuffer.Num(), 8) + NameHashesBuffer.Num();
-
-		TargetFile.HeaderSerialSize =
-			sizeof(FPackageSummary)
-			+ TargetFile.NameMapSize
-			+ TargetFile.ImportMapSize
-			+ TargetFile.ExportMapSize
-			+ TargetFile.ExportBundlesHeaderSize
-			+ TargetFile.UGraphSize;
-
-		TargetFile.PackageHeaderData.AddZeroed(TargetFile.HeaderSerialSize);
-		uint8* PackageHeaderBuffer = TargetFile.PackageHeaderData.GetData();
-		FPackageSummary* PackageSummary = reinterpret_cast<FPackageSummary*>(PackageHeaderBuffer);
-
-		PackageSummary->Name = MappedPackageName;
-		PackageSummary->SourceName = MappedPackageSourceName;
-		PackageSummary->PackageFlags = Package->PackageFlags;
-		PackageSummary->CookedHeaderSize = Package->CookedHeaderSize;
-		FBufferWriter SummaryArchive(PackageHeaderBuffer, TargetFile.HeaderSerialSize);
-		SummaryArchive.Seek(sizeof(FPackageSummary));
-
-		// NameMap data
-		{
-			PackageSummary->NameMapNamesOffset = SummaryArchive.Tell();
-			check(PackageSummary->NameMapNamesOffset % 8 == 0);
-			PackageSummary->NameMapNamesSize = NamesBuffer.Num();
-			SummaryArchive.Serialize(NamesBuffer.GetData(), NamesBuffer.Num());
-			PackageSummary->NameMapHashesOffset = Align(SummaryArchive.Tell(), 8);
-			int32 PaddingByteCount = PackageSummary->NameMapHashesOffset - SummaryArchive.Tell();
-			if (PaddingByteCount)
-			{
-				check(PaddingByteCount < 8);
-				uint8 PaddingBytes[8]{ 0 };
-				SummaryArchive.Serialize(PaddingBytes, PaddingByteCount);
-			}
-			PackageSummary->NameMapHashesSize = NameHashesBuffer.Num();
-			SummaryArchive.Serialize(NameHashesBuffer.GetData(), NameHashesBuffer.Num());
-		}
-
-		// ImportMap data
-		{
-			check(ImportMapArchive.Tell() == TargetFile.ImportMapSize);
-			PackageSummary->ImportMapOffset = SummaryArchive.Tell();
-			SummaryArchive.Serialize(ImportMapArchive.GetWriterData(), ImportMapArchive.Tell());
-		}
-
-		// ExportMap data
-		{
-			check(ExportMapArchive.Tell() == TargetFile.ExportMapSize);
-			PackageSummary->ExportMapOffset = SummaryArchive.Tell();
-			SummaryArchive.Serialize(ExportMapArchive.GetWriterData(), ExportMapArchive.Tell());
-		}
-
-		// ExportBundle data
-		{
-			check(ExportBundlesArchive.Tell() == TargetFile.ExportBundlesHeaderSize);
-			PackageSummary->ExportBundlesOffset = SummaryArchive.Tell();
-			SummaryArchive.Serialize(ExportBundlesArchive.GetWriterData(), ExportBundlesArchive.Tell());
-		}
-
-		// Graph data
-		{
-			check(GraphArchive.Tell() == TargetFile.UGraphSize);
-			PackageSummary->GraphDataOffset = SummaryArchive.Tell();
-			PackageSummary->GraphDataSize = TargetFile.UGraphSize;
-			SummaryArchive.Serialize(GraphArchive.GetWriterData(), GraphArchive.Tell());
-		}
-	}
-}
-
-void FinalizePackageStoreContainerHeader(FContainerTargetSpec& ContainerTarget)
-{
-	IOSTORE_CPU_SCOPE(FinalizePackageStoreContainerHeader);
-	check(ContainerTarget.NameMapBuilder);
-	FNameMapBuilder& NameMapBuilder = *ContainerTarget.NameMapBuilder;
-	FCulturePackageMap& CulturePackageMap = ContainerTarget.Header.CulturePackageMap;
-	TArray<FPackageId>& PackageIds = ContainerTarget.Header.PackageIds;
-	TArray<TTuple<FPackageId, FPackageId>>& PackageRedirects = ContainerTarget.Header.PackageRedirects;
-	
-	int32 StoreTocSize = ContainerTarget.PackageCount * sizeof(FPackageStoreEntry);
-	FLargeMemoryWriter StoreTocArchive(0, true);
-	FLargeMemoryWriter StoreDataArchive(0, true);
-
-	auto SerializePackageEntryCArrayHeader = [&StoreTocSize,&StoreTocArchive,&StoreDataArchive](int32 Count)
-	{
-		const int32 RemainingTocSize = StoreTocSize - StoreTocArchive.Tell();
-		const int32 OffsetFromThis = RemainingTocSize + StoreDataArchive.Tell();
-		uint32 ArrayNum = Count > 0 ? Count : 0; 
-		uint32 OffsetToDataFromThis = ArrayNum > 0 ? OffsetFromThis : 0;
-
-		StoreTocArchive << ArrayNum;
-		StoreTocArchive << OffsetToDataFromThis;
-	};
-
-	PackageIds.Reserve(ContainerTarget.PackageCount);
-	TArray<FContainerTargetFile*> SortedTargetFiles;
-	SortedTargetFiles.Reserve(ContainerTarget.TargetFiles.Num());
-	for (FContainerTargetFile& TargetFile : ContainerTarget.TargetFiles)
-	{
-		if (!TargetFile.bIsBulkData)
-		{
-			SortedTargetFiles.Add(&TargetFile);
-		}
-	}
-	Algo::Sort(SortedTargetFiles, [](const FContainerTargetFile* A, const FContainerTargetFile* B)
-	{
-		return A->Package->GlobalPackageId < B->Package->GlobalPackageId;
-	});
-
-	for (FContainerTargetFile* TargetFile : SortedTargetFiles)
-	{
-		FPackage* Package = TargetFile->Package;
-
-		// PackageIds
-		{
-			//check(!PackageIds.Contains(Package->GlobalPackageId));
-			PackageIds.Add(Package->GlobalPackageId);
-		}
-
-		// CulturePackageMap
-		if (Package->bIsLocalizedAndConformed)
-		{
-			CulturePackageMap.FindOrAdd(Package->Region).Emplace(Package->SourceGlobalPackageId, Package->GlobalPackageId);
-		}
-
-		// Redirects
-		if (Package->RedirectedPackageId.IsValid())
-		{
-			PackageRedirects.Add(MakeTuple(Package->RedirectedPackageId, Package->GlobalPackageId));
-		}
-
-		// StoreEntries
-		{
-			uint64 ExportBundlesSize = TargetFile->HeaderSerialSize + Package->ExportsSerialSize;
-			int32 ExportBundleCount = Package->ExportBundles.Num();
-			uint32 LoadOrder = Package->ExportBundles.Num() > 0 ? Package->ExportBundles[0].LoadOrder : 0;
-			uint32 Pad = 0;
-
-			StoreTocArchive << ExportBundlesSize;
-			StoreTocArchive << Package->ExportCount;
-			StoreTocArchive << ExportBundleCount;
-			StoreTocArchive << LoadOrder;
-			StoreTocArchive << Pad;
-
-			// ImportedPackages
-			{
-				SerializePackageEntryCArrayHeader(Package->ImportedPackageIds.Num());
-				for (FPackageId PackageId : Package->ImportedPackageIds)
-				{
-					check(PackageId.IsValid());
-					StoreDataArchive << PackageId;
-				}
-			}
-		}
-	}
-
-	const int32 StoreByteCount = StoreTocArchive.TotalSize() + StoreDataArchive.TotalSize();
-	ContainerTarget.Header.PackageCount = ContainerTarget.PackageCount;
-	ContainerTarget.Header.StoreEntries.AddUninitialized(StoreByteCount);
-	FBufferWriter PackageStoreArchive(ContainerTarget.Header.StoreEntries.GetData(), StoreByteCount);
-	PackageStoreArchive.Serialize(StoreTocArchive.GetData(), StoreTocArchive.TotalSize());
-	PackageStoreArchive.Serialize(StoreDataArchive.GetData(), StoreDataArchive.TotalSize());
-}
-
-static void FinalizeInitialLoadMeta(
-	FNameMapBuilder& GlobalNameMapBuilder,
-	const FGlobalScriptObjects& GlobalScriptImports,
-	FArchive& InitialLoadArchive)
-{
-	IOSTORE_CPU_SCOPE(FinalizeInitialLoad);
-	UE_LOG(LogIoStore, Display, TEXT("Finalizing initial load..."));
-
-	int32 NumScriptObjects = GlobalScriptImports.Num();
-	InitialLoadArchive << NumScriptObjects; 
-
-	TArray<FScriptObjectData> ScriptObjects;
-	GlobalScriptImports.GenerateValueArray(ScriptObjects );
-	Algo::Sort(ScriptObjects, [](const FScriptObjectData& A, const FScriptObjectData& B)
-	{
-		return A.FullName < B.FullName;
-	});
-
-	for (const FScriptObjectData& ImportData : ScriptObjects)
-	{
-		GlobalNameMapBuilder.MarkNameAsReferenced(ImportData.ObjectName);
-		FScriptObjectEntry Entry;
-		Entry.ObjectName = GlobalNameMapBuilder.MapName(ImportData.ObjectName).ToUnresolvedMinimalName();
-		Entry.GlobalIndex = ImportData.GlobalIndex;
-		Entry.OuterIndex = ImportData.OuterIndex;
-		Entry.CDOClassIndex = ImportData.CDOClassIndex;
-
-		InitialLoadArchive << Entry; 
-	}
-};
-
-static FIoBuffer CreateExportBundleBuffer(const FContainerTargetFile& TargetFile, const TArray<FObjectExport>& ObjectExports, const FIoBuffer UExpBuffer, TArray<FFileRegion>* InOutFileRegions = nullptr)
-{
-	const FPackage* Package = TargetFile.Package;
-	check(TargetFile.PackageHeaderData.Num() > 0);
-	const uint64 BundleBufferSize = TargetFile.PackageHeaderData.Num() + TargetFile.Package->ExportsSerialSize;
-	FIoBuffer BundleBuffer(BundleBufferSize);
-	FMemory::Memcpy(BundleBuffer.Data(), TargetFile.PackageHeaderData.GetData(), TargetFile.PackageHeaderData.Num());
-	uint64 BundleBufferOffset = TargetFile.PackageHeaderData.Num();
-
-	TArray<FFileRegion> OutputRegions;
-
-	for (const FExportBundle& ExportBundle : TargetFile.Package->ExportBundles)
-	{
-		for (const FExportGraphNode* Node : ExportBundle.Nodes)
-		{
-			if (Node->BundleEntry.CommandType == FExportBundleEntry::ExportCommandType_Serialize)
-			{
-				const FObjectExport& ObjectExport = ObjectExports[Package->ExportIndexOffset + Node->BundleEntry.LocalExportIndex];
-				const uint64 Offset = uint64(ObjectExport.SerialOffset - Package->UAssetSize);
-				const uint64 End = uint64(Offset + ObjectExport.SerialSize);
-				check(End <= UExpBuffer.DataSize());
-				FMemory::Memcpy(BundleBuffer.Data() + BundleBufferOffset, UExpBuffer.Data() + Offset, ObjectExport.SerialSize);
-
-				if (InOutFileRegions)
-				{
-					// Find overlapping regions and adjust them to match the new offset of the export data
-					for (const FFileRegion& Region : *InOutFileRegions)
-					{
-						uint64 RegionStart = Region.Offset;
-						uint64 RegionEnd = RegionStart + Region.Length;
-
-						if (Offset <= RegionStart && RegionEnd <= End)
-						{
-							FFileRegion NewRegion = Region;
-							NewRegion.Offset -= Offset;
-							NewRegion.Offset += BundleBufferOffset;
-							OutputRegions.Add(NewRegion);
-						}
-					}
-				}
-
-				BundleBufferOffset += ObjectExport.SerialSize;
-			}
-		}
-	}
-	check(BundleBufferOffset == BundleBuffer.DataSize());
-
-	if (InOutFileRegions)
-	{
-		*InOutFileRegions = OutputRegions;
-	}
-
-	return BundleBuffer;
-}
-
-static void ParsePackageAssets(
-	TArray<FPackage*>& Packages,
-	FPackageAssetData& PackageAssetData)
-{
-	IOSTORE_CPU_SCOPE(ParsePackageAssets);
-	UE_LOG(LogIoStore, Display, TEXT("Parsing packages..."));
-
-	TAtomic<int32> ReadCount {0};
-	TAtomic<int32> ParseCount {0};
-	const int32 TotalPackageCount = Packages.Num();
-
-	TArray<FPackageFileSummary> PackageFileSummaries;
-	PackageFileSummaries.SetNum(TotalPackageCount);
-
-	uint8* UAssetMemory = nullptr;
-	TArray<uint8*> PackageAssetBuffers;
-	PackageAssetBuffers.SetNum(TotalPackageCount);
-
-	UE_LOG(LogIoStore, Display, TEXT("Reading package assets..."));
-	{
-		IOSTORE_CPU_SCOPE(ReadUAssetFiles);
-
-		uint64 TotalUAssetSize = 0;
-		for (const FPackage* Package : Packages)
-		{
-			TotalUAssetSize += Package->UAssetSize;
-		}
-		UAssetMemory = reinterpret_cast<uint8*>(FMemory::Malloc(TotalUAssetSize));
-		uint8* UAssetMemoryPtr = UAssetMemory;
-		for (int32 Index = 0; Index < TotalPackageCount; ++Index)
-		{
-			PackageAssetBuffers[Index] = UAssetMemoryPtr;
-			UAssetMemoryPtr += Packages[Index]->UAssetSize;
-		}
-
-		TAtomic<uint64> CurrentFileIndex{ 0 };
-		ParallelFor(TotalPackageCount, [&ReadCount, &PackageAssetBuffers, &Packages, &CurrentFileIndex](int32 Index)
-		{
-			TRACE_CPUPROFILER_EVENT_SCOPE(ReadUAssetFile);
-			FPackage* Package = Packages[Index];
-			uint8* Buffer = PackageAssetBuffers[Index];
-			IFileHandle* FileHandle = FPlatformFileManager::Get().GetPlatformFile().OpenRead(*Package->FileName);
-			if (FileHandle)
-			{
-				bool bSuccess = FileHandle->Read(Buffer, Package->UAssetSize);
-				UE_CLOG(!bSuccess, LogIoStore, Warning, TEXT("Failed reading file '%s'"), *Package->FileName);
-				delete FileHandle;
-			}
-			else
-			{
-				UE_LOG(LogIoStore, Warning, TEXT("Couldn't open file '%s'"), *Package->FileName);
-			}
-			uint64 LocalFileIndex = CurrentFileIndex.IncrementExchange() + 1;
-			UE_CLOG(LocalFileIndex % 1000 == 0, LogIoStore, Display, TEXT("Reading %d/%d: '%s'"), LocalFileIndex, Packages.Num(), *Package->FileName);
-		}, EParallelForFlags::Unbalanced);
-	}
-
-	{
-		IOSTORE_CPU_SCOPE(SerializeSummaries);
-
-		ParallelFor(TotalPackageCount, [
-				&ReadCount,
-				&PackageAssetBuffers,
-				&PackageFileSummaries,
-				&Packages](int32 Index)
-		{
-			uint8* PackageBuffer = PackageAssetBuffers[Index];
-			FPackageFileSummary& Summary = PackageFileSummaries[Index];
-			FPackage& Package = *Packages[Index];
-
-			TArrayView<const uint8> MemView(PackageBuffer, Package.UAssetSize);
-			if (!Package.UAssetSize)
-			{
-				return;
-			}
-
-			FMemoryReaderView Ar(MemView);
-			Ar << Summary;
-
-			Package.SummarySize = Ar.Tell();
-			Package.NameCount = Summary.NameCount;
-			Package.ImportCount = Summary.ImportCount;
-			Package.PreloadDependencyCount = Summary.PreloadDependencyCount;
-			Package.ExportCount = Summary.ExportCount;
-			Package.PackageFlags = Summary.PackageFlags;
-			Package.CookedHeaderSize = Summary.TotalHeaderSize;
-
-		}, EParallelForFlags::Unbalanced);
-	}
-
-	int32 TotalImportCount = 0;
-	int32 TotalPreloadDependencyCount = 0;
-	int32 TotalExportCount = 0;
-	for (FPackage* Package : Packages)
-	{
-		if (Package->ImportCount > 0)
-		{
-			Package->ImportIndexOffset = TotalImportCount;
-			TotalImportCount += Package->ImportCount;
-		}
-
-		if (Package->PreloadDependencyCount > 0)
-		{
-			Package->PreloadIndexOffset = TotalPreloadDependencyCount;
-			TotalPreloadDependencyCount += Package->PreloadDependencyCount;
-		}
-
-		if (Package->ExportCount > 0)
-		{
-			Package->ExportIndexOffset = TotalExportCount;
-			TotalExportCount += Package->ExportCount;
-		}
-	}
-	PackageAssetData.ObjectImports.AddUninitialized(TotalImportCount);
-	PackageAssetData.PreloadDependencies.AddUninitialized(TotalPreloadDependencyCount);
-	PackageAssetData.ObjectExports.AddUninitialized(TotalExportCount);
-
-	UE_LOG(LogIoStore, Display, TEXT("Parsing package assets..."));
-	{
-		IOSTORE_CPU_SCOPE(SerializeAssets);
-
-		for (int32 PackageIndex = 0; PackageIndex < TotalPackageCount; ++PackageIndex)
-		{
-			uint8* PackageBuffer = PackageAssetBuffers[PackageIndex];
-			const FPackageFileSummary& Summary = PackageFileSummaries[PackageIndex];
-			FPackage& Package = *Packages[PackageIndex];
-			TArrayView<const uint8> MemView(PackageBuffer, Package.UAssetSize);
-			FMemoryReaderView Ar(MemView);
-
-			if (Summary.NameCount > 0)
-			{
-				Ar.Seek(Summary.NameOffset);
-
-				FNameEntrySerialized NameEntry(ENAME_LinkerConstructor);
-
-				Package.SummaryNames.Reserve(Summary.NameCount);
-				for (int32 I = 0; I < Summary.NameCount; ++I)
-				{
-					Ar << NameEntry;
-					FName Name(NameEntry);
-					Package.SummaryNames.Add(Name);
-					Package.LocalNameMapBuilder.AddName(Name);
-				}
-			}
-		}
-
-		ParallelFor(TotalPackageCount,[
-				&ParseCount,
-				&PackageAssetBuffers,
-				&PackageFileSummaries,
-				&Packages,
-				&PackageAssetData](int32 Index)
-		{
-			uint8* PackageBuffer = PackageAssetBuffers[Index];
-			const FPackageFileSummary& Summary = PackageFileSummaries[Index];
-			FPackage& Package = *Packages[Index];
-			TArrayView<const uint8> MemView(PackageBuffer, Package.UAssetSize);
-			FMemoryReaderView Ar(MemView);
-			Ar.SetFilterEditorOnly((Package.PackageFlags & EPackageFlags::PKG_FilterEditorOnly) != 0);
-
-			IOSTORE_CPU_SCOPE_DATA(ParsePackage, TCHAR_TO_ANSI(*Package.FileName));
-
-			const int32 Count = ParseCount.IncrementExchange();
-			UE_CLOG(Count % 1000 == 0, LogIoStore, Display, TEXT("Parsing %d/%d: '%s'"), Count, Packages.Num(), *Package.FileName);
-
-			if (Summary.ImportCount > 0)
-			{
-				FNameReaderProxyArchive ProxyAr(Ar, Package.LocalNameMapBuilder.GetNameMap());
-				ProxyAr.Seek(Summary.ImportOffset);
-
-				for (int32 I = 0; I < Summary.ImportCount; ++I)
-				{
-					FObjectImport& ObjectImport = PackageAssetData.ObjectImports[Package.ImportIndexOffset + I];
-					ProxyAr << ObjectImport;
-				}
-			}
-
-			if (Summary.PreloadDependencyCount > 0)
-			{
-				Ar.Seek(Summary.PreloadDependencyOffset);
-				Ar.Serialize(PackageAssetData.PreloadDependencies.GetData() + Package.PreloadIndexOffset, Summary.PreloadDependencyCount * sizeof(FPackageIndex));
-			}
-
-			if (Summary.ExportCount > 0)
-			{
-				FNameReaderProxyArchive ProxyAr(Ar, Package.LocalNameMapBuilder.GetNameMap());
-				ProxyAr.Seek(Summary.ExportOffset);
-
-				for (int32 I = 0; I < Summary.ExportCount; ++I)
-				{
-					FObjectExport& ObjectExport = PackageAssetData.ObjectExports[Package.ExportIndexOffset + I];
-					ProxyAr << ObjectExport;
-					Package.ExportsSerialSize += ObjectExport.SerialSize;
-				}
-			}
-		}, EParallelForFlags::Unbalanced);
-	}
-
-	FMemory::Free(UAssetMemory);
-}
-
-// modified copy from SavePackage
-EObjectMark GetExcludedObjectMarksForTargetPlatform(const ITargetPlatform* TargetPlatform)
-{
-	EObjectMark Marks = OBJECTMARK_NOMARKS;
-	if (!TargetPlatform->HasEditorOnlyData())
-	{
-		Marks = (EObjectMark)(Marks | OBJECTMARK_EditorOnly);
-	}
-	if (TargetPlatform->IsServerOnly())
-	{
-		Marks = (EObjectMark)(Marks | OBJECTMARK_NotForServer);
-	}
-	if (TargetPlatform->IsClientOnly())
-	{
-		Marks = (EObjectMark)(Marks | OBJECTMARK_NotForClient);
-	}
-	return Marks;
-}
-
-// modified copy from SavePackage
-EObjectMark GetExcludedObjectMarksForObject(const UObject* Object, const ITargetPlatform* TargetPlatform)
-{
-	EObjectMark Marks = OBJECTMARK_NOMARKS;
-	if (!Object->NeedsLoadForClient())
-	{
-		Marks = (EObjectMark)(Marks | OBJECTMARK_NotForClient);
-	}
-	if (!Object->NeedsLoadForServer())
-	{
-		Marks = (EObjectMark)(Marks | OBJECTMARK_NotForServer);
-	}
-	if (!Object->NeedsLoadForTargetPlatform(TargetPlatform))
-	{
-		Marks = (EObjectMark)(Marks | OBJECTMARK_NotForClient | OBJECTMARK_NotForServer);
-	}
-	if (Object->IsEditorOnly())
-	{
-		Marks = (EObjectMark)(Marks | OBJECTMARK_EditorOnly);
-	}
-	if ((Marks & OBJECTMARK_NotForClient) && (Marks & OBJECTMARK_NotForServer))
-	{
-		Marks = (EObjectMark)(Marks | OBJECTMARK_EditorOnly);
-	}
-	return Marks;
-}
-
-static void FindScriptObjectsRecursive(
-	FGlobalPackageData& GlobalPackageData,
-	FPackageObjectIndex OuterIndex,
-	UObject* Object,
-	const ITargetPlatform* TargetPlatform,
-	const EObjectMark ExcludedObjectMarks)
-{
-	if (!Object->HasAllFlags(RF_Public))
-	{
-		UE_LOG(LogIoStore, Log, TEXT("Skipping script object: %s (!RF_Public)"), *Object->GetFullName());
-		return;
-	}
-
-	const UObject* ObjectForExclusion = Object->HasAnyFlags(RF_ClassDefaultObject) ? (const UObject*)Object->GetClass() : Object;
-	const EObjectMark ObjectMarks = GetExcludedObjectMarksForObject(ObjectForExclusion, TargetPlatform);
-
-	if (ObjectMarks & ExcludedObjectMarks)
-	{
-		UE_LOG(LogIoStore, Log, TEXT("Skipping script object: %s (Excluded for target platform)"), *Object->GetFullName());
-		return;
-	}
-
-	FGlobalScriptObjects& ScriptObjects = GlobalPackageData.ScriptObjects;
-	const FScriptObjectData* Outer = ScriptObjects.Find(OuterIndex);
-	check(Outer);
-
-	FName ObjectName = Object->GetFName();
-
-	FString TempFullName = ScriptObjects.FindRef(OuterIndex).FullName;
-	TempFullName.AppendChar(TEXT('/'));
-	ObjectName.AppendString(TempFullName);
-
-	TempFullName.ToLowerInline();
-	FPackageObjectIndex GlobalImportIndex = FPackageObjectIndex::FromScriptPath(TempFullName);
-
-	check(!GlobalPackageData.ImportsByFullName.Contains(TempFullName));
-	FScriptObjectData* ScriptImport = ScriptObjects.Find(GlobalImportIndex);
-	if (ScriptImport)
-	{
-		UE_LOG(LogIoStore, Fatal, TEXT("Import name hash collision \"%s\" and \"%s"), *TempFullName, *ScriptImport->FullName);
-	}
-
-	FPackageObjectIndex CDOClassIndex = Outer->CDOClassIndex;
-	if (CDOClassIndex.IsNull())
-	{
-		TCHAR NameBuffer[FName::StringBufferSize];
-		uint32 Len = ObjectName.ToString(NameBuffer);
-		if (FCString::Strncmp(NameBuffer, TEXT("Default__"), 9) == 0)
-		{
-			FString CDOClassFullName = Outer->FullName;
-			CDOClassFullName.AppendChar(TEXT('/'));
-			CDOClassFullName.AppendChars(NameBuffer + 9, Len - 9);
-			CDOClassFullName.ToLowerInline();
-
-			CDOClassIndex = GlobalPackageData.ImportsByFullName.FindRef(CDOClassFullName);
-			check(CDOClassIndex.IsScriptImport());
-		}
-	}
-
-	GlobalPackageData.ImportsByFullName.Add(TempFullName, GlobalImportIndex);
-	ScriptImport = &ScriptObjects.Add(GlobalImportIndex);
-	ScriptImport->GlobalIndex = GlobalImportIndex;
-	ScriptImport->FullName = MoveTemp(TempFullName);
-	ScriptImport->OuterIndex = Outer->GlobalIndex;
-	ScriptImport->ObjectName = ObjectName;
-	ScriptImport->CDOClassIndex = CDOClassIndex;
-
-	TArray<UObject*> InnerObjects;
-	GetObjectsWithOuter(Object, InnerObjects, /*bIncludeNestedObjects*/false);
-	for (UObject* InnerObject : InnerObjects)
-	{
-		FindScriptObjectsRecursive(GlobalPackageData, GlobalImportIndex, InnerObject, TargetPlatform, ExcludedObjectMarks);
+	case FPackageStoreBulkDataManifest::EBulkdataType::Normal:
+		return EIoChunkType::BulkData;
+	case FPackageStoreBulkDataManifest::EBulkdataType::Optional:
+		return EIoChunkType::OptionalBulkData;
+	case FPackageStoreBulkDataManifest::EBulkdataType::MemoryMapped:
+		return EIoChunkType::MemoryMappedBulkData;
+	default:
+		UE_LOG(LogIoStore, Error, TEXT("Invalid EBulkdataType (%d) found!"), Type);
+		return EIoChunkType::Invalid;
 	}
-};
+}
 
-static void CreateGlobalScriptObjects(
-	FGlobalPackageData& GlobalPackageData,
-	const ITargetPlatform* TargetPlatform)
+FContainerTargetSpec* AddContainer(
+	FName Name,
+	TArray<FContainerTargetSpec*>& Containers)
 {
-	IOSTORE_CPU_SCOPE(CreateGlobalScriptObjects);
-	UE_LOG(LogIoStore, Display, TEXT("Creating global script objects..."));
-
-	const EObjectMark ExcludedObjectMarks = GetExcludedObjectMarksForTargetPlatform(TargetPlatform);
-
-	TArray<UPackage*> ScriptPackages;
-	FindAllRuntimeScriptPackages(ScriptPackages);
-
-	TArray<UObject*> InnerObjects;
-	for (UPackage* Package : ScriptPackages)
+	FIoContainerId ContainerId = FIoContainerId::FromName(Name);
+	for (FContainerTargetSpec* ExistingContainer : Containers)
 	{
-		FGlobalScriptObjects& ScriptObjects = GlobalPackageData.ScriptObjects;
-
-		FName ObjectName = Package->GetFName();
-		FString FullName = Package->GetName();
-
-		FullName.ToLowerInline();
-		FPackageObjectIndex GlobalImportIndex = FPackageObjectIndex::FromScriptPath(FullName);
-
-		check(!GlobalPackageData.ImportsByFullName.Contains(FullName));
-		FScriptObjectData* ScriptImport = ScriptObjects.Find(GlobalImportIndex);
-		if (ScriptImport)
+		if (ExistingContainer->Name == Name)
 		{
-			UE_LOG(LogIoStore, Fatal, TEXT("Import name hash collision \"%s\" and \"%s"), *FullName, *ScriptImport->FullName);
+			UE_LOG(LogIoStore, Fatal, TEXT("Duplicate container name: '%s'"), *Name.ToString());
+			return nullptr;
 		}
-
-		GlobalPackageData.ImportsByFullName.Add(FullName, GlobalImportIndex);
-		ScriptImport = &ScriptObjects.Add(GlobalImportIndex);
-		ScriptImport->GlobalIndex = GlobalImportIndex;
-		ScriptImport->FullName = FullName;
-		ScriptImport->OuterIndex = FPackageObjectIndex();
-		ScriptImport->ObjectName = ObjectName;
-
-		InnerObjects.Reset();
-		GetObjectsWithOuter(Package, InnerObjects, /*bIncludeNestedObjects*/false);
-		for (UObject* InnerObject : InnerObjects)
+		if (ExistingContainer->ContainerId == ContainerId)
 		{
-			FindScriptObjectsRecursive(GlobalPackageData, GlobalImportIndex, InnerObject, TargetPlatform, ExcludedObjectMarks);
+			UE_LOG(LogIoStore, Fatal, TEXT("Hash collision for container names: '%s' and '%s'"), *Name.ToString(), *ExistingContainer->Name.ToString());
+			return nullptr;
 		}
 	}
+	
+	FContainerTargetSpec* ContainerTargetSpec = new FContainerTargetSpec();
+	ContainerTargetSpec->Name = Name;
+	ContainerTargetSpec->ContainerId = ContainerId;
+	Containers.Add(ContainerTargetSpec);
+	return ContainerTargetSpec;
 }
 
-static void CreateGlobalImportsAndExports(
+FLegacyCookedPackage* FindOrAddPackage(
 	const FIoStoreArguments& Arguments,
-	TArray<FPackage*>& Packages,
-	const FPackageIdMap& PackageIdMap,
-	FPackageAssetData& PackageAssetData,
-	FGlobalPackageData& GlobalPackageData,
-	FExportGraph& ExportGraph,
-	TMap<FPackageId, FPackageId>& OutSourceToRedirectedPackageMap)
+	const TCHAR* RelativeFileName,
+	TArray<FLegacyCookedPackage*>& Packages,
+	FPackageNameMap& PackageNameMap,
+	FPackageIdMap& PackageIdMap,
+	FPackageStoreOptimizer& PackageStoreOptimizer)
 {
-	IOSTORE_CPU_SCOPE(CreateGlobalImportsAndExports);
-	UE_LOG(LogIoStore, Display, TEXT("Creating global imports and exports..."));
-
-	TArray<FString> TempFullNames;
-	TMap<FName, FName> Redirects;
-	const FString DLCPrefix = Arguments.IsDLC() ? FString::Printf(TEXT("/%s/"), *FPaths::GetBaseFilename(Arguments.DLCPluginPath)) : FString();
-
-	TSet<FPackage*> TempImportedPackages;
-	for (FPackage* Package : Packages)
+	FString PackageName;
+	FString ErrorMessage;
+	if (!FPackageName::TryConvertFilenameToLongPackageName(RelativeFileName, PackageName, &ErrorMessage))
 	{
-		if (Package->ExportCount == 0)
-		{
-			continue;
-		}
-
-		if (Package->RedirectedPackageId.IsValid())
-		{
-			Redirects.Add(Package->Name, Package->SourcePackageName);
-			OutSourceToRedirectedPackageMap.Add(Package->RedirectedPackageId, Package->GlobalPackageId);
-		}
-
-		TempFullNames.Reset();
-		TempFullNames.SetNum(Package->ExportCount, false);
-		FObjectExport* ExportMap = PackageAssetData.ObjectExports.GetData() + Package->ExportIndexOffset;
-		for (int32 ExportIndex = 0; ExportIndex < Package->ExportCount; ExportIndex++)
-		{
-			int32 GlobalExportIndex = FindExport(
-				GlobalPackageData,
-				TempFullNames,
-				PackageAssetData.ObjectExports.GetData() + Package->ExportIndexOffset,
-				ExportIndex,
-				Package);
-
-			FExportObjectData& ExportData = GlobalPackageData.ExportObjects[GlobalExportIndex];
-			ExportData.CreateNode = ExportGraph.AddNode(Package, { uint32(ExportIndex), FExportBundleEntry::ExportCommandType_Create });
-			ExportData.SerializeNode = ExportGraph.AddNode(Package, { uint32(ExportIndex), FExportBundleEntry::ExportCommandType_Serialize });
-
-			Package->Exports.Add(GlobalExportIndex);
-			Package->CreateExportNodes.Add(ExportData.CreateNode);
-			Package->SerializeExportNodes.Add(ExportData.SerializeNode);
-			ExportGraph.AddInternalDependency(ExportData.CreateNode, ExportData.SerializeNode);
-		}
+		UE_LOG(LogIoStore, Warning, TEXT("Failed to obtain package name from file name '%s'"), *ErrorMessage);
+		return nullptr;
 	}
 
-	for (FExportObjectData& Export : GlobalPackageData.ExportObjects)
-	{
-		if (Export.IsPublicExport())
-		{
-			TMap<FPackageObjectIndex, int32>& PublicExports = GlobalPackageData.PublicExportIndices;
-			FPackageObjectIndex GlobalImportIndex = FPackageObjectIndex::FromPackagePath(Export.FullName);
-
-			check(!GlobalPackageData.ImportsByFullName.Contains(Export.FullName));
-			int32* ExportIndex = PublicExports.Find(GlobalImportIndex);
-			if (ExportIndex)
-			{
-				UE_LOG(LogIoStore, Fatal, TEXT("Import name hash collision \"%s\" and \"%s"),
-					*Export.FullName,
-					*GlobalPackageData.ExportObjects[*ExportIndex].FullName);
-			}
-			GlobalPackageData.ImportsByFullName.Add(Export.FullName, GlobalImportIndex);
-			PublicExports.Add(GlobalImportIndex, Export.GlobalIndex);
-			Export.GlobalImportIndex = GlobalImportIndex;
-		}
-	}
+	FName PackageFName = *PackageName;
 
-	for (FPackage* Package : Packages)
+	FLegacyCookedPackage* Package = PackageNameMap.FindRef(PackageFName);
+	if (!Package)
 	{
-		if (Package->ImportCount == 0)
+		FPackageId PackageId = FPackageId::FromName(PackageFName);
+		if (FLegacyCookedPackage* FindById = PackageIdMap.FindRef(PackageId))
 		{
-			continue;
+			UE_LOG(LogIoStore, Fatal, TEXT("Package name hash collision \"%s\" and \"%s"), *FindById->PackageName.ToString(), *PackageFName.ToString());
 		}
 
-		FObjectImport* ImportMap = PackageAssetData.ObjectImports.GetData() + Package->ImportIndexOffset;
-		TempFullNames.Reset();
-		TempFullNames.SetNum(Package->ImportCount, false);
-		Package->Imports.Reserve(Package->ImportCount);
-		Package->ImportedPackages.Reserve(Package->ImportCount / 2);
-		TempImportedPackages.Reset();
-
-		for (int32 ImportIndex = 0; ImportIndex < Package->ImportCount; ++ImportIndex)
+		if (const FName* ReleasedPackageName = Arguments.ReleasedPackages.PackageIdToName.Find(PackageId))
 		{
-			FindImportFullName(Redirects, TempFullNames, ImportMap, ImportIndex);
-			FString& FullName = TempFullNames[ImportIndex];
-			FObjectImport& Import = ImportMap[ImportIndex];
-			const bool bIsPackage = Import.OuterIndex.IsNull();
-
-			FPackageObjectIndex GlobalImportIndex = FindAndVerifyGlobalImport(
-				Package,
-				GlobalPackageData,
-				Import,
-				FullName,
-				DLCPrefix);
-
-			// When building DLC:s and we don't have all packages available,
-			// then a global package import object can be missing and still be valid
-			if (GlobalImportIndex.IsNull() && !bIsPackage && !FullName.StartsWith(ScriptPrefix))
-			{
-				GlobalImportIndex = FPackageObjectIndex::FromPackagePath(FullName);
-			}
-
-			Package->Imports.Add(GlobalImportIndex);
-
-			if (bIsPackage && GlobalImportIndex.IsNull())
-			{
-				FPackageId PackageId = FPackageId::FromName(Import.ObjectName);
-				Package->ImportedPackageIds.Add(PackageId);
-				if (FPackage* ImportedPackage = PackageIdMap.FindRef(PackageId))
-				{
-					Package->ImportedPackages.Add(ImportedPackage);
-				}
-			}
+			UE_LOG(LogIoStore, Fatal, TEXT("Package name hash collision \"%s\" and \"%s"), *ReleasedPackageName->ToString(), *PackageFName.ToString());
 		}
+
+		Package = new FLegacyCookedPackage();
+		Package->PackageName = PackageFName;
+		Package->GlobalPackageId = PackageId;
+		Packages.Add(Package);
+		PackageNameMap.Add(PackageFName, Package);
+		PackageIdMap.Add(PackageId, Package);
 	}
+
+	return Package;
 }
 
-static void MapExportEntryIndices(
-	TArray<FObjectExport>& ObjectExports,
-	TArray<FExportObjectData>& GlobalExports,
-	TArray<FPackage*>& Packages)
+static void ParsePackageAssets(TArray<FLegacyCookedPackage*>& Packages, const FPackageStoreOptimizer& PackageStoreOptimizer)
 {
-	IOSTORE_CPU_SCOPE(ExportData);
-	UE_LOG(LogIoStore, Display, TEXT("Converting export map import indices..."));
-
-	auto PackageObjectIdFromPackageIndex =
-		[](const TArray<FPackageObjectIndex>& Imports, const FPackageIndex& PackageIndex) -> FPackageObjectIndex 
-		{ 
-			if (PackageIndex.IsImport())
-			{
-				return Imports[PackageIndex.ToImport()];
-			}
-			if (PackageIndex.IsExport())
-			{
-				return FPackageObjectIndex::FromExportIndex(PackageIndex.ToExport());
-			}
-			return FPackageObjectIndex();
-		};
+	IOSTORE_CPU_SCOPE(ParsePackageAssets);
+	UE_LOG(LogIoStore, Display, TEXT("Parsing packages..."));
 
-	for (FPackage* Package : Packages)
-	{
-		for (int32 I = 0; I < Package->ExportCount; ++I)
-		{
-			const FObjectExport& ObjectExport = ObjectExports[Package->ExportIndexOffset + I];
-			FExportObjectData& ExportData = GlobalExports[Package->Exports[I]];
-			ExportData.OuterIndex = PackageObjectIdFromPackageIndex(Package->Imports, ObjectExport.OuterIndex);
-			ExportData.ClassIndex = PackageObjectIdFromPackageIndex(Package->Imports, ObjectExport.ClassIndex);
-			ExportData.SuperIndex = PackageObjectIdFromPackageIndex(Package->Imports, ObjectExport.SuperIndex);
-			ExportData.TemplateIndex = PackageObjectIdFromPackageIndex(Package->Imports, ObjectExport.TemplateIndex);
-		}
-	}
-};
+	TAtomic<int32> ReadCount {0};
+	TAtomic<int32> ParseCount {0};
+	const int32 TotalPackageCount = Packages.Num();
 
-static void ProcessLocalizedPackages(
-	const TArray<FPackage*>& Packages,
-	const FPackageNameMap& PackageMap,
-	FGlobalPackageData& GlobalPackageData,
-	FSourceToLocalizedPackageMultimap& OutSourceToLocalizedPackageMap)
-{
-	IOSTORE_CPU_SCOPE(ProcessLocalizedPackages);
+	TArray<FPackageFileSummary> PackageFileSummaries;
+	PackageFileSummaries.SetNum(TotalPackageCount);
 
-	FLocalizedToSourceImportIndexMap LocalizedToSourceImportIndexMap;
+	uint8* UAssetMemory = nullptr;
+	TArray<uint8*> PackageAssetBuffers;
+	PackageAssetBuffers.SetNum(TotalPackageCount);
 
-	UE_LOG(LogIoStore, Display, TEXT("Conforming localized packages..."));
-	for (FPackage* Package : Packages)
+	UE_LOG(LogIoStore, Display, TEXT("Reading package assets..."));
 	{
-		if (Package->Region.Len() == 0)
-		{
-			continue;
-		}
+		IOSTORE_CPU_SCOPE(ReadUAssetFiles);
 
-		check(!Package->RedirectedPackageId.IsValid());
-		if (Package->Name == Package->SourcePackageName)
+		uint64 TotalUAssetSize = 0;
+		for (const FLegacyCookedPackage* Package : Packages)
 		{
-			UE_LOG(LogIoStore, Error,
-				TEXT("For culture '%s': Localized package '%s' (0x%llX) should have a package name different from source name."),
-				*Package->Region,
-				*Package->Name.ToString(),
-				Package->GlobalPackageId.ValueForDebugging())
-			continue;
+			TotalUAssetSize += Package->UAssetSize;
 		}
-
-		FPackage* SourcePackage = PackageMap.FindRef(Package->SourcePackageName);
-		if (!SourcePackage)
+		UAssetMemory = reinterpret_cast<uint8*>(FMemory::Malloc(TotalUAssetSize));
+		uint8* UAssetMemoryPtr = UAssetMemory;
+		for (int32 Index = 0; Index < TotalPackageCount; ++Index)
 		{
-			// no update or verification required
-			UE_LOG(LogIoStore, Verbose,
-				TEXT("For culture '%s': Localized package '%s' (0x%llX) is unique and does not override a source package."),
-				*Package->Region,
-				*Package->Name.ToString(),
-				Package->GlobalPackageId.ValueForDebugging());
-			continue;
+			PackageAssetBuffers[Index] = UAssetMemoryPtr;
+			UAssetMemoryPtr += Packages[Index]->UAssetSize;
 		}
 
-		Package->SourceGlobalPackageId = SourcePackage->GlobalPackageId;
-
-		Package->bIsLocalizedAndConformed = ConformLocalizedPackage(
-			PackageMap, GlobalPackageData, *SourcePackage,
-			*Package, LocalizedToSourceImportIndexMap);
-
-		if (Package->bIsLocalizedAndConformed)
-		{
-			UE_LOG(LogIoStore, Verbose, TEXT("For culture '%s': Adding conformed localized package '%s' (0x%llX) for '%s' (0x%llX). ")
-				TEXT("When loading the source package, it will be remapped to this localized package."),
-				*Package->Region,
-				*Package->Name.ToString(),
-				Package->GlobalPackageId.ValueForDebugging(),
-				*Package->SourcePackageName.ToString(),
-				SourcePackage->GlobalPackageId.ValueForDebugging());
-
-			OutSourceToLocalizedPackageMap.Add(SourcePackage, Package);
-		}
-		else
+		TAtomic<uint64> CurrentFileIndex{ 0 };
+		ParallelFor(TotalPackageCount, [&ReadCount, &PackageAssetBuffers, &Packages, &CurrentFileIndex](int32 Index)
 		{
-			UE_LOG(LogIoStore, Display,
-				TEXT("For culture '%s': Localized package '%s' (0x%llX) does not conform to source package '%s' (0x%llX) due to mismatching public exports. ")
-				TEXT("When loading the source package, it will never be remapped to this localized package."),
-				*Package->Region,
-				*Package->Name.ToString(),
-				Package->GlobalPackageId.ValueForDebugging(),
-				*Package->SourcePackageName.ToString(),
-				SourcePackage->GlobalPackageId.ValueForDebugging());
-		}
+			TRACE_CPUPROFILER_EVENT_SCOPE(ReadUAssetFile);
+			FLegacyCookedPackage* Package = Packages[Index];
+			uint8* Buffer = PackageAssetBuffers[Index];
+			IFileHandle* FileHandle = FPlatformFileManager::Get().GetPlatformFile().OpenRead(*Package->FileName);
+			if (FileHandle)
+			{
+				bool bSuccess = FileHandle->Read(Buffer, Package->UAssetSize);
+				UE_CLOG(!bSuccess, LogIoStore, Warning, TEXT("Failed reading file '%s'"), *Package->FileName);
+				delete FileHandle;
+			}
+			else
+			{
+				UE_LOG(LogIoStore, Warning, TEXT("Couldn't open file '%s'"), *Package->FileName);
+			}
+			uint64 LocalFileIndex = CurrentFileIndex.IncrementExchange() + 1;
+			UE_CLOG(LocalFileIndex % 1000 == 0, LogIoStore, Display, TEXT("Reading %d/%d: '%s'"), LocalFileIndex, Packages.Num(), *Package->FileName);
+		}, EParallelForFlags::Unbalanced);
 	}
 
-	UE_LOG(LogIoStore, Display, TEXT("Adding localized import packages..."));
-	TArray<FPackage*> LocalizedPackages;
-
-	for (FPackage* Package : Packages)
 	{
-		LocalizedPackages.Reset();
-		for (FPackage* ImportedPackage : Package->ImportedPackages)
-		{
-			OutSourceToLocalizedPackageMap.MultiFind(ImportedPackage, LocalizedPackages);
-		}
-		for (FPackage* LocalizedPackage : LocalizedPackages)
-		{
-			UE_LOG(LogIoStore, Verbose, TEXT("For package '%s' (0x%llX): Adding localized imported package '%s' (0x%llX)"),
-				*Package->Name.ToString(),
-				Package->GlobalPackageId.ValueForDebugging(),
-				*LocalizedPackage->Name.ToString(),
-				LocalizedPackage->GlobalPackageId.ValueForDebugging());
-		}
-		Package->ImportedPackages.Append(LocalizedPackages);
-	}
+		IOSTORE_CPU_SCOPE(SerializeSummaries);
 
-	UE_LOG(LogIoStore, Display, TEXT("Conforming localized imports..."));
-	for (FPackage* Package : Packages)
-	{
-		for (FPackageObjectIndex& GlobalImportIndex : Package->Imports)
+		ParallelFor(TotalPackageCount, [
+				&ReadCount,
+				&PackageAssetBuffers,
+				&PackageFileSummaries,
+				&Packages,
+				&PackageStoreOptimizer](int32 Index)
 		{
-			if (GlobalImportIndex.IsPackageImport())
-			{
-				// When building DLC's the export can be missing
-				if (const FExportObjectData* Export = GlobalPackageData.FindPublicExport(GlobalImportIndex))
-				{
-					if (Export->Package->SourcePackageName != Export->Package->Name)
-					{
-						const FPackageObjectIndex* SourceGlobalImportIndex = LocalizedToSourceImportIndexMap.Find(GlobalImportIndex);
-						if (SourceGlobalImportIndex)
-						{
-							GlobalImportIndex = *SourceGlobalImportIndex;
-
-							const FExportObjectData& SourceExportData = *GlobalPackageData.FindPublicExport(*SourceGlobalImportIndex);
-							UE_LOG(LogIoStore, Verbose,
-								TEXT("For package '%s' (0x%llX): Remap localized import %s to source import %s (in a conformed localized package)"),
-								*Package->Name.ToString(),
-								Package->GlobalPackageId.ValueForDebugging(),
-								*Export->FullName,
-								*SourceExportData.FullName);
-						}
-						else
-						{
-							UE_LOG(LogIoStore, Verbose,
-								TEXT("For package '%s' (0x%llX): Skip remap for localized import %s")
-								TEXT(", either there is no source package or the localized package did not conform to it."),
-								*Package->Name.ToString(),
-								Package->GlobalPackageId.ValueForDebugging(),
-								*Export->FullName);
-						}
-					}
-				}
-			}
-		}
+			uint8* PackageBuffer = PackageAssetBuffers[Index];
+			FLegacyCookedPackage* Package = Packages[Index];
+
+			FIoBuffer CookedHeaderBuffer = FIoBuffer(FIoBuffer::Wrap, PackageBuffer, Package->UAssetSize);
+			Package->OptimizedPackage = PackageStoreOptimizer.CreatePackageFromCookedHeader(Package->PackageName, CookedHeaderBuffer);
+			check(Package->OptimizedPackage->GetId() == Package->GlobalPackageId);
+		}, EParallelForFlags::Unbalanced);
 	}
+
+	FMemory::Free(UAssetMemory);
 }
 
 TUniquePtr<FIoStoreReader> CreateIoStoreReader(const TCHAR* Path, const FKeyChain& KeyChain)
@@ -3055,11 +798,11 @@ TArray<TUniquePtr<FIoStoreReader>> CreatePatchSourceReaders(const TArray<FString
 
 void InitializeContainerTargetsAndPackages(
 	const FIoStoreArguments& Arguments,
-	TArray<FPackage*>& Packages,
+	TArray<FLegacyCookedPackage*>& Packages,
 	FPackageNameMap& PackageNameMap,
 	FPackageIdMap& PackageIdMap,
 	TArray<FContainerTargetSpec*>& ContainerTargets,
-	FNameMapBuilder& GlobalNameMapBuilder)
+	FPackageStoreOptimizer& PackageStoreOptimizer)
 {
 	FString ProjectName = FApp::GetProjectName();
 	FString RelativeEnginePath = FPaths::GetRelativePathToRoot();
@@ -3112,7 +855,6 @@ void InitializeContainerTargetsAndPackages(
 
 		ContainerTarget->PatchSourceReaders = CreatePatchSourceReaders(ContainerSource.PatchSourceContainerFiles, Arguments);
 
-		ContainerTarget->LocalNameMapBuilder.SetNameMapType(FMappedName::EType::Container);
 		{
 			IOSTORE_CPU_SCOPE(ProcessSourceFiles);
 			for (const FContainerSourceFile& SourceFile : ContainerSource.SourceFiles)
@@ -3137,17 +879,17 @@ void InitializeContainerTargetsAndPackages(
 				}
 
 				FString RelativeFileName = ConvertCookedPathToRelativePath(SourceFile.NormalizedPath);
-				FPackage* Package = nullptr;
+				FLegacyCookedPackage* Package = nullptr;
 				const bool bIsMemoryMappedBulkData = CookedFileStatData->FileExt == FCookedFileStatData::EFileExt::UMappedBulk;
 
 				if (bIsMemoryMappedBulkData)
 				{
 					FString TmpFileName = FString(RelativeFileName.Len() - 8, GetData(RelativeFileName)) + TEXT(".ubulk");
-					Package = FindOrAddPackage(Arguments, *TmpFileName, Packages, PackageNameMap, PackageIdMap);
+					Package = FindOrAddPackage(Arguments, *TmpFileName, Packages, PackageNameMap, PackageIdMap, PackageStoreOptimizer);
 				}
 				else
 				{
-					Package = FindOrAddPackage(Arguments, *RelativeFileName, Packages, PackageNameMap, PackageIdMap);
+					Package = FindOrAddPackage(Arguments, *RelativeFileName, Packages, PackageNameMap, PackageIdMap, PackageStoreOptimizer);
 				}
 
 				if (Package)
@@ -3170,7 +912,6 @@ void InitializeContainerTargetsAndPackages(
 					if (SourceFile.bNeedsEncryption)
 					{
 						ContainerTarget->ContainerFlags |= EIoContainerFlags::Encrypted;
-						ContainerTarget->bUseLocalNameMap = true;
 					}
 
 					if (CookedFileStatData->FileType == FCookedFileStatData::BulkData)
@@ -3206,7 +947,6 @@ void InitializeContainerTargetsAndPackages(
 						Package->UAssetSize = OriginalCookedFileStatData->FileSize;
 						Package->UExpSize = CookedFileStatData->FileSize;
 						TargetFile.ChunkId = CreateChunkId(Package->GlobalPackageId, 0, EIoChunkType::ExportBundleData, *TargetFile.TargetPath);
-						TargetFile.NameMapBuilder = &Package->LocalNameMapBuilder;
 					}
 
 					if (TargetFile.bForceUncompressed && !SourceFile.bNeedsEncryption)
@@ -3216,15 +956,10 @@ void InitializeContainerTargetsAndPackages(
 					}
 				}
 			}
-
-			if (!ContainerTarget->bUseLocalNameMap)
-			{
-				ContainerTarget->NameMapBuilder = &GlobalNameMapBuilder;
-			}
 		}
 	}
 
-	Algo::Sort(Packages, [](const FPackage* A, const FPackage* B)
+	Algo::Sort(Packages, [](const FLegacyCookedPackage* A, const FLegacyCookedPackage* B)
 	{
 		return A->GlobalPackageId < B->GlobalPackageId;
 	});
@@ -3386,10 +1121,10 @@ void LogContainerPackageInfo(const TArray<FContainerTargetSpec*>& ContainerTarge
 				continue;
 			}
 
-			UGraphSize += TargetFile.UGraphSize;
-			ImportMapSize += TargetFile.ImportMapSize;
-			ExportMapSize += TargetFile.ExportMapSize;
-			NameMapSize += TargetFile.NameMapSize;
+			UGraphSize += TargetFile.Package->OptimizedPackage->GetExportBundleArcsSize();
+			ImportMapSize += TargetFile.Package->OptimizedPackage->GetImportMapSize();
+			ExportMapSize += TargetFile.Package->OptimizedPackage->GetExportMapSize();
+			NameMapSize += TargetFile.Package->OptimizedPackage->GetNameMapSize();
 		}
 
 		HeaderSize = SummarySize + UGraphSize + ImportMapSize + ExportMapSize + NameMapSize;
@@ -3427,8 +1162,9 @@ void LogContainerPackageInfo(const TArray<FContainerTargetSpec*>& ContainerTarge
 class FIoStoreWriteRequestManager
 {
 public:
-	FIoStoreWriteRequestManager()
-		: MemoryAvailableEvent(FPlatformProcess::GetSynchEventFromPool(false))
+	FIoStoreWriteRequestManager(FPackageStoreOptimizer& InPackageStoreOptimizer)
+		: PackageStoreOptimizer(InPackageStoreOptimizer)
+		, MemoryAvailableEvent(FPlatformProcess::GetSynchEventFromPool(false))
 	{
 		InitiatorThread = Async(EAsyncExecution::Thread, [this]() { InitiatorThreadFunc(); });
 		RetirerThread = Async(EAsyncExecution::Thread, [this]() { RetirerThreadFunc(); });
@@ -3443,9 +1179,9 @@ public:
 		FPlatformProcess::ReturnSynchEventToPool(MemoryAvailableEvent);
 	}
 
-	IIoStoreWriteRequest* Read(const FContainerTargetFile& InTargetFile, const TArray<FObjectExport>* InObjectExports)
+	IIoStoreWriteRequest* Read(const FContainerTargetFile& InTargetFile)
 	{
-		return new FWriteContainerTargetFileRequest(*this, InTargetFile, InObjectExports);
+		return new FWriteContainerTargetFileRequest(*this, InTargetFile);
 	}
 
 private:
@@ -3453,13 +1189,11 @@ private:
 		: public IIoStoreWriteRequest
 	{
 	public:
-		FWriteContainerTargetFileRequest(FIoStoreWriteRequestManager& InManager, const FContainerTargetFile& InTargetFile, const TArray<FObjectExport>* InObjectExports)
+		FWriteContainerTargetFileRequest(FIoStoreWriteRequestManager& InManager, const FContainerTargetFile& InTargetFile)
 			: Manager(InManager)
 			, TargetFile(InTargetFile)
-			, ObjectExports(InObjectExports)
 			, FileRegions(TargetFile.FileRegions)
 		{
-			check(InTargetFile.bIsBulkData || InObjectExports);
 		}
 
 		~FWriteContainerTargetFileRequest()
@@ -3511,8 +1245,7 @@ private:
 			//TRACE_CPUPROFILER_EVENT_SCOPE(ReadCallback);
 			if (!TargetFile.bIsBulkData)
 			{
-				check(ObjectExports);
-				*SourceBuffer = CreateExportBundleBuffer(TargetFile, *ObjectExports, *SourceBuffer, bHasUpdatedExportBundleRegions ? nullptr : &FileRegions);
+				*SourceBuffer = Manager.PackageStoreOptimizer.CreatePackageBuffer(TargetFile.Package->OptimizedPackage, *SourceBuffer, bHasUpdatedExportBundleRegions ? nullptr : &FileRegions);
 				bHasUpdatedExportBundleRegions = true;
 			}
 			TArray<FBaseGraphTask*> NewTasks;
@@ -3522,7 +1255,6 @@ private:
 	private:
 		FIoStoreWriteRequestManager& Manager;
 		const FContainerTargetFile& TargetFile;
-		const TArray<FObjectExport>* ObjectExports;
 		TArray<FFileRegion> FileRegions;
 		FGraphEventRef CompletionEvent;
 		FIoBuffer* SourceBuffer = nullptr;
@@ -3696,6 +1428,7 @@ private:
 		}
 	}
 
+	FPackageStoreOptimizer& PackageStoreOptimizer;
 	TFuture<void> InitiatorThread;
 	TFuture<void> RetirerThread;
 	FQueue InitiatorQueue;
@@ -3759,22 +1492,19 @@ int32 CreateTarget(const FIoStoreArguments& Arguments, const FIoStoreWriterSetti
 	ChunkIdCsv.CreateOutputFile(CookedDir);
 #endif
 
-	FNameMapBuilder GlobalNameMapBuilder;
-	GlobalNameMapBuilder.SetNameMapType(FMappedName::EType::Global);
-	FPackageAssetData PackageAssetData;
-	FGlobalPackageData GlobalPackageData;
-
-	TArray<FPackage*> Packages;
+	TArray<FLegacyCookedPackage*> Packages;
 	FPackageNameMap PackageNameMap;
 	FPackageIdMap PackageIdMap;
 
-	FIoStoreWriteRequestManager WriteRequestManager;
+	FPackageStoreOptimizer PackageStoreOptimizer;
+	PackageStoreOptimizer.Initialize(Arguments.TargetPlatform);
+	FIoStoreWriteRequestManager WriteRequestManager(PackageStoreOptimizer);
 
 	TArray<FContainerTargetSpec*> ContainerTargets;
 	UE_LOG(LogIoStore, Display, TEXT("Creating container targets..."));
 	{
 		IOSTORE_CPU_SCOPE(CreateContainerTargets);
-		InitializeContainerTargetsAndPackages(Arguments, Packages, PackageNameMap, PackageIdMap, ContainerTargets, GlobalNameMapBuilder);
+		InitializeContainerTargetsAndPackages(Arguments, Packages, PackageNameMap, PackageIdMap, ContainerTargets, PackageStoreOptimizer);
 	}
 
 	TUniquePtr<FIoStoreWriterContext> IoStoreWriterContext(new FIoStoreWriterContext());
@@ -3789,7 +1519,7 @@ int32 CreateTarget(const FIoStoreArguments& Arguments, const FIoStoreWriterSetti
 		}
 		for (FContainerTargetSpec* ContainerTarget : ContainerTargets)
 		{
-			check(ContainerTarget->Header.ContainerId.IsValid());
+			check(ContainerTarget->ContainerId.IsValid());
 			if (!ContainerTarget->OutputPath.IsEmpty())
 			{
 				ContainerTarget->IoStoreWriter = new FIoStoreWriter(*ContainerTarget->OutputPath);
@@ -3815,7 +1545,7 @@ int32 CreateTarget(const FIoStoreArguments& Arguments, const FIoStoreWriterSetti
 			if (ContainerTarget->IoStoreWriter)
 			{
 				FIoContainerSettings ContainerSettings;
-				ContainerSettings.ContainerId = ContainerTarget->Header.ContainerId;
+				ContainerSettings.ContainerId = ContainerTarget->ContainerId;
 				if (Arguments.bCreateDirectoryIndex)
 				{
 					ContainerSettings.ContainerFlags = ContainerTarget->ContainerFlags | EIoContainerFlags::Indexed;
@@ -3840,128 +1570,67 @@ int32 CreateTarget(const FIoStoreArguments& Arguments, const FIoStoreWriterSetti
 		}
 	}
 
-	ParsePackageAssets(Packages, PackageAssetData);
+	ParsePackageAssets(Packages, PackageStoreOptimizer);
 
-	for (FContainerTargetSpec* ContainerTarget : ContainerTargets)
+	if (Arguments.IsDLC() && Arguments.bRemapPluginContentToGame)
 	{
-		if (ContainerTarget->IoStoreWriter)
+		for (FLegacyCookedPackage* Package : Packages)
 		{
-			for (FContainerTargetFile& TargetFile : ContainerTarget->TargetFiles)
+			const int32 DLCNameLen = Arguments.DLCName.Len() + 1;
+			FString PackageNameStr = Package->PackageName.ToString();
+			FString RedirectedPackageNameStr = TEXT("/Game");
+			RedirectedPackageNameStr.AppendChars(*PackageNameStr + DLCNameLen, PackageNameStr.Len() - DLCNameLen);
+			FName RedirectedPackageName = FName(*RedirectedPackageNameStr);
+			if (Arguments.ReleasedPackages.PackageNames.Contains(RedirectedPackageName))
 			{
-				if (TargetFile.bIsBulkData)
-				{
-					FIoWriteOptions WriteOptions;
-					WriteOptions.DebugName = *TargetFile.TargetPath;
-					WriteOptions.bForceUncompressed = TargetFile.bForceUncompressed;
-					WriteOptions.bIsMemoryMapped = TargetFile.bIsMemoryMappedBulkData;
-					WriteOptions.FileName = TargetFile.DestinationPath;
-					ContainerTarget->IoStoreWriter->Append(TargetFile.ChunkId, WriteRequestManager.Read(TargetFile, nullptr), WriteOptions);
-				}
+				Package->OptimizedPackage->RedirectFrom(RedirectedPackageName);
 			}
-		}
-	}
-
-	FExportGraph ExportGraph(PackageAssetData.ObjectExports.Num(), PackageAssetData.PreloadDependencies.Num());
-	GlobalPackageData.Reserve(PackageAssetData.ObjectExports.Num());
-
-	CreateGlobalScriptObjects(GlobalPackageData, Arguments.TargetPlatform);
-	TMap<FPackageId, FPackageId> SourceToRedirectedPackageMap;
-	CreateGlobalImportsAndExports(Arguments, Packages, PackageIdMap, PackageAssetData, GlobalPackageData, ExportGraph, SourceToRedirectedPackageMap);
-
-	// Mapped import and exports are required before processing localization, and preload/postload arcs
-	MapExportEntryIndices(PackageAssetData.ObjectExports, GlobalPackageData.ExportObjects, Packages);
-	// Intermediate map for adding extra ImportedPackages and ExternalArcs to the build graph
-	// for every dependency on a source package with localizations, add dependencies to all language localizations
-	FSourceToLocalizedPackageMultimap SourceToLocalizedPackageMap;
-
-	ProcessLocalizedPackages(Packages, PackageNameMap, GlobalPackageData, SourceToLocalizedPackageMap);
-
-	for (FPackage* Package : Packages)
-	{
-		if (Package->RedirectedPackageId.IsValid())
-		{
-			for (int32 ExportIndex : Package->Exports)
+			else
 			{
-				const FExportObjectData& ExportData = GlobalPackageData.ExportObjects[ExportIndex];
-				if (!ExportData.SuperIndex.IsNull() && ExportData.OuterIndex.IsNull())
-				{
-					UE_LOG(LogIoStore, Warning, TEXT("Skipping redirect to package '%s' due to presence of UStruct '%s'"), *Package->Name.ToString(), *ExportData.ObjectName.ToString());
-					Package->RedirectedPackageId = FPackageId();
-					break;
-				}
+				// TODO: Should NAME_None be the default source package name instead?
+				// RemapLocalizationPathIfNeeded sets it to the original
+				// package name when not remapping plugin content
+				Package->OptimizedPackage->RedirectFrom(NAME_None);
 			}
 		}
 	}
 
-	AddPreloadDependencies(
-		PackageAssetData,
-		GlobalPackageData,
-		SourceToLocalizedPackageMap,
-		SourceToRedirectedPackageMap,
-		ExportGraph,
-		Packages);
 
-	BuildBundles(ExportGraph, Packages);
-
-#if OUTPUT_DEBUG_PACKAGE_EXPORT_BUNDLES
-	if (!Arguments.MetaOutputDir.IsEmpty())
+	for (FContainerTargetSpec* ContainerTarget : ContainerTargets)
 	{
-		TUniquePtr<IFileHandle> ExportBundleMetaFile(FPlatformFileManager::Get().GetPlatformFile().OpenWrite(*(Arguments.MetaOutputDir / TEXT("iodispatcher.uexportbundles"))));
-		for (const FPackage* Package : Packages)
+		if (ContainerTarget->IoStoreWriter)
 		{
-			int32 ExportBundleIndex = 0;
-			for (const FExportBundle& ExportBundle : Package->ExportBundles)
+			for (FContainerTargetFile& TargetFile : ContainerTarget->TargetFiles)
 			{
-				FString Text = FString::Printf(TEXT("%s[%d]\n"), *Package->Name.ToString(), ExportBundleIndex);
-				ExportBundleMetaFile->Write((uint8*)StringCast<ANSICHAR>(*Text).Get(), Text.Len());
-				for (const FExportGraphNode* Node : ExportBundle.Nodes)
+				if (TargetFile.bIsBulkData)
 				{
-					Text = FString::Printf(TEXT("- %s %d\n"), Node->BundleEntry.CommandType == FExportBundleEntry::ExportCommandType_Serialize ? "S" : "C", Node->BundleEntry.LocalExportIndex);
-					ExportBundleMetaFile->Write((uint8*)StringCast<ANSICHAR>(*Text).Get(), Text.Len());
+					FIoWriteOptions WriteOptions;
+					WriteOptions.DebugName = *TargetFile.TargetPath;
+					WriteOptions.bForceUncompressed = TargetFile.bForceUncompressed;
+					WriteOptions.bIsMemoryMapped = TargetFile.bIsMemoryMappedBulkData;
+					WriteOptions.FileName = TargetFile.DestinationPath;
+					ContainerTarget->IoStoreWriter->Append(TargetFile.ChunkId, WriteRequestManager.Read(TargetFile), WriteOptions);
 				}
-				++ExportBundleIndex;
 			}
 		}
-		ExportBundleMetaFile->Flush();
-	}
-#endif
-
-	{
-		IOSTORE_CPU_SCOPE(FinalizeNameMaps);
-		UE_LOG(LogIoStore, Display, TEXT("Finalizing name maps..."));
-
-		for (FContainerTargetSpec* ContainerTarget : ContainerTargets)
-		{
-			FinalizeNameMaps(*ContainerTarget);
-		}
 	}
 
+	UE_LOG(LogIoStore, Display, TEXT("Optimizing packages..."));
+	for (FLegacyCookedPackage* Package : Packages)
 	{
-		IOSTORE_CPU_SCOPE(FinalizePackageHeaders);
-		UE_LOG(LogIoStore, Display, TEXT("Finalizing package headers..."));
-
-		for (FContainerTargetSpec* ContainerTarget : ContainerTargets)
-		{
-			FinalizePackageHeaders(
-				*ContainerTarget,
-				PackageAssetData.ObjectExports,
-				GlobalPackageData.ExportObjects,
-				GlobalPackageData.ImportsByFullName);
-
-			FinalizePackageStoreContainerHeader(*ContainerTarget);
-
-			const FNameMapBuilder& NameMapBuilder = *ContainerTarget->NameMapBuilder;
-			SaveNameBatch(ContainerTarget->LocalNameMapBuilder.GetNameMap(), ContainerTarget->Header.Names, ContainerTarget->Header.NameHashes);
-		}
+		check(Package->OptimizedPackage);
+		PackageStoreOptimizer.BeginResolvePackage(Package->OptimizedPackage);
 	}
+	PackageStoreOptimizer.Flush(true);
 
 	UE_LOG(LogIoStore, Display, TEXT("Creating disk layout..."));
-	CreateDiskLayout(ContainerTargets, Packages, Arguments.GameOrderMap, Arguments.CookerOrderMap);
+	CreateDiskLayout(ContainerTargets, Packages, Arguments.GameOrderMap, Arguments.CookerOrderMap, PackageIdMap);
 
 	for (FContainerTargetSpec* ContainerTarget : ContainerTargets)
 	{
 		if (ContainerTarget->IoStoreWriter)
 		{
+			TArray<FPackageStoreContainerHeaderEntry> ContainerHeaderEntries;
 			for (FContainerTargetFile& TargetFile : ContainerTarget->TargetFiles)
 			{
 				if (!TargetFile.bIsBulkData)
@@ -3970,18 +1639,23 @@ int32 CreateTarget(const FIoStoreArguments& Arguments, const FIoStoreWriterSetti
 					WriteOptions.DebugName = *TargetFile.TargetPath;
 					WriteOptions.bForceUncompressed = TargetFile.bForceUncompressed;
 					WriteOptions.FileName = TargetFile.DestinationPath;
-					ContainerTarget->IoStoreWriter->Append(TargetFile.ChunkId, WriteRequestManager.Read(TargetFile, &PackageAssetData.ObjectExports), WriteOptions);
+					ContainerTarget->IoStoreWriter->Append(TargetFile.ChunkId, WriteRequestManager.Read(TargetFile), WriteOptions);
+					ContainerHeaderEntries.Add(PackageStoreOptimizer.CreateContainerHeaderEntry(TargetFile.Package->OptimizedPackage));
 				}
 			}
 
-			FLargeMemoryWriter Ar(0, true);
-			Ar << ContainerTarget->Header;
-
+			ContainerTarget->Header = PackageStoreOptimizer.CreateContainerHeader(ContainerTarget->ContainerId, ContainerHeaderEntries);
+			FLargeMemoryWriter HeaderAr(0, true);
+			HeaderAr << ContainerTarget->Header;
+			int64 DataSize = HeaderAr.TotalSize();
+			FIoBuffer ContainerHeaderBuffer(FIoBuffer::AssumeOwnership, HeaderAr.ReleaseOwnership(), DataSize);
+			
 			FIoWriteOptions WriteOptions;
 			WriteOptions.DebugName = TEXT("ContainerHeader");
 			ContainerTarget->IoStoreWriter->Append(
-				CreateIoChunkId(ContainerTarget->Header.ContainerId.Value(), 0, EIoChunkType::ContainerHeader),
-				FIoBuffer(FIoBuffer::Wrap, Ar.GetData(), Ar.TotalSize()), WriteOptions);
+				CreateIoChunkId(ContainerTarget->ContainerId.Value(), 0, EIoChunkType::ContainerHeader),
+				ContainerHeaderBuffer,
+				WriteOptions);
 		}
 
 		// Check if we need to dump the final order of the packages. Useful, to debug packing.
@@ -4011,48 +1685,7 @@ int32 CreateTarget(const FIoStoreArguments& Arguments, const FIoStoreWriterSetti
 	uint64 InitialLoadSize = 0;
 	if (GlobalIoStoreWriter)
 	{
-		FLargeMemoryWriter InitialLoadArchive(0, true);
-		FinalizeInitialLoadMeta(
-			GlobalNameMapBuilder,
-			GlobalPackageData.ScriptObjects,
-			InitialLoadArchive);
-
-		InitialLoadSize = InitialLoadArchive.Tell();
-
-		UE_LOG(LogIoStore, Display, TEXT("Serializing global meta data"));
-		IOSTORE_CPU_SCOPE(SerializeInitialLoadMeta);
-		FIoWriteOptions WriteOptions;
-		WriteOptions.DebugName = TEXT("LoaderInitialLoadMeta");
-		GlobalIoStoreWriter->Append(CreateIoChunkId(0, 0, EIoChunkType::LoaderInitialLoadMeta), FIoBuffer(FIoBuffer::Wrap, InitialLoadArchive.GetData(), InitialLoadArchive.TotalSize()), WriteOptions);
-	}
-
-	uint64 GlobalNamesMB = 0;
-	uint64 GlobalNameHashesMB = 0;
-	if (GlobalIoStoreWriter)
-	{
-		IOSTORE_CPU_SCOPE(SerializeGlobalNameMap);
-
-		UE_LOG(LogIoStore, Display, TEXT("Saving global name map to container file"));
-
-		TArray<uint8> Names;
-		TArray<uint8> Hashes;
-		SaveNameBatch(GlobalNameMapBuilder.GetNameMap(), /* out */ Names, /* out */ Hashes);
-
-		InitialLoadSize += Names.Num() + Hashes.Num();
-		GlobalNamesMB = Names.Num() >> 20;
-		GlobalNameHashesMB = Hashes.Num() >> 20;
-
-		FIoWriteOptions WriteOptions;
-		WriteOptions.DebugName = TEXT("LoaderGlobalNames");
-		GlobalIoStoreWriter->Append(CreateIoChunkId(0, 0, EIoChunkType::LoaderGlobalNames), 
-													 FIoBuffer(FIoBuffer::Wrap, Names.GetData(), Names.Num()), WriteOptions);
-		WriteOptions.DebugName = TEXT("LoaderGlobalNameHashes");
-		GlobalIoStoreWriter->Append(CreateIoChunkId(0, 0, EIoChunkType::LoaderGlobalNameHashes),
-													 FIoBuffer(FIoBuffer::Wrap, Hashes.GetData(), Hashes.Num()), WriteOptions);
-		
-#if OUTPUT_NAMEMAP_CSV
-		NameMapBuilder.SaveCsv(OutputDir / TEXT("Container.namemap.csv"));
-#endif
+		InitialLoadSize = PackageStoreOptimizer.WriteScriptObjects(GlobalIoStoreWriter);
 	}
 
 	UE_LOG(LogIoStore, Display, TEXT("Serializing container(s)..."));
@@ -4075,45 +1708,18 @@ int32 CreateTarget(const FIoStoreArguments& Arguments, const FIoStoreWriterSetti
 	UE_LOG(LogIoStore, Display, TEXT("Calculating stats..."));
 	uint64 UExpSize = 0;
 	uint64 UAssetSize = 0;
-	uint64 SummarySize = 0;
-	uint64 PackageSummarySize = Packages.Num() * sizeof(FPackageSummary);
 	uint64 ImportedPackagesCount = 0;
 	uint64 NoImportedPackagesCount = 0;
-	uint64 PublicExportsCount = 0;
-	uint64 TotalExternalArcCount = 0;
 	uint64 NameMapCount = 0;
-
-	uint64 BundleCount = 0;
-	uint64 BundleEntryCount = 0;
-
-	for (const FPackage* Package : Packages)
+	
+	for (const FLegacyCookedPackage* Package : Packages)
 	{
 		UExpSize += Package->UExpSize;
 		UAssetSize += Package->UAssetSize;
-		SummarySize += Package->SummarySize;
-		NameMapCount += Package->SummaryNames.Num();
-		ImportedPackagesCount += Package->ImportedPackages.Num();
-		NoImportedPackagesCount += Package->ImportedPackages.Num() == 0;
-
-		for (auto& KV : Package->ExternalArcs)
-		{
-			const TArray<FArc>& Arcs = KV.Value;
-			TotalExternalArcCount += Arcs.Num();
-		}
-
-		for (const FExportBundle& Bundle : Package->ExportBundles)
-		{
-			++BundleCount;
-			BundleEntryCount += Bundle.Nodes.Num();
-		}
-	}
-
-	for (FExportObjectData& Export : GlobalPackageData.ExportObjects)
-	{
-		if (Export.IsPublicExport())
-		{
-			++PublicExportsCount;
-		}
+		NameMapCount += Package->OptimizedPackage->GetNameCount();
+		int32 PackageImportedPackagesCount = Package->OptimizedPackage->GetImportedPackageIds().Num();
+		ImportedPackagesCount += PackageImportedPackagesCount;
+		NoImportedPackagesCount += PackageImportedPackagesCount == 0;
 	}
 
 	LogWriterResults(IoStoreWriterResults);
@@ -4121,20 +1727,15 @@ int32 CreateTarget(const FIoStoreArguments& Arguments, const FIoStoreWriterSetti
 	
 	UE_LOG(LogIoStore, Display, TEXT("Input:  %8.2lf MB UExp"), (double)UExpSize / 1024.0 / 1024.0);
 	UE_LOG(LogIoStore, Display, TEXT("Input:  %8.2lf MB UAsset"), (double)UAssetSize / 1024.0 / 1024.0);
-	UE_LOG(LogIoStore, Display, TEXT("Input:  %8.2lf MB FPackageFileSummary"), (double)SummarySize / 1024.0 / 1024.0);
 	UE_LOG(LogIoStore, Display, TEXT("Input:  %8d Packages"), Packages.Num());
-	UE_LOG(LogIoStore, Display, TEXT("Input:  %8llu Imported package entries"), ImportedPackagesCount);
-	UE_LOG(LogIoStore, Display, TEXT("Input:  %8llu Packages without imports"), NoImportedPackagesCount);
-	UE_LOG(LogIoStore, Display, TEXT("Input:  %8llu Name map entries"), NameMapCount);
-	UE_LOG(LogIoStore, Display, TEXT("Input:  %8d PreloadDependencies entries"), PackageAssetData.PreloadDependencies.Num());
-	UE_LOG(LogIoStore, Display, TEXT("Input:  %8d ImportMap entries"), PackageAssetData.ObjectImports.Num());
-	UE_LOG(LogIoStore, Display, TEXT("Input:  %8d ExportMap entries"), PackageAssetData.ObjectExports.Num());
-	UE_LOG(LogIoStore, Display, TEXT("Input:  %8llu Public exports"), PublicExportsCount);
 	UE_LOG(LogIoStore, Display, TEXT(""));
-	UE_LOG(LogIoStore, Display, TEXT("Output: %8llu Export bundles"), BundleCount);
-	UE_LOG(LogIoStore, Display, TEXT("Output: %8llu Export bundle entries"), BundleEntryCount);
-	UE_LOG(LogIoStore, Display, TEXT("Output: %8llu Export bundle arcs"), TotalExternalArcCount);
-	UE_LOG(LogIoStore, Display, TEXT("Output: %8d Public runtime script objects"), GlobalPackageData.ScriptObjects.Num());
+	UE_LOG(LogIoStore, Display, TEXT("Output: %8llu Export bundles"), PackageStoreOptimizer.GetTotalExportBundleCount());
+	UE_LOG(LogIoStore, Display, TEXT("Output: %8llu Export bundle entries"), PackageStoreOptimizer.GetTotalExportBundleEntryCount());
+	UE_LOG(LogIoStore, Display, TEXT("Output: %8llu Export bundle arcs"), PackageStoreOptimizer.GetTotalExportBundleArcCount());
+	UE_LOG(LogIoStore, Display, TEXT("Output: %8llu Name map entries"), NameMapCount);
+	UE_LOG(LogIoStore, Display, TEXT("Output: %8llu Imported package entries"), ImportedPackagesCount);
+	UE_LOG(LogIoStore, Display, TEXT("Output: %8llu Packages without imports"), NoImportedPackagesCount);
+	UE_LOG(LogIoStore, Display, TEXT("Output: %8d Public runtime script objects"), PackageStoreOptimizer.GetTotalScriptObjectCount());
 	UE_LOG(LogIoStore, Display, TEXT("Output: %8.2lf MB InitialLoadData"), (double)InitialLoadSize / 1024.0 / 1024.0);
 
 	return 0;
diff --git a/Engine/Source/Developer/IoStoreUtilities/Private/PackageStoreOptimizer.cpp b/Engine/Source/Developer/IoStoreUtilities/Private/PackageStoreOptimizer.cpp
new file mode 100644
--- /dev/null
+++ b/Engine/Source/Developer/IoStoreUtilities/Private/PackageStoreOptimizer.cpp
@@ -0,0 +1,1633 @@
+// Copyright Epic Games, Inc. All Rights Reserved.
+
+#include "PackageStoreOptimizer.h"
+#include "ProfilingDebugging/CountersTrace.h"
+#include "Serialization/BufferWriter.h"
+#include "Serialization/LargeMemoryWriter.h"
+#include "UObject/NameBatchSerialization.h"
+#include "Containers/Map.h"
+#include "UObject/UObjectHash.h"
+#include "Serialization/PackageStore.h"
+#include "Interfaces/ITargetPlatform.h"
+#include "UObject/Object.h"
+#include "Serialization/MemoryReader.h"
+#include "UObject/Package.h"
+
+DEFINE_LOG_CATEGORY_STATIC(LogPackageStoreOptimizer, Log, All);
+
+// modified copy from SavePackage
+EObjectMark GetExcludedObjectMarksForTargetPlatform(const ITargetPlatform* TargetPlatform)
+{
+	EObjectMark Marks = OBJECTMARK_NOMARKS;
+	if (!TargetPlatform->HasEditorOnlyData())
+	{
+		Marks = (EObjectMark)(Marks | OBJECTMARK_EditorOnly);
+	}
+	if (TargetPlatform->IsServerOnly())
+	{
+		Marks = (EObjectMark)(Marks | OBJECTMARK_NotForServer);
+	}
+	if (TargetPlatform->IsClientOnly())
+	{
+		Marks = (EObjectMark)(Marks | OBJECTMARK_NotForClient);
+	}
+	return Marks;
+}
+
+// modified copy from SavePackage
+EObjectMark GetExcludedObjectMarksForObject(const UObject* Object, const ITargetPlatform* TargetPlatform)
+{
+	EObjectMark Marks = OBJECTMARK_NOMARKS;
+	if (!Object->NeedsLoadForClient())
+	{
+		Marks = (EObjectMark)(Marks | OBJECTMARK_NotForClient);
+	}
+	if (!Object->NeedsLoadForServer())
+	{
+		Marks = (EObjectMark)(Marks | OBJECTMARK_NotForServer);
+	}
+	if (!Object->NeedsLoadForTargetPlatform(TargetPlatform))
+	{
+		Marks = (EObjectMark)(Marks | OBJECTMARK_NotForClient | OBJECTMARK_NotForServer);
+	}
+	if (Object->IsEditorOnly())
+	{
+		Marks = (EObjectMark)(Marks | OBJECTMARK_EditorOnly);
+	}
+	if ((Marks & OBJECTMARK_NotForClient) && (Marks & OBJECTMARK_NotForServer))
+	{
+		Marks = (EObjectMark)(Marks | OBJECTMARK_EditorOnly);
+	}
+	return Marks;
+}
+
+// modified copy from PakFileUtilities
+static FString RemapLocalizationPathIfNeeded(const FString& Path, FString* OutRegion)
+{
+	static constexpr TCHAR L10NString[] = TEXT("/L10N/");
+	static constexpr int32 L10NPrefixLength = sizeof(L10NString) / sizeof(TCHAR) - 1;
+
+	int32 BeginL10NOffset = Path.Find(L10NString, ESearchCase::IgnoreCase);
+	if (BeginL10NOffset >= 0)
+	{
+		int32 EndL10NOffset = BeginL10NOffset + L10NPrefixLength;
+		int32 NextSlashIndex = Path.Find(TEXT("/"), ESearchCase::IgnoreCase, ESearchDir::FromStart, EndL10NOffset);
+		int32 RegionLength = NextSlashIndex - EndL10NOffset;
+		if (RegionLength >= 2)
+		{
+			FString NonLocalizedPath = Path.Mid(0, BeginL10NOffset) + Path.Mid(NextSlashIndex);
+			if (OutRegion)
+			{
+				*OutRegion = Path.Mid(EndL10NOffset, RegionLength);
+				OutRegion->ToLowerInline();
+			}
+			return NonLocalizedPath;
+		}
+	}
+	return Path;
+}
+
+void FPackageStoreOptimizer::Initialize(const ITargetPlatform* TargetPlatform)
+{
+	FindScriptObjects(TargetPlatform);
+}
+
+FPackageStorePackage* FPackageStoreOptimizer::CreatePackageFromCookedHeader(const FName& Name, const FIoBuffer& CookedHeaderBuffer) const
+{
+	FPackageStorePackage* Package = new FPackageStorePackage();
+	Package->Name = Name;
+	Package->Id = FPackageId::FromName(Name);
+
+	Package->SourceName = *RemapLocalizationPathIfNeeded(Name.ToString(), &Package->Region);
+
+	LoadCookedHeader(Package, CookedHeaderBuffer);
+	ProcessImports(Package);
+	ProcessExports(Package);
+
+	return Package;
+}
+
+void FPackageStoreOptimizer::LoadCookedHeader(FPackageStorePackage* Package, const FIoBuffer& CookedHeaderBuffer) const
+{
+	TArrayView<const uint8> MemView(CookedHeaderBuffer.Data(), CookedHeaderBuffer.DataSize());
+	FMemoryReaderView Ar(MemView);
+
+	{
+		TGuardValue<int32> GuardAllowUnversionedContentInEditor(GAllowUnversionedContentInEditor, 1);
+		Ar << Package->Summary;
+	}
+
+	Package->PackageFlags = Package->Summary.PackageFlags;
+	Package->CookedHeaderSize = Package->Summary.TotalHeaderSize;
+
+	Ar.SetFilterEditorOnly((Package->PackageFlags & EPackageFlags::PKG_FilterEditorOnly) != 0);
+
+	if (Package->Summary.NameCount > 0)
+	{
+		Ar.Seek(Package->Summary.NameOffset);
+
+		FNameEntrySerialized NameEntry(ENAME_LinkerConstructor);
+
+		//Package.SummaryNames.Reserve(Summary.NameCount);
+		for (int32 I = 0; I < Package->Summary.NameCount; ++I)
+		{
+			Ar << NameEntry;
+			FName Name(NameEntry);
+			//Package.SummaryNames.Add(Name);
+			Package->NameMapBuilder.AddName(Name);
+		}
+	}
+
+	class FNameReaderProxyArchive
+		: public FArchiveProxy
+	{
+	public:
+		using FArchiveProxy::FArchiveProxy;
+
+		FNameReaderProxyArchive(FArchive& InAr, const TArray<FNameEntryId>& InNameMap)
+			: FArchiveProxy(InAr)
+			, NameMap(InNameMap)
+		{
+			// Replicate the filter editor only state of the InnerArchive as FArchiveProxy will
+			// not intercept it.
+			FArchive::SetFilterEditorOnly(InAr.IsFilterEditorOnly());
+		}
+
+		FArchive& operator<<(FName& Name)
+		{
+			int32 NameIndex, Number;
+			InnerArchive << NameIndex << Number;
+
+			if (!NameMap.IsValidIndex(NameIndex))
+			{
+				UE_LOG(LogPackageStoreOptimizer, Fatal, TEXT("Bad name index %i/%i"), NameIndex, NameMap.Num());
+			}
+
+			const FNameEntryId MappedName = NameMap[NameIndex];
+			Name = FName::CreateFromDisplayId(MappedName, Number);
+
+			return *this;
+		}
+
+	private:
+		const TArray<FNameEntryId>& NameMap;
+	};
+	FNameReaderProxyArchive ProxyAr(Ar, Package->NameMapBuilder.GetNameMap());
+
+	if (Package->Summary.ImportCount > 0)
+	{
+		Package->ObjectImports.Reserve(Package->Summary.ImportCount);
+		ProxyAr.Seek(Package->Summary.ImportOffset);
+		for (int32 I = 0; I < Package->Summary.ImportCount; ++I)
+		{
+			ProxyAr << Package->ObjectImports.AddDefaulted_GetRef();
+		}
+	}
+
+	if (Package->Summary.PreloadDependencyCount > 0)
+	{
+		Package->PreloadDependencies.Reserve(Package->Summary.PreloadDependencyCount);
+		ProxyAr.Seek(Package->Summary.PreloadDependencyOffset);
+		for (int32 I = 0; I < Package->Summary.PreloadDependencyCount; ++I)
+		{
+			ProxyAr << Package->PreloadDependencies.AddDefaulted_GetRef();
+		}
+	}
+
+	if (Package->Summary.ExportCount > 0)
+	{
+		Package->ObjectExports.Reserve(Package->Summary.ExportCount);
+		ProxyAr.Seek(Package->Summary.ExportOffset);
+		for (int32 I = 0; I < Package->Summary.ExportCount; ++I)
+		{
+			FObjectExport& ObjectExport = Package->ObjectExports.AddDefaulted_GetRef();
+			ProxyAr << ObjectExport;
+			Package->ExportsSerialSize += ObjectExport.SerialSize;
+		}
+	}
+}
+
+void FPackageStoreOptimizer::ResolveImport(FPackageStorePackage::FImport* Imports, const FObjectImport* ObjectImports, int32 LocalImportIndex) const
+{
+	FPackageStorePackage::FImport* Import = Imports + LocalImportIndex;
+	if (Import->FullName.Len() == 0)
+	{
+		Import->FullName.Reserve(256);
+
+		const FObjectImport* ObjectImport = ObjectImports + LocalImportIndex;
+		if (ObjectImport->OuterIndex.IsNull())
+		{
+			FName PackageName = ObjectImport->ObjectName;
+			PackageName.AppendString(Import->FullName);
+			Import->FullName.ToLowerInline();
+			Import->PackageId = FPackageId::FromName(PackageName);
+			Import->bIsScriptImport = Import->FullName.StartsWith(TEXT("/Script/"));
+			Import->bIsPackageImport = true;
+		}
+		else
+		{
+			const int32 OuterIndex = ObjectImport->OuterIndex.ToImport();
+			ResolveImport(Imports, ObjectImports, OuterIndex);
+			FPackageStorePackage::FImport* OuterImport = Imports + OuterIndex;
+			check(OuterImport->FullName.Len() > 0);
+			Import->bIsScriptImport = OuterImport->bIsScriptImport;
+			Import->FullName.Append(OuterImport->FullName);
+			Import->FullName.AppendChar(TEXT('/'));
+			ObjectImport->ObjectName.AppendString(Import->FullName);
+			Import->FullName.ToLowerInline();
+			Import->PackageId = OuterImport->PackageId;
+		}
+	}
+}
+
+void FPackageStoreOptimizer::ProcessImports(FPackageStorePackage* Package) const
+{
+	int32 ImportCount = Package->ObjectImports.Num();
+	Package->Imports.SetNum(ImportCount);
+
+	for (int32 ImportIndex = 0; ImportIndex < ImportCount; ++ImportIndex)
+	{
+		ResolveImport(Package->Imports.GetData(), Package->ObjectImports.GetData(), ImportIndex);
+		FPackageStorePackage::FImport& Import = Package->Imports[ImportIndex];
+		if (Import.bIsScriptImport)
+		{
+			Import.GlobalImportIndex = FPackageObjectIndex::FromScriptPath(Import.FullName);
+		}
+		else if (!Import.bIsPackageImport)
+		{
+			Import.GlobalImportIndex = FPackageObjectIndex::FromPackagePath(Import.FullName);
+		}
+		else
+		{
+			Package->ImportedPackageIds.Add(Import.PackageId);
+		}
+	}
+}
+
+void FPackageStoreOptimizer::ResolveExport(
+	FPackageStorePackage::FExport* Exports,
+	const FObjectExport* ObjectExports,
+	const int32 LocalExportIndex,
+	const FName& PackageName) const
+{
+	FPackageStorePackage::FExport* Export = Exports + LocalExportIndex;
+	if (Export->FullName.Len() == 0)
+	{
+		Export->FullName.Reserve(256);
+		const FObjectExport* ObjectExport = ObjectExports + LocalExportIndex;
+		if (ObjectExport->OuterIndex.IsNull())
+		{
+			PackageName.AppendString(Export->FullName);
+			Export->FullName.AppendChar(TEXT('/'));
+			ObjectExport->ObjectName.AppendString(Export->FullName);
+			Export->FullName.ToLowerInline();
+			check(Export->FullName.Len() > 0);
+		}
+		else
+		{
+			check(ObjectExport->OuterIndex.IsExport());
+
+			int32 OuterExportIndex = ObjectExport->OuterIndex.ToExport();
+			ResolveExport(Exports, ObjectExports, OuterExportIndex, PackageName);
+			FString& OuterName = Exports[OuterExportIndex].FullName;
+			check(OuterName.Len() > 0);
+			Export->FullName.Append(OuterName);
+			Export->FullName.AppendChar(TEXT('/'));
+			ObjectExport->ObjectName.AppendString(Export->FullName);
+			Export->FullName.ToLowerInline();
+		}
+	}
+}
+
+void FPackageStoreOptimizer::ProcessExports(FPackageStorePackage* Package) const
+{
+	int32 ExportCount = Package->ObjectExports.Num();
+	Package->Exports.SetNum(ExportCount);
+	Package->ExportGraphNodes.Reserve(ExportCount * 2);
+
+	auto PackageObjectIdFromPackageIndex =
+		[](const TArray<FPackageStorePackage::FImport>& Imports, const FPackageIndex& PackageIndex) -> FPackageObjectIndex
+	{
+		if (PackageIndex.IsImport())
+		{
+			return Imports[PackageIndex.ToImport()].GlobalImportIndex;
+		}
+		if (PackageIndex.IsExport())
+		{
+			return FPackageObjectIndex::FromExportIndex(PackageIndex.ToExport());
+		}
+		return FPackageObjectIndex();
+	};
+
+	for (int32 ExportIndex = 0; ExportIndex < ExportCount; ++ExportIndex)
+	{
+		const FObjectExport& ObjectExport = Package->ObjectExports[ExportIndex];
+		FPackageStorePackage::FExport& Export = Package->Exports[ExportIndex];
+		Export.ObjectName = ObjectExport.ObjectName;
+		Export.ObjectFlags = ObjectExport.ObjectFlags;
+		Export.CookedSerialOffset = ObjectExport.SerialOffset;
+		Export.CookedSerialSize = ObjectExport.SerialSize;
+		Export.bNotForClient = ObjectExport.bNotForClient;
+		Export.bNotForServer = ObjectExport.bNotForServer;
+		Export.bIsPublic = (Export.ObjectFlags & RF_Public) > 0;
+		ResolveExport(Package->Exports.GetData(), Package->ObjectExports.GetData(), ExportIndex, Package->Name);
+		if (Export.bIsPublic)
+		{
+			check(Export.FullName.Len() > 0);
+			Export.GlobalImportIndex = FPackageObjectIndex::FromPackagePath(Export.FullName);
+		}
+
+		Export.OuterIndex = PackageObjectIdFromPackageIndex(Package->Imports, ObjectExport.OuterIndex);
+		Export.ClassIndex = PackageObjectIdFromPackageIndex(Package->Imports, ObjectExport.ClassIndex);
+		Export.SuperIndex = PackageObjectIdFromPackageIndex(Package->Imports, ObjectExport.SuperIndex);
+		Export.TemplateIndex = PackageObjectIdFromPackageIndex(Package->Imports, ObjectExport.TemplateIndex);
+
+		Export.CreateNode = &Package->ExportGraphNodes.AddDefaulted_GetRef();
+		Export.CreateNode->Package = Package;
+		Export.CreateNode->BundleEntry.CommandType = FExportBundleEntry::ExportCommandType_Create;
+		Export.CreateNode->BundleEntry.LocalExportIndex = ExportIndex;
+
+		Export.SerializeNode = &Package->ExportGraphNodes.AddDefaulted_GetRef();
+		Export.SerializeNode->Package = Package;
+		Export.SerializeNode->BundleEntry.CommandType = FExportBundleEntry::ExportCommandType_Serialize;
+		Export.SerializeNode->BundleEntry.LocalExportIndex = ExportIndex;
+	}
+}
+
+void FPackageStoreOptimizer::SortPackagesInLoadOrder(TArray<FPackageStorePackage*>& Packages) const
+{
+	TRACE_CPUPROFILER_EVENT_SCOPE(SortPackagesInLoadOrder);
+	Algo::Sort(Packages, [](const FPackageStorePackage* A, const FPackageStorePackage* B)
+	{
+		return A->Id < B->Id;
+	});
+
+	TMap<FPackageStorePackage*, TArray<FPackageStorePackage*>> SortedEdges;
+	for (FPackageStorePackage* Package : Packages)
+	{
+		for (FPackageStorePackage* ImportedPackage : Package->ImportedWaitingPackages)
+		{
+			TArray<FPackageStorePackage*>& SourceArray = SortedEdges.FindOrAdd(ImportedPackage);
+			SourceArray.Add(Package);
+		}
+		for (FPackageId ImportedRedirectedPackageId : Package->ImportedRedirectedPackageIds)
+		{
+			FPackageStorePackage* FindImportedPackage = WaitingPackagesMap.FindRef(ImportedRedirectedPackageId);
+			if (FindImportedPackage)
+			{
+				TArray<FPackageStorePackage*>& SourceArray = SortedEdges.FindOrAdd(FindImportedPackage);
+				SourceArray.Add(Package);
+			}
+		}
+	}
+	for (auto& KV : SortedEdges)
+	{
+		TArray<FPackageStorePackage*>& SourceArray = KV.Value;
+		Algo::Sort(SourceArray, [](const FPackageStorePackage* A, const FPackageStorePackage* B)
+		{
+			return A->Id < B->Id;
+		});
+	}
+
+	TArray<FPackageStorePackage*> Result;
+	Result.Reserve(Packages.Num());
+	struct
+	{
+		void Visit(FPackageStorePackage* Package)
+		{
+			if (Package->bPermanentMark || Package->bTemporaryMark)
+			{
+				return;
+			}
+			Package->bTemporaryMark = true;
+			TArray<FPackageStorePackage*>* TargetNodes = Edges.Find(Package);
+			if (TargetNodes)
+			{
+				for (FPackageStorePackage* ToNode : *TargetNodes)
+				{
+					Visit(ToNode);
+				}
+			}
+			Package->bTemporaryMark = false;
+			Package->bPermanentMark = true;
+			Result.Add(Package);
+		}
+
+		TMap<FPackageStorePackage*, TArray<FPackageStorePackage*>>& Edges;
+		TArray<FPackageStorePackage*>& Result;
+	} Visitor{ SortedEdges, Result };
+
+	for (FPackageStorePackage* Package : Packages)
+	{
+		Visitor.Visit(Package);
+	}
+	check(Result.Num() == Packages.Num());
+	Algo::Reverse(Result);
+	Swap(Result, Packages);
+}
+
+FPackageStoreOptimizer::FExportGraphEdges FPackageStoreOptimizer::ProcessPreloadDependencies(const TArray<FPackageStorePackage*>& Packages) const
+{
+	TRACE_CPUPROFILER_EVENT_SCOPE(ProcessPreloadDependencies);
+	FExportGraphEdges Edges;
+	auto AddExternalDependency = [this, &Edges, &Packages](FPackageStorePackage* Package, int32 FromImportIndex, FExportBundleEntry::EExportCommandType FromExportBundleCommandType, FPackageStorePackage::FExportGraphNode* ToNode)
+	{
+		const FPackageStorePackage::FImport& FromImport = Package->Imports[FromImportIndex];
+		if (FromImport.GlobalImportIndex.IsScriptImport())
+		{
+			return;
+		}
+
+		FPackageStorePackage::FExternalDependency& ExternalDependency = ToNode->ExternalDependencies.AddDefaulted_GetRef();
+		ExternalDependency.PackageId = FromImport.PackageId;
+		ExternalDependency.GlobalImportIndex = FromImport.GlobalImportIndex;
+		ExternalDependency.DebugImportFullName = FName(FromImport.FullName);
+		ExternalDependency.ExportBundleCommandType = FromExportBundleCommandType;
+
+		const TArray<FResolvedPublicExport>* FindResolvedPublicExports = ResolvedPublicExportsMap.Find(FromImport.PackageId);
+		if (FindResolvedPublicExports)
+		{
+			bool bFoundExport = false;
+			for (const FResolvedPublicExport& PublicExport : *FindResolvedPublicExports)
+			{
+				if (PublicExport.GlobalImportIndex == FromImport.GlobalImportIndex)
+				{
+					if (FromExportBundleCommandType == FExportBundleEntry::ExportCommandType_Create)
+					{
+						check(PublicExport.CreateExportBundleIndex >= 0);
+						ExternalDependency.ResolvedExportBundleIndex = PublicExport.CreateExportBundleIndex;
+					}
+					else
+					{
+						check(PublicExport.SerializeExportBundleIndex >= 0);
+						ExternalDependency.ResolvedExportBundleIndex = PublicExport.SerializeExportBundleIndex;
+					}
+					bFoundExport = true;
+					break;
+				}
+			}
+			if (!bFoundExport)
+			{
+				ExternalDependency.bIsConfirmedMissing = true;
+			}
+		}
+		else
+		{
+			FPackageStorePackage* FindWaitingFromPackage = WaitingPackagesMap.FindRef(FromImport.PackageId);
+			if (FindWaitingFromPackage)
+			{
+				bool bFoundExport = false;
+				for (int32 ExportIndex = 0; ExportIndex < FindWaitingFromPackage->Exports.Num(); ++ExportIndex)
+				{
+					FPackageStorePackage::FExport& FromExport = FindWaitingFromPackage->Exports[ExportIndex];
+					if (FromExport.GlobalImportIndex == FromImport.GlobalImportIndex)
+					{
+						FPackageStorePackage::FExportGraphNode* FromNode = FromExportBundleCommandType == FExportBundleEntry::ExportCommandType_Create ? FromExport.CreateNode : FromExport.SerializeNode;
+						Edges.Add(FromNode, ToNode);
+						ExternalDependency.ExportGraphNode = FromNode;
+
+						for (const FPackageId& FromRedirectedPackageId : FindWaitingFromPackage->RedirectedToPackageIds)
+						{
+							FPackageStorePackage* FindRedirectedFromPackage = WaitingPackagesMap.FindRef(FromRedirectedPackageId);
+							if (FindRedirectedFromPackage)
+							{
+								FPackageStorePackage::FExport& FromRedirectedExport = FindRedirectedFromPackage->Exports[ExportIndex];
+								FPackageStorePackage::FExportGraphNode* FromRedirectedNode = FromExportBundleCommandType == FExportBundleEntry::ExportCommandType_Create ? FromRedirectedExport.CreateNode : FromRedirectedExport.SerializeNode;
+								Edges.Add(FromRedirectedNode, ToNode);
+
+								FPackageStorePackage::FExternalDependency& RedirectedExternalDependency = ToNode->ExternalDependencies.AddDefaulted_GetRef();
+								RedirectedExternalDependency = ExternalDependency;
+								RedirectedExternalDependency.PackageId = FromRedirectedPackageId;
+								RedirectedExternalDependency.ExportGraphNode = FromRedirectedNode;
+							}
+						}
+
+						bFoundExport = true;
+						break;
+					}
+				}
+				if (!bFoundExport)
+				{
+					ExternalDependency.bIsConfirmedMissing = true;
+				}
+
+			}
+		}
+	};
+
+	for (FPackageStorePackage* Package : Packages)
+	{
+		for (int32 ExportIndex = 0; ExportIndex < Package->Exports.Num(); ++ExportIndex)
+		{
+			FPackageStorePackage::FExport& Export = Package->Exports[ExportIndex];
+			const FObjectExport& ObjectExport = Package->ObjectExports[ExportIndex];
+
+			Edges.Add(Export.CreateNode, Export.SerializeNode);
+
+			if (ObjectExport.FirstExportDependency >= 0)
+			{
+				int32 RunningIndex = ObjectExport.FirstExportDependency;
+				for (int32 Index = ObjectExport.SerializationBeforeSerializationDependencies; Index > 0; Index--)
+				{
+					FPackageIndex Dep = Package->PreloadDependencies[RunningIndex++];
+					if (Dep.IsExport())
+					{
+						Edges.Add(Package->Exports[Dep.ToExport()].SerializeNode, Export.SerializeNode);
+					}
+					else
+					{
+						AddExternalDependency(Package, Dep.ToImport(), FExportBundleEntry::ExportCommandType_Serialize, Export.SerializeNode);
+					}
+				}
+
+				for (int32 Index = ObjectExport.CreateBeforeSerializationDependencies; Index > 0; Index--)
+				{
+					FPackageIndex Dep = Package->PreloadDependencies[RunningIndex++];
+					if (Dep.IsExport())
+					{
+						Edges.Add(Package->Exports[Dep.ToExport()].CreateNode, Export.SerializeNode);
+					}
+					else
+					{
+						AddExternalDependency(Package, Dep.ToImport(), FExportBundleEntry::ExportCommandType_Create, Export.SerializeNode);
+					}
+				}
+
+				for (int32 Index = ObjectExport.SerializationBeforeCreateDependencies; Index > 0; Index--)
+				{
+					FPackageIndex Dep = Package->PreloadDependencies[RunningIndex++];
+					if (Dep.IsExport())
+					{
+						Edges.Add(Package->Exports[Dep.ToExport()].SerializeNode, Export.CreateNode);
+					}
+					else
+					{
+						AddExternalDependency(Package, Dep.ToImport(), FExportBundleEntry::ExportCommandType_Serialize, Export.CreateNode);
+					}
+				}
+
+				for (int32 Index = ObjectExport.CreateBeforeCreateDependencies; Index > 0; Index--)
+				{
+					FPackageIndex Dep = Package->PreloadDependencies[RunningIndex++];
+					if (Dep.IsExport())
+					{
+						Edges.Add(Package->Exports[Dep.ToExport()].CreateNode, Export.CreateNode);
+					}
+					else
+					{
+						AddExternalDependency(Package, Dep.ToImport(), FExportBundleEntry::ExportCommandType_Create, Export.CreateNode);
+					}
+				}
+			}
+		}
+	}
+	return Edges;
+}
+
+TArray<FPackageStorePackage::FExportGraphNode*> FPackageStoreOptimizer::SortExportGraphNodesInLoadOrder(const TArray<FPackageStorePackage*>& Packages, FExportGraphEdges& Edges) const
+{
+	TRACE_CPUPROFILER_EVENT_SCOPE(SortExportGraphNodesInLoadOrder);
+	int32 NodeCount = 0;
+	for (FPackageStorePackage* Package : Packages)
+	{
+		NodeCount += Package->ExportGraphNodes.Num();
+	}
+	for (auto& KV : Edges)
+	{
+		FPackageStorePackage::FExportGraphNode* ToNode = KV.Value;
+		++ToNode->IncomingEdgeCount;
+	}
+
+	auto NodeSorter = [](const FPackageStorePackage::FExportGraphNode& A, const FPackageStorePackage::FExportGraphNode& B)
+	{
+		if (A.BundleEntry.LocalExportIndex == B.BundleEntry.LocalExportIndex)
+		{
+			return A.BundleEntry.CommandType < B.BundleEntry.CommandType;
+		}
+		return A.BundleEntry.LocalExportIndex < B.BundleEntry.LocalExportIndex;
+	};
+
+	for (FPackageStorePackage* Package : Packages)
+	{
+		for (FPackageStorePackage::FExportGraphNode& Node : Package->ExportGraphNodes)
+		{
+			if (Node.IncomingEdgeCount == 0)
+			{
+				Package->NodesWithNoIncomingEdges.HeapPush(&Node, NodeSorter);
+			}
+		}
+	}
+	TArray<FPackageStorePackage::FExportGraphNode*> LoadOrder;
+	LoadOrder.Reserve(NodeCount);
+	while (LoadOrder.Num() < NodeCount)
+	{
+		bool bMadeProgress = false;
+		for (FPackageStorePackage* Package : Packages)
+		{
+			while (Package->NodesWithNoIncomingEdges.Num())
+			{
+				FPackageStorePackage::FExportGraphNode* RemovedNode;
+				Package->NodesWithNoIncomingEdges.HeapPop(RemovedNode, NodeSorter, false);
+				LoadOrder.Add(RemovedNode);
+				bMadeProgress = true;
+				for (auto EdgeIt = Edges.CreateKeyIterator(RemovedNode); EdgeIt; ++EdgeIt)
+				{
+					FPackageStorePackage::FExportGraphNode* ToNode = EdgeIt.Value();
+					check(ToNode->IncomingEdgeCount > 0);
+					if (--ToNode->IncomingEdgeCount == 0)
+					{
+						ToNode->Package->NodesWithNoIncomingEdges.HeapPush(ToNode, NodeSorter);
+					}
+					EdgeIt.RemoveCurrent();
+				}
+			}
+		}
+		check(bMadeProgress);
+	}
+	return LoadOrder;
+}
+
+void FPackageStoreOptimizer::CreateExportBundles(TArray<FPackageStorePackage*>& Packages)
+{
+	TRACE_CPUPROFILER_EVENT_SCOPE(CreateExportBundles);
+	SortPackagesInLoadOrder(Packages);
+	FExportGraphEdges Edges = ProcessPreloadDependencies(Packages);
+	TArray<FPackageStorePackage::FExportGraphNode*> LoadOrder = SortExportGraphNodesInLoadOrder(Packages, Edges);
+
+	FPackageStorePackage* LastPackage = nullptr;
+	for (FPackageStorePackage::FExportGraphNode* Node : LoadOrder)
+	{
+		check(Node->Package);
+		FPackageStorePackage::FExportBundle* Bundle;
+		if (Node->Package != LastPackage)
+		{
+			Node->ExportBundleIndex = Node->Package->ExportBundles.Num();
+			Bundle = &Node->Package->ExportBundles.AddDefaulted_GetRef();
+			Bundle->LoadOrder = NextLoadOrder++;
+			++TotalExportBundleCount;
+			LastPackage = Node->Package;
+		}
+		else
+		{
+			Node->ExportBundleIndex = Node->Package->ExportBundles.Num() - 1;
+			Bundle = &Node->Package->ExportBundles[Node->ExportBundleIndex];
+		}
+		Bundle->Entries.Add(Node->BundleEntry);
+		++TotalExportBundleEntryCount;
+
+		for (FPackageStorePackage::FExternalDependency& ExternalDependency : Node->ExternalDependencies)
+		{
+			if (ExternalDependency.ExportGraphNode)
+			{
+				check(ExternalDependency.ExportGraphNode->ExportBundleIndex >= 0);
+				ExternalDependency.ResolvedExportBundleIndex = ExternalDependency.ExportGraphNode->ExportBundleIndex;
+				ExternalDependency.ExportGraphNode = nullptr;
+			}
+		}
+	}
+}
+
+static const TCHAR* GetExportNameSafe(const FString& ExportFullName, const FName& PackageName, int32 PackageNameLen)
+{
+	const bool bValidNameLen = ExportFullName.Len() > PackageNameLen + 1;
+	if (bValidNameLen)
+	{
+		const TCHAR* ExportNameStr = *ExportFullName + PackageNameLen;
+		const bool bValidNameFormat = *ExportNameStr == '/';
+		if (bValidNameFormat)
+		{
+			return ExportNameStr + 1; // skip verified '/'
+		}
+		else
+		{
+			UE_CLOG(!bValidNameFormat, LogPackageStoreOptimizer, Warning,
+				TEXT("Export name '%s' should start with '/' at position %d, i.e. right after package prefix '%s'"),
+				*ExportFullName,
+				PackageNameLen,
+				*PackageName.ToString());
+		}
+	}
+	else
+	{
+		UE_CLOG(!bValidNameLen, LogPackageStoreOptimizer, Warning,
+			TEXT("Export name '%s' with length %d should be longer than package name '%s' with length %d"),
+			*ExportFullName,
+			PackageNameLen,
+			*PackageName.ToString());
+	}
+
+	return nullptr;
+};
+
+bool FPackageStoreOptimizer::RedirectPackage(
+	const FPackageStorePackage& SourcePackage,
+	FPackageStorePackage& TargetPackage,
+	TMap<FPackageObjectIndex, FPackageObjectIndex>& RedirectedToSourceImportIndexMap) const
+{
+	const int32 ExportCount =
+		SourcePackage.Exports.Num() < TargetPackage.Exports.Num() ?
+		SourcePackage.Exports.Num() :
+		TargetPackage.Exports.Num();
+
+	UE_CLOG(SourcePackage.Exports.Num() != TargetPackage.Exports.Num(), LogPackageStoreOptimizer, Verbose,
+		TEXT("Redirection target package '%s' (0x%llX) for source package '%s' (0x%llX)  - Has ExportCount %d vs. %d"),
+		*TargetPackage.Name.ToString(),
+		TargetPackage.Id.ValueForDebugging(),
+		*SourcePackage.Name.ToString(),
+		SourcePackage.Id.ValueForDebugging(),
+		TargetPackage.Exports.Num(),
+		SourcePackage.Exports.Num());
+
+	auto AppendMismatchMessage = [&TargetPackage, &SourcePackage](
+		const TCHAR* Text, FName ExportName, FPackageObjectIndex TargetIndex, FPackageObjectIndex SourceIndex, FString& FailReason)
+	{
+		FailReason.Appendf(TEXT("Public export '%s' has %s %s vs. %s"),
+			*ExportName.ToString(),
+			Text,
+			*TargetPackage.Exports[TargetIndex.ToExport()].FullName,
+			*SourcePackage.Exports[SourceIndex.ToExport()].FullName);
+	};
+
+	const int32 TargetPackageNameLen = TargetPackage.Name.GetStringLength();
+	const int32 SourcePackageNameLen = SourcePackage.Name.GetStringLength();
+
+	TArray <TPair<int32, int32>, TInlineAllocator<64>> NewPublicExports;
+	NewPublicExports.Reserve(ExportCount);
+
+	bool bSuccess = true;
+	int32 TargetIndex = 0;
+	int32 SourceIndex = 0;
+	while (TargetIndex < ExportCount && SourceIndex < ExportCount)
+	{
+		FString FailReason;
+		const FPackageStorePackage::FExport& TargetExport = TargetPackage.Exports[TargetIndex];
+		const FPackageStorePackage::FExport& SourceExport = SourcePackage.Exports[SourceIndex];
+
+		const TCHAR* TargetExportStr = GetExportNameSafe(
+			TargetExport.FullName, TargetPackage.Name, TargetPackageNameLen);
+		const TCHAR* SourceExportStr = GetExportNameSafe(
+			SourceExport.FullName, SourcePackage.Name, SourcePackageNameLen);
+
+		if (!TargetExportStr || !SourceExportStr)
+		{
+			UE_LOG(LogPackageStoreOptimizer, Error,
+				TEXT("Redirection target package '%s' (0x%llX) for source package '%s' (0x%llX) - Has some bad data from an earlier phase."),
+				*TargetPackage.Name.ToString(),
+				TargetPackage.Id.ValueForDebugging(),
+				*SourcePackage.Name.ToString(),
+				SourcePackage.Id.ValueForDebugging())
+				return false;
+		}
+
+		int32 CompareResult = FCString::Stricmp(TargetExportStr, SourceExportStr);
+		if (CompareResult < 0)
+		{
+			++TargetIndex;
+
+			if (TargetExport.bIsPublic)
+			{
+				// public localized export is missing in the source package, so just keep it as it is
+				NewPublicExports.Emplace(TargetIndex - 1, 1);
+			}
+		}
+		else if (CompareResult > 0)
+		{
+			++SourceIndex;
+
+			if (SourceExport.bIsPublic)
+			{
+				FailReason.Appendf(TEXT("Public source export '%s' is missing in the localized package"),
+					*SourceExport.ObjectName.ToString());
+			}
+		}
+		else
+		{
+			++TargetIndex;
+			++SourceIndex;
+
+			if (SourceExport.bIsPublic)
+			{
+				if (!TargetExport.bIsPublic)
+				{
+					FailReason.Appendf(TEXT("Public source export '%s' exists in the localized package")
+						TEXT(", but is not a public localized export."),
+						*SourceExport.ObjectName.ToString());
+				}
+				else if (TargetExport.ClassIndex != SourceExport.ClassIndex)
+				{
+					AppendMismatchMessage(TEXT("class"), TargetExport.ObjectName,
+						TargetExport.ClassIndex, SourceExport.ClassIndex, FailReason);
+				}
+				else if (TargetExport.TemplateIndex != SourceExport.TemplateIndex)
+				{
+					AppendMismatchMessage(TEXT("template"), TargetExport.ObjectName,
+						TargetExport.TemplateIndex, SourceExport.TemplateIndex, FailReason);
+				}
+				else if (TargetExport.SuperIndex != SourceExport.SuperIndex)
+				{
+					AppendMismatchMessage(TEXT("super"), TargetExport.ObjectName,
+						TargetExport.SuperIndex, SourceExport.SuperIndex, FailReason);
+				}
+				else
+				{
+					NewPublicExports.Emplace(TargetIndex - 1, SourceIndex - 1);
+				}
+			}
+			else if (TargetExport.bIsPublic)
+			{
+				FailReason.Appendf(TEXT("Export '%s' exists in the source package")
+					TEXT(", but is not a public source export."),
+					*TargetExport.ObjectName.ToString());
+			}
+		}
+
+		if (FailReason.Len() > 0)
+		{
+			UE_LOG(LogPackageStoreOptimizer, Warning,
+				TEXT("Redirection target package '%s' (0x%llX) for '%s' (0x%llX) - %s"),
+				*TargetPackage.Name.ToString(),
+				TargetPackage.Id.ValueForDebugging(),
+				*SourcePackage.Name.ToString(),
+				SourcePackage.Id.ValueForDebugging(),
+				*FailReason);
+			bSuccess = false;
+		}
+	}
+
+	if (bSuccess)
+	{
+		for (TPair<int32, int32>& Pair : NewPublicExports)
+		{
+			FPackageStorePackage::FExport& TargetExport = TargetPackage.Exports[Pair.Key];
+			if (Pair.Value != -1)
+			{
+				const FPackageStorePackage::FExport& SourceExport = SourcePackage.Exports[Pair.Value];
+
+				RedirectedToSourceImportIndexMap.Add(TargetExport.GlobalImportIndex, SourceExport.GlobalImportIndex);
+				TargetExport.GlobalImportIndex = SourceExport.GlobalImportIndex;
+			}
+		}
+	}
+
+	return bSuccess;
+}
+
+void FPackageStoreOptimizer::RedirectPackageUnverified(FPackageStorePackage& TargetPackage, TMap<FPackageObjectIndex, FPackageObjectIndex>& RedirectedToSourceImportIndexMap) const
+{
+	const int32 TargetPackageNameLen = TargetPackage.Name.GetStringLength();
+	TStringBuilder<256> NameStringBuilder;
+	for (FPackageStorePackage::FExport& Export : TargetPackage.Exports)
+	{
+		if (Export.bIsPublic)
+		{
+			const TCHAR* ExportNameStr = GetExportNameSafe(Export.FullName, TargetPackage.Name, TargetPackageNameLen);
+			NameStringBuilder.Reset();
+			TargetPackage.SourceName.AppendString(NameStringBuilder);
+			NameStringBuilder.Append(TEXT("/"));
+			NameStringBuilder.Append(ExportNameStr);
+			FPackageObjectIndex SourceGlobalImportIndex = FPackageObjectIndex::FromPackagePath(NameStringBuilder);
+			RedirectedToSourceImportIndexMap.Add(Export.GlobalImportIndex, SourceGlobalImportIndex);
+			Export.GlobalImportIndex = SourceGlobalImportIndex;
+		}
+	}
+}
+
+void FPackageStoreOptimizer::ProcessRedirects(const TArray<FPackageStorePackage*> Packages) const
+{
+	TRACE_CPUPROFILER_EVENT_SCOPE(ProcessRedirects);
+	TMap<FPackageObjectIndex, FPackageObjectIndex> RedirectedToSourceImportIndexMap;
+	TMap<FPackageId, FPackageId> SourceToRedirectedPackageMap;
+
+	for (FPackageStorePackage* Package : Packages)
+	{
+		if (Package->SourceName.IsNone() || Package->Name == Package->SourceName)
+		{
+			continue;
+		}
+		
+		FPackageId SourcePackageId = FPackageId::FromName(Package->SourceName);
+		FPackageStorePackage* SourcePackage = WaitingPackagesMap.FindRef(SourcePackageId);
+		
+		if (SourcePackage)
+		{
+			Package->bIsRedirected = RedirectPackage(*SourcePackage, *Package, RedirectedToSourceImportIndexMap);
+		}
+		else
+		{
+			if (Package->Region.Len() > 0)
+			{
+				// no update or verification required
+				UE_LOG(LogPackageStoreOptimizer, Verbose,
+					TEXT("For culture '%s': Localized package '%s' (0x%llX) is unique and does not override a source package."),
+					*Package->Region,
+					*Package->Name.ToString(),
+					Package->Id.ValueForDebugging());
+				continue;
+			}
+
+			// Assume DLC base game redirection
+			// TODO: Unify these two redirection paths
+			RedirectPackageUnverified(*Package, RedirectedToSourceImportIndexMap);
+			SourceToRedirectedPackageMap.Add(SourcePackageId, Package->Id);
+			Package->bIsRedirected = true;
+			for (const FPackageStorePackage::FExport& Export : Package->Exports)
+			{
+				if (!Export.SuperIndex.IsNull() && Export.OuterIndex.IsNull())
+				{
+					UE_LOG(LogPackageStoreOptimizer, Warning, TEXT("Skipping redirect to package '%s' due to presence of UStruct '%s'"), *Package->Name.ToString(), *Export.ObjectName.ToString());
+					Package->bIsRedirected = false;
+					break;
+				}
+			}
+		}
+		
+		if (Package->bIsRedirected)
+		{
+			UE_LOG(LogPackageStoreOptimizer, Verbose, TEXT("Adding package redirect from '%s' (0x%llX) to '%s' (0x%llX)."),
+				*Package->SourceName.ToString(),
+				SourcePackageId.ValueForDebugging(),
+				*Package->Name.ToString(),
+				Package->Id.ValueForDebugging());
+			if (SourcePackage)
+			{
+				SourcePackage->RedirectedToPackageIds.Add(Package->Id);
+			}
+		}
+		else
+		{
+			UE_LOG(LogPackageStoreOptimizer, Display,
+				TEXT("Skipping package redirect from '%s' (0x%llX) to '%s' (0x%llX) due to mismatching public exports."),
+				*Package->SourceName.ToString(),
+				SourcePackageId.ValueForDebugging(),
+				*Package->Name.ToString(),
+				Package->Id.ValueForDebugging());
+		}
+	}
+
+	for (FPackageStorePackage* Package : Packages)
+	{
+		for (FPackageStorePackage* ImportedPackage : Package->ImportedWaitingPackages)
+		{
+			Package->ImportedRedirectedPackageIds.Append(ImportedPackage->RedirectedToPackageIds);
+		}
+
+		for (FPackageStorePackage::FImport& Import : Package->Imports)
+		{
+			if (Import.GlobalImportIndex.IsPackageImport())
+			{
+				const FPackageObjectIndex* FindSourceGlobalImportIndex = RedirectedToSourceImportIndexMap.Find(Import.GlobalImportIndex);
+				if (FindSourceGlobalImportIndex)
+				{
+					Import.GlobalImportIndex = *FindSourceGlobalImportIndex;
+				}
+
+				const FPackageId* FindRedirectedPackageId = SourceToRedirectedPackageMap.Find(Import.PackageId);
+				if (FindRedirectedPackageId)
+				{
+					Import.PackageId = *FindRedirectedPackageId;
+				}
+			}
+		}
+	}
+}
+
+void FPackageStoreOptimizer::ResolvePackages(TArray<FPackageStorePackage*>& Packages, bool bAllowMissingImports)
+{
+	ProcessRedirects(Packages);
+	CreateExportBundles(Packages);
+	for (FPackageStorePackage* Package : Packages)
+	{
+		FinalizePackage(Package, bAllowMissingImports);
+	}
+}
+
+void FPackageStoreOptimizer::FinalizePackageHeader(FPackageStorePackage* Package) const
+{
+	// Temporary Archive for serializing ImportMap
+	FBufferWriter ImportMapArchive(nullptr, 0, EBufferWriterFlags::AllowResize | EBufferWriterFlags::TakeOwnership);
+	for (const FPackageStorePackage::FImport& Import : Package->Imports)
+	{
+		FPackageObjectIndex GlobalImportIndex = Import.GlobalImportIndex;
+		ImportMapArchive << GlobalImportIndex;
+	}
+	Package->ImportMapSize = ImportMapArchive.Tell();
+
+	// Temporary Archive for serializing EDL graph data
+	FBufferWriter GraphArchive(nullptr, 0, EBufferWriterFlags::AllowResize | EBufferWriterFlags::TakeOwnership);
+	int32 ReferencedPackagesCount = Package->ExternalArcs.Num();
+	GraphArchive << ReferencedPackagesCount;
+	TArray<TTuple<FPackageId, TArray<FPackageStorePackage::FArc>>> SortedExternalArcs;
+	SortedExternalArcs.Reserve(Package->ExternalArcs.Num());
+	for (auto& KV : Package->ExternalArcs)
+	{
+		FPackageId ImportedPackageId = KV.Key;
+		TArray<FPackageStorePackage::FArc> SortedArcs = KV.Value;
+		Algo::Sort(SortedArcs, [](const FPackageStorePackage::FArc& A, const FPackageStorePackage::FArc& B)
+		{
+			if (A.FromNodeIndex == B.FromNodeIndex)
+			{
+				return A.ToNodeIndex < B.ToNodeIndex;
+			}
+			return A.FromNodeIndex < B.FromNodeIndex;
+		});
+		SortedExternalArcs.Emplace(ImportedPackageId, MoveTemp(SortedArcs));
+	}
+	Algo::Sort(SortedExternalArcs, [](const TTuple<FPackageId, TArray<FPackageStorePackage::FArc>>& A, const TTuple<FPackageId, TArray<FPackageStorePackage::FArc>>& B)
+	{
+		return A.Key < B.Key;
+	});
+	for (auto& KV : SortedExternalArcs)
+	{
+		FPackageId ImportedPackageId = KV.Key;
+		TArray<FPackageStorePackage::FArc>& Arcs = KV.Value;
+		int32 ExternalArcCount = Arcs.Num();
+
+		GraphArchive << ImportedPackageId;
+		GraphArchive << ExternalArcCount;
+		GraphArchive.Serialize(Arcs.GetData(), ExternalArcCount * sizeof(FPackageStorePackage::FArc));
+	}
+	Package->ExportBundleArcsSize = GraphArchive.Tell();
+
+	// Temporary Archive for serializing export map data
+	FBufferWriter ExportMapArchive(nullptr, 0, EBufferWriterFlags::AllowResize | EBufferWriterFlags::TakeOwnership);
+	for (const FPackageStorePackage::FExport& Export : Package->Exports)
+	{
+		FExportMapEntry ExportMapEntry;
+		ExportMapEntry.CookedSerialOffset = Export.CookedSerialOffset;
+		ExportMapEntry.CookedSerialSize = Export.CookedSerialSize;
+		ExportMapEntry.ObjectName = Package->NameMapBuilder.MapName(Export.ObjectName);
+		ExportMapEntry.OuterIndex = Export.OuterIndex;
+		ExportMapEntry.ClassIndex = Export.ClassIndex;
+		ExportMapEntry.SuperIndex = Export.SuperIndex;
+		ExportMapEntry.TemplateIndex = Export.TemplateIndex;
+		ExportMapEntry.GlobalImportIndex = Export.GlobalImportIndex;
+		ExportMapEntry.ObjectFlags = Export.ObjectFlags;
+		ExportMapEntry.FilterFlags = EExportFilterFlags::None;
+		if (Export.bNotForClient)
+		{
+			ExportMapEntry.FilterFlags = EExportFilterFlags::NotForClient;
+		}
+		else if (Export.bNotForServer)
+		{
+			ExportMapEntry.FilterFlags = EExportFilterFlags::NotForServer;
+		}
+
+		ExportMapArchive << ExportMapEntry;
+	}
+	Package->ExportMapSize = ExportMapArchive.Tell();
+
+	// Temporary archive for serializing export bundle data
+	FBufferWriter ExportBundlesArchive(nullptr, 0, EBufferWriterFlags::AllowResize | EBufferWriterFlags::TakeOwnership);
+	uint32 ExportBundleEntryIndex = 0;
+	for (const FPackageStorePackage::FExportBundle& ExportBundle : Package->ExportBundles)
+	{
+		const uint32 EntryCount = ExportBundle.Entries.Num();
+		FExportBundleHeader ExportBundleHeader{ ExportBundleEntryIndex, EntryCount };
+		ExportBundlesArchive << ExportBundleHeader;
+
+		ExportBundleEntryIndex += EntryCount;
+	}
+	for (const FPackageStorePackage::FExportBundle& ExportBundle : Package->ExportBundles)
+	{
+		for (FExportBundleEntry BundleEntry : ExportBundle.Entries)
+		{
+			ExportBundlesArchive << BundleEntry;
+		}
+	}
+	Package->ExportBundleHeadersSize = ExportBundlesArchive.Tell();
+
+	Package->NameMapBuilder.MarkNameAsReferenced(Package->Name);
+	FMappedName MappedPackageName = Package->NameMapBuilder.MapName(Package->Name);
+	Package->NameMapBuilder.MarkNameAsReferenced(Package->SourceName);
+	FMappedName MappedPackageSourceName = Package->NameMapBuilder.MapName(Package->SourceName);
+
+	TArray<uint8> NamesBuffer;
+	TArray<uint8> NameHashesBuffer;
+	SaveNameBatch(Package->NameMapBuilder.GetNameMap(), NamesBuffer, NameHashesBuffer);
+	Package->NameMapSize = Align(NamesBuffer.Num(), 8) + NameHashesBuffer.Num();
+
+	uint64 HeaderSerialSize =
+		sizeof(FPackageSummary)
+		+ Package->NameMapSize
+		+ Package->ImportMapSize
+		+ Package->ExportMapSize
+		+ Package->ExportBundleHeadersSize
+		+ Package->ExportBundleArcsSize;
+
+	Package->OptimizedHeaderBuffer = FIoBuffer(HeaderSerialSize);
+	uint8* HeaderData = Package->OptimizedHeaderBuffer.Data();
+	FMemory::Memzero(HeaderData, HeaderSerialSize);
+	FPackageSummary* PackageSummary = reinterpret_cast<FPackageSummary*>(HeaderData);
+
+	PackageSummary->Name = MappedPackageName;
+	PackageSummary->SourceName = MappedPackageSourceName;
+	PackageSummary->PackageFlags = Package->PackageFlags;
+	PackageSummary->CookedHeaderSize = Package->CookedHeaderSize;
+	FBufferWriter HeaderArchive(HeaderData, HeaderSerialSize);
+	HeaderArchive.Seek(sizeof(FPackageSummary));
+
+	// NameMap data
+	{
+		PackageSummary->NameMapNamesOffset = HeaderArchive.Tell();
+		check(PackageSummary->NameMapNamesOffset % 8 == 0);
+		PackageSummary->NameMapNamesSize = NamesBuffer.Num();
+		HeaderArchive.Serialize(NamesBuffer.GetData(), NamesBuffer.Num());
+		PackageSummary->NameMapHashesOffset = Align(HeaderArchive.Tell(), 8);
+		int32 PaddingByteCount = PackageSummary->NameMapHashesOffset - HeaderArchive.Tell();
+		if (PaddingByteCount)
+		{
+			check(PaddingByteCount < 8);
+			uint8 PaddingBytes[8]{ 0 };
+			HeaderArchive.Serialize(PaddingBytes, PaddingByteCount);
+		}
+		PackageSummary->NameMapHashesSize = NameHashesBuffer.Num();
+		HeaderArchive.Serialize(NameHashesBuffer.GetData(), NameHashesBuffer.Num());
+	}
+
+	// ImportMap data
+	{
+		check(ImportMapArchive.Tell() == Package->ImportMapSize);
+		PackageSummary->ImportMapOffset = HeaderArchive.Tell();
+		HeaderArchive.Serialize(ImportMapArchive.GetWriterData(), ImportMapArchive.Tell());
+	}
+
+	// ExportMap data
+	{
+		check(ExportMapArchive.Tell() == Package->ExportMapSize);
+		PackageSummary->ExportMapOffset = HeaderArchive.Tell();
+		HeaderArchive.Serialize(ExportMapArchive.GetWriterData(), ExportMapArchive.Tell());
+	}
+
+	// ExportBundle data
+	{
+		check(ExportBundlesArchive.Tell() == Package->ExportBundleHeadersSize);
+		PackageSummary->ExportBundlesOffset = HeaderArchive.Tell();
+		HeaderArchive.Serialize(ExportBundlesArchive.GetWriterData(), ExportBundlesArchive.Tell());
+	}
+
+	// Graph data
+	{
+		check(GraphArchive.Tell() == Package->ExportBundleArcsSize);
+		PackageSummary->GraphDataOffset = HeaderArchive.Tell();
+		PackageSummary->GraphDataSize = Package->ExportBundleArcsSize;
+		HeaderArchive.Serialize(GraphArchive.GetWriterData(), GraphArchive.Tell());
+	}
+}
+
+void FPackageStoreOptimizer::FinalizePackage(FPackageStorePackage* Package, bool bAllowMissingImports)
+{
+	for (FPackageStorePackage::FExportGraphNode& Node : Package->ExportGraphNodes)
+	{
+		check(Node.ExportBundleIndex >= 0);
+		for (FPackageStorePackage::FExternalDependency& ExternalDependency : Node.ExternalDependencies)
+		{
+			if (ExternalDependency.bIsConfirmedMissing)
+			{
+				UE_LOG(LogPackageStoreOptimizer, Warning, TEXT("Package '%s' missing import '%s'"), *Package->Name.ToString(), *ExternalDependency.DebugImportFullName.ToString());
+				continue;
+			}
+			TArray<FPackageStorePackage::FArc>& ArcsFromImportedPackage = Package->ExternalArcs.FindOrAdd(ExternalDependency.PackageId);
+			FPackageStorePackage::FArc Arc;
+			Arc.FromNodeIndex = ExternalDependency.ResolvedExportBundleIndex;
+			Arc.ToNodeIndex = Node.ExportBundleIndex;
+			if (!ArcsFromImportedPackage.Contains(Arc))
+			{
+				ArcsFromImportedPackage.Add(Arc);
+				++TotalExportBundleArcCount;
+			}
+		}
+	}
+
+	FinalizePackageHeader(Package);
+
+	for (FPackageStorePackage* ImportedPackage : Package->ImportedWaitingPackages)
+	{
+		ImportedPackage->ImportedByWaitingPackages.Remove(Package);
+	}
+	for (FPackageStorePackage* ImportedByPackage : Package->ImportedByWaitingPackages)
+	{
+		ImportedByPackage->ImportedWaitingPackages.Remove(Package);
+	}
+	WaitingPackagesMap.Remove(Package->Id);
+
+	TArray<FResolvedPublicExport>& ResolvedPublicExports = ResolvedPublicExportsMap.FindOrAdd(Package->Id);
+	check(ResolvedPublicExports.IsEmpty());
+	for (const FPackageStorePackage::FExport& Export : Package->Exports)
+	{
+		if (Export.bIsPublic)
+		{
+			FResolvedPublicExport& PublicExport = ResolvedPublicExports.AddDefaulted_GetRef();
+			PublicExport.GlobalImportIndex = Export.GlobalImportIndex;
+			check(!PublicExport.GlobalImportIndex.IsNull());
+			PublicExport.CreateExportBundleIndex = Export.CreateNode->ExportBundleIndex;
+			check(PublicExport.CreateExportBundleIndex >= 0);
+			PublicExport.SerializeExportBundleIndex = Export.SerializeNode->ExportBundleIndex;
+			check(PublicExport.SerializeExportBundleIndex >= 0);
+		}
+	}
+
+	ResolvedPackages.Add(Package);
+}
+
+bool FPackageStoreOptimizer::IsWaitingForImportsRecursive(FPackageStorePackage* Package) const
+{
+	if (Package->UnresolvedImportedPackagesCount)
+	{
+		return true;
+	}
+
+	if (Package->bTemporaryMark)
+	{
+		return false;
+	}
+	Package->bTemporaryMark = true;
+
+	for (FPackageStorePackage* ImportedPackage : Package->ImportedWaitingPackages)
+	{
+		if (IsWaitingForImportsRecursive(ImportedPackage))
+		{
+			Package->bTemporaryMark = false;
+			return true;
+		}
+	}
+	Package->bTemporaryMark = false;
+	return false;
+}
+
+void FPackageStoreOptimizer::BeginResolvePackage(FPackageStorePackage* Package)
+{
+	++TotalPackageCount;
+	for (const FPackageId& ImportedPackageId : Package->ImportedPackageIds)
+	{
+		check(ImportedPackageId != Package->Id);
+		if (!ResolvedPublicExportsMap.Contains(ImportedPackageId))
+		{
+			FPackageStorePackage* FindWaitingPackage = WaitingPackagesMap.FindRef(ImportedPackageId);
+			if (FindWaitingPackage)
+			{
+				Package->ImportedWaitingPackages.Add(FindWaitingPackage);
+				FindWaitingPackage->ImportedByWaitingPackages.Add(Package);
+			}
+			else
+			{
+				++Package->UnresolvedImportedPackagesCount;
+				ImportedPackageToWaitingPackagesMap.Add(ImportedPackageId, Package);
+			}
+		}
+		else
+		{
+			check(!WaitingPackagesMap.Contains(ImportedPackageId));
+		}
+	}
+
+	WaitingPackagesMap.Add(Package->Id, Package);
+	
+	for (auto WaitingPackageIt = ImportedPackageToWaitingPackagesMap.CreateKeyIterator(Package->Id); WaitingPackageIt; ++WaitingPackageIt)
+	{
+		FPackageStorePackage* WaitingPackage = WaitingPackageIt.Value();
+		WaitingPackage->ImportedWaitingPackages.Add(Package);
+		Package->ImportedByWaitingPackages.Add(WaitingPackage);
+		check(WaitingPackage->UnresolvedImportedPackagesCount > 0);
+		WaitingPackageIt.RemoveCurrent();
+	}
+
+	if (!bIncrementalResolveEnabled)
+	{
+		return;
+	}
+	TArray<FPackageStorePackage*> PackagesReadyToResolve;
+	{
+		TRACE_CPUPROFILER_EVENT_SCOPE(UpdateWaitList);
+		for (auto& KV : WaitingPackagesMap)
+		{
+			FPackageStorePackage* WaitingPackage = KV.Value;
+			if (!IsWaitingForImportsRecursive(WaitingPackage))
+			{
+				PackagesReadyToResolve.Add(WaitingPackage);
+			}
+		}
+	}
+
+	if (PackagesReadyToResolve.Num())
+	{
+		ResolvePackages(PackagesReadyToResolve, false);
+	}
+}
+
+FIoBuffer FPackageStoreOptimizer::CreatePackageBuffer(const FPackageStorePackage* Package, const FIoBuffer& CookedExportsDataBuffer, TArray<FFileRegion>* InOutFileRegions) const
+{
+	check(Package->OptimizedHeaderBuffer.DataSize() > 0);
+	const uint64 BundleBufferSize = Package->OptimizedHeaderBuffer.DataSize() + Package->ExportsSerialSize;
+	FIoBuffer BundleBuffer(BundleBufferSize);
+	FMemory::Memcpy(BundleBuffer.Data(), Package->OptimizedHeaderBuffer.Data(), Package->OptimizedHeaderBuffer.DataSize());
+	uint64 BundleBufferOffset = Package->OptimizedHeaderBuffer.DataSize();
+
+	TArray<FFileRegion> OutputRegions;
+
+	for (const FPackageStorePackage::FExportBundle& ExportBundle : Package->ExportBundles)
+	{
+		for (const FExportBundleEntry& BundleEntry : ExportBundle.Entries)
+		{
+			if (BundleEntry.CommandType == FExportBundleEntry::ExportCommandType_Serialize)
+			{
+				const FPackageStorePackage::FExport& Export = Package->Exports[BundleEntry.LocalExportIndex];
+				uint64 AdjustedSerialOffset = Export.CookedSerialOffset - Package->CookedHeaderSize;
+				check(AdjustedSerialOffset + Export.CookedSerialSize <= CookedExportsDataBuffer.DataSize());
+				FMemory::Memcpy(BundleBuffer.Data() + BundleBufferOffset, CookedExportsDataBuffer.Data() + AdjustedSerialOffset, Export.CookedSerialSize);
+
+				if (InOutFileRegions)
+				{
+					// Find overlapping regions and adjust them to match the new offset of the export data
+					for (const FFileRegion& Region : *InOutFileRegions)
+					{
+						if (AdjustedSerialOffset <= Region.Offset && Region.Offset + Region.Length <= AdjustedSerialOffset + Export.CookedSerialSize)
+						{
+							FFileRegion NewRegion = Region;
+							NewRegion.Offset -= AdjustedSerialOffset;
+							NewRegion.Offset += BundleBufferOffset;
+							OutputRegions.Add(NewRegion);
+						}
+					}
+				}
+				BundleBufferOffset += Export.CookedSerialSize;
+			}
+		}
+	}
+	check(BundleBufferOffset == BundleBuffer.DataSize());
+
+	if (InOutFileRegions)
+	{
+		*InOutFileRegions = OutputRegions;
+	}
+
+	return BundleBuffer;
+}
+
+void FPackageStoreOptimizer::FindScriptObjectsRecursive(FPackageObjectIndex OuterIndex, UObject* Object, EObjectMark ExcludedObjectMarks, const ITargetPlatform* TargetPlatform)
+{
+	if (!Object->HasAllFlags(RF_Public))
+	{
+		UE_LOG(LogPackageStoreOptimizer, Verbose, TEXT("Skipping script object: %s (!RF_Public)"), *Object->GetFullName());
+		return;
+	}
+
+	const UObject* ObjectForExclusion = Object->HasAnyFlags(RF_ClassDefaultObject) ? (const UObject*)Object->GetClass() : Object;
+	const EObjectMark ObjectMarks = GetExcludedObjectMarksForObject(ObjectForExclusion, TargetPlatform);
+
+	if (ObjectMarks & ExcludedObjectMarks)
+	{
+		UE_LOG(LogPackageStoreOptimizer, Verbose, TEXT("Skipping script object: %s (Excluded for target platform)"), *Object->GetFullName());
+		return;
+	}
+
+	const FScriptObjectData* Outer = ScriptObjectsMap.Find(OuterIndex);
+	check(Outer);
+
+	FName ObjectName = Object->GetFName();
+
+	FString TempFullName = ScriptObjectsMap.FindRef(OuterIndex).FullName;
+	TempFullName.AppendChar(TEXT('/'));
+	ObjectName.AppendString(TempFullName);
+
+	TempFullName.ToLowerInline();
+	FPackageObjectIndex GlobalImportIndex = FPackageObjectIndex::FromScriptPath(TempFullName);
+
+	FScriptObjectData* ScriptImport = ScriptObjectsMap.Find(GlobalImportIndex);
+	if (ScriptImport)
+	{
+		UE_LOG(LogPackageStoreOptimizer, Fatal, TEXT("Import name hash collision \"%s\" and \"%s"), *TempFullName, *ScriptImport->FullName);
+	}
+
+	FPackageObjectIndex CDOClassIndex = Outer->CDOClassIndex;
+	if (CDOClassIndex.IsNull())
+	{
+		TCHAR NameBuffer[FName::StringBufferSize];
+		uint32 Len = ObjectName.ToString(NameBuffer);
+		if (FCString::Strncmp(NameBuffer, TEXT("Default__"), 9) == 0)
+		{
+			FString CDOClassFullName = Outer->FullName;
+			CDOClassFullName.AppendChar(TEXT('/'));
+			CDOClassFullName.AppendChars(NameBuffer + 9, Len - 9);
+			CDOClassFullName.ToLowerInline();
+
+			CDOClassIndex = FPackageObjectIndex::FromScriptPath(CDOClassFullName);
+		}
+	}
+
+	ScriptImport = &ScriptObjectsMap.Add(GlobalImportIndex);
+	ScriptImport->GlobalIndex = GlobalImportIndex;
+	ScriptImport->FullName = MoveTemp(TempFullName);
+	ScriptImport->OuterIndex = Outer->GlobalIndex;
+	ScriptImport->ObjectName = ObjectName;
+	ScriptImport->CDOClassIndex = CDOClassIndex;
+
+	TArray<UObject*> InnerObjects;
+	GetObjectsWithOuter(Object, InnerObjects, /*bIncludeNestedObjects*/false);
+	for (UObject* InnerObject : InnerObjects)
+	{
+		FindScriptObjectsRecursive(GlobalImportIndex, InnerObject, ExcludedObjectMarks, TargetPlatform);
+	}
+};
+
+void FPackageStoreOptimizer::FindScriptObjects(const ITargetPlatform* TargetPlatform)
+{
+	const EObjectMark ExcludedObjectMarks = GetExcludedObjectMarksForTargetPlatform(TargetPlatform);
+
+	TArray<UPackage*> ScriptPackages;
+	FindAllRuntimeScriptPackages(ScriptPackages);
+
+	TArray<UObject*> InnerObjects;
+	for (UPackage* Package : ScriptPackages)
+	{
+		FName ObjectName = Package->GetFName();
+		FString FullName = Package->GetName();
+
+		FullName.ToLowerInline();
+		FPackageObjectIndex GlobalImportIndex = FPackageObjectIndex::FromScriptPath(FullName);
+
+		FScriptObjectData* ScriptImport = ScriptObjectsMap.Find(GlobalImportIndex);
+		checkf(!ScriptImport, TEXT("Import name hash collision \"%s\" and \"%s"), *FullName, *ScriptImport->FullName);
+
+		ScriptImport = &ScriptObjectsMap.Add(GlobalImportIndex);
+		ScriptImport->GlobalIndex = GlobalImportIndex;
+		ScriptImport->FullName = FullName;
+		ScriptImport->OuterIndex = FPackageObjectIndex();
+		ScriptImport->ObjectName = ObjectName;
+
+		InnerObjects.Reset();
+		GetObjectsWithOuter(Package, InnerObjects, /*bIncludeNestedObjects*/false);
+		for (UObject* InnerObject : InnerObjects)
+		{
+			FindScriptObjectsRecursive(GlobalImportIndex, InnerObject, ExcludedObjectMarks, TargetPlatform);
+		}
+	}
+
+	FLargeMemoryWriter ScriptObjectsArchive(0, true);
+
+	int32 NumScriptObjects = ScriptObjectsMap.Num();
+	ScriptObjectsArchive << NumScriptObjects;
+	TotalScriptObjectCount = ScriptObjectsMap.Num();
+}
+
+void FPackageStoreOptimizer::Flush(bool bAllowMissingImports, TFunction<void(FPackageStorePackage*)> ResolvedPackageCallback)
+{
+	if (bAllowMissingImports)
+	{
+		TArray<FPackageStorePackage*> PackagesToResolve;
+		WaitingPackagesMap.GenerateValueArray(PackagesToResolve);
+		ResolvePackages(PackagesToResolve, true);
+		
+		check(WaitingPackagesMap.IsEmpty());
+		ImportedPackageToWaitingPackagesMap.Empty();
+	}
+
+	if (ResolvedPackageCallback)
+	{
+		for (FPackageStorePackage* ResolvedPackage : ResolvedPackages)
+		{
+			ResolvedPackageCallback(ResolvedPackage);
+		}
+	}
+	ResolvedPackages.Empty();
+}
+
+uint64 FPackageStoreOptimizer::WriteScriptObjects(FIoStoreWriter* IoStoreWriter) const
+{
+	TArray<FScriptObjectData> ScriptObjectsAsArray;
+	ScriptObjectsMap.GenerateValueArray(ScriptObjectsAsArray);
+	Algo::Sort(ScriptObjectsAsArray, [](const FScriptObjectData& A, const FScriptObjectData& B)
+	{
+		return A.FullName < B.FullName;
+	});
+	
+	FLargeMemoryWriter ScriptObjectsArchive(0, true);
+	int32 NumScriptObjects = ScriptObjectsAsArray.Num();
+	ScriptObjectsArchive << NumScriptObjects;
+	FPackageStoreNameMapBuilder NameMapBuilder;
+	NameMapBuilder.SetNameMapType(FMappedName::EType::Global);
+	for (const FScriptObjectData& ImportData : ScriptObjectsAsArray)
+	{
+		NameMapBuilder.MarkNameAsReferenced(ImportData.ObjectName);
+		FScriptObjectEntry Entry;
+		Entry.ObjectName = NameMapBuilder.MapName(ImportData.ObjectName).ToUnresolvedMinimalName();
+		Entry.GlobalIndex = ImportData.GlobalIndex;
+		Entry.OuterIndex = ImportData.OuterIndex;
+		Entry.CDOClassIndex = ImportData.CDOClassIndex;
+
+		ScriptObjectsArchive << Entry;
+	}
+
+	FIoWriteOptions WriteOptions;
+	WriteOptions.DebugName = TEXT("LoaderInitialLoadMeta");
+	IoStoreWriter->Append(CreateIoChunkId(0, 0, EIoChunkType::LoaderInitialLoadMeta), FIoBuffer(FIoBuffer::Wrap, ScriptObjectsArchive.GetData(), ScriptObjectsArchive.TotalSize()), WriteOptions);
+
+	TArray<uint8> Names;
+	TArray<uint8> Hashes;
+	SaveNameBatch(NameMapBuilder.GetNameMap(), /* out */ Names, /* out */ Hashes);
+	WriteOptions.DebugName = TEXT("LoaderGlobalNames");
+	IoStoreWriter->Append(CreateIoChunkId(0, 0, EIoChunkType::LoaderGlobalNames),
+		FIoBuffer(FIoBuffer::Wrap, Names.GetData(), Names.Num()), WriteOptions);
+	WriteOptions.DebugName = TEXT("LoaderGlobalNameHashes");
+	IoStoreWriter->Append(CreateIoChunkId(0, 0, EIoChunkType::LoaderGlobalNameHashes),
+		FIoBuffer(FIoBuffer::Wrap, Hashes.GetData(), Hashes.Num()), WriteOptions);
+
+	return ScriptObjectsArchive.TotalSize() + Names.Num() + Hashes.Num();
+}
+
+FPackageStoreContainerHeaderEntry FPackageStoreOptimizer::CreateContainerHeaderEntry(const FPackageStorePackage* Package) const
+{
+	FPackageStoreContainerHeaderEntry Result;
+	Result.Name = Package->Name;
+	Result.SourceName = Package->SourceName;
+	Result.Region = Package->Region;
+	Result.ExportBundlesSize = Package->OptimizedHeaderBuffer.DataSize() + Package->ExportsSerialSize;
+	Result.ExportCount = Package->Exports.Num();
+	Result.ExportBundleCount = Package->ExportBundles.Num();
+	Result.LoadOrder = Package->GetLoadOrder();
+	Result.ImportedPackageIds = Package->ImportedPackageIds;
+	Result.bIsRedirected = Package->bIsRedirected;
+	return Result;
+}
+
+FContainerHeader FPackageStoreOptimizer::CreateContainerHeader(const FIoContainerId& ContainerId, const TArray<FPackageStoreContainerHeaderEntry>& Packages) const
+{
+	FContainerHeader Header;
+	Header.ContainerId = ContainerId;
+	Header.PackageCount = Packages.Num();
+	int32 StoreTocSize = Header.PackageCount * sizeof(FPackageStoreEntry);
+	FLargeMemoryWriter StoreTocArchive(0, true);
+	FLargeMemoryWriter StoreDataArchive(0, true);
+
+	auto SerializePackageEntryCArrayHeader = [&StoreTocSize, &StoreTocArchive, &StoreDataArchive](int32 Count)
+	{
+		const int32 RemainingTocSize = StoreTocSize - StoreTocArchive.Tell();
+		const int32 OffsetFromThis = RemainingTocSize + StoreDataArchive.Tell();
+		uint32 ArrayNum = Count > 0 ? Count : 0;
+		uint32 OffsetToDataFromThis = ArrayNum > 0 ? OffsetFromThis : 0;
+
+		StoreTocArchive << ArrayNum;
+		StoreTocArchive << OffsetToDataFromThis;
+	};
+
+	TArray<const FPackageStoreContainerHeaderEntry*> SortedPackages;
+	SortedPackages.Reserve(Packages.Num());
+	for (const FPackageStoreContainerHeaderEntry& Package : Packages)
+	{
+		SortedPackages.Add(&Package);
+	}
+	Algo::Sort(SortedPackages, [](const FPackageStoreContainerHeaderEntry* A, const FPackageStoreContainerHeaderEntry* B)
+	{
+		return A->GetId() < B->GetId();
+	});
+
+	Header.PackageIds.Reserve(SortedPackages.Num());
+	for (const FPackageStoreContainerHeaderEntry* Package : SortedPackages)
+	{
+		Header.PackageIds.Add(Package->GetId());
+		if (Package->bIsRedirected)
+		{
+			if (Package->Region.Len() > 0)
+			{
+				Header.CulturePackageMap.FindOrAdd(Package->Region).Emplace(Package->GetSourceId(), Package->GetId());
+			}
+			else
+			{
+				Header.PackageRedirects.Add(MakeTuple(Package->GetSourceId(), Package->GetId()));
+			}
+		}
+		
+		// StoreEntries
+		uint64 ExportBundlesSize = Package->ExportBundlesSize;
+		int32 ExportCount = Package->ExportCount;
+		int32 ExportBundleCount = Package->ExportBundleCount;
+		uint32 LoadOrder = Package->LoadOrder;
+		uint32 Pad = 0;
+
+		StoreTocArchive << ExportBundlesSize;
+		StoreTocArchive << ExportCount;
+		StoreTocArchive << ExportBundleCount;
+		StoreTocArchive << LoadOrder;
+		StoreTocArchive << Pad;
+
+		// ImportedPackages
+		const TArray<FPackageId>& ImportedPackageIds = Package->ImportedPackageIds;
+		SerializePackageEntryCArrayHeader(ImportedPackageIds.Num());
+		for (FPackageId ImportedPackageId : ImportedPackageIds)
+		{
+			check(ImportedPackageId.IsValid());
+			StoreDataArchive << ImportedPackageId;
+		}
+	}
+
+	check(StoreTocArchive.TotalSize() == StoreTocSize);
+
+	const int32 StoreByteCount = StoreTocArchive.TotalSize() + StoreDataArchive.TotalSize();
+	Header.StoreEntries.AddUninitialized(StoreByteCount);
+	FBufferWriter PackageStoreArchive(Header.StoreEntries.GetData(), StoreByteCount);
+	PackageStoreArchive.Serialize(StoreTocArchive.GetData(), StoreTocArchive.TotalSize());
+	PackageStoreArchive.Serialize(StoreDataArchive.GetData(), StoreDataArchive.TotalSize());
+	FPackageStoreNameMapBuilder DummyNameMapBuilder;
+	SaveNameBatch(DummyNameMapBuilder.GetNameMap(), Header.Names, Header.NameHashes);
+
+	return Header;
+}
diff --git a/Engine/Source/Developer/IoStoreUtilities/Private/PackageStoreOptimizer.h b/Engine/Source/Developer/IoStoreUtilities/Private/PackageStoreOptimizer.h
new file mode 100644
--- /dev/null
+++ b/Engine/Source/Developer/IoStoreUtilities/Private/PackageStoreOptimizer.h
@@ -0,0 +1,380 @@
+// Copyright Epic Games, Inc. All Rights Reserved.
+
+#pragma once
+
+#include "Serialization/AsyncLoading2.h"
+#include "IO/IoDispatcher.h"
+#include "UObject/PackageFileSummary.h"
+#include "UObject/UObjectMarks.h"
+
+class FPackageStoreNameMapBuilder
+{
+public:
+	void SetNameMapType(FMappedName::EType InNameMapType)
+	{
+		NameMapType = InNameMapType;
+	}
+
+	void AddName(const FName& Name)
+	{
+		const FNameEntryId ComparisonIndex = Name.GetComparisonIndex();
+		const FNameEntryId DisplayIndex = Name.GetDisplayIndex();
+		NameMap.Add(DisplayIndex);
+		int32 Index = NameMap.Num();
+		NameIndices.Add(ComparisonIndex, Index);
+	}
+
+	void MarkNamesAsReferenced(const TArray<FName>& Names, TArray<int32>& OutNameIndices)
+	{
+		for (const FName& Name : Names)
+		{
+			const FNameEntryId ComparisonIndex = Name.GetComparisonIndex();
+			const FNameEntryId DisplayIndex = Name.GetDisplayIndex();
+			int32& Index = NameIndices.FindOrAdd(ComparisonIndex);
+			if (Index == 0)
+			{
+				NameMap.Add(DisplayIndex);
+				Index = NameMap.Num();
+			}
+
+			OutNameIndices.Add(Index - 1);
+		}
+	}
+
+	void MarkNameAsReferenced(const FName& Name)
+	{
+		const FNameEntryId ComparisonIndex = Name.GetComparisonIndex();
+		const FNameEntryId DisplayIndex = Name.GetDisplayIndex();
+		int32& Index = NameIndices.FindOrAdd(ComparisonIndex);
+		if (Index == 0)
+		{
+			NameMap.Add(DisplayIndex);
+			Index = NameMap.Num();
+		}
+	}
+
+	FMappedName MapName(const FName& Name) const
+	{
+		const FNameEntryId Id = Name.GetComparisonIndex();
+		const int32* Index = NameIndices.Find(Id);
+		check(Index);
+		return FMappedName::Create(*Index - 1, Name.GetNumber(), NameMapType);
+	}
+
+	const TArray<FNameEntryId>& GetNameMap() const
+	{
+		return NameMap;
+	}
+
+	void Empty()
+	{
+		NameIndices.Empty();
+		NameMap.Empty();
+	}
+
+private:
+	TMap<FNameEntryId, int32> NameIndices;
+	TArray<FNameEntryId> NameMap;
+	FMappedName::EType NameMapType = FMappedName::EType::Package;
+};
+
+class FPackageStoreContainerHeaderEntry
+{
+public:
+	FPackageId GetId() const
+	{
+		return FPackageId::FromName(Name);
+	}
+
+	FPackageId GetSourceId() const
+	{
+		if (SourceName.IsNone())
+		{
+			return FPackageId();
+		}
+		return FPackageId::FromName(SourceName);
+	}
+
+private:
+	FName Name;
+	FName SourceName;
+	FString Region;
+	uint64 ExportBundlesSize = 0;
+	int32 ExportCount = 0;
+	int32 ExportBundleCount = 0;
+	uint32 LoadOrder = 0;
+	TArray<FPackageId> ImportedPackageIds;
+	bool bIsRedirected = false;
+
+	friend class FPackageStoreOptimizer;
+};
+
+class FPackageStorePackage
+{
+public:
+	FPackageId GetId() const
+	{
+		return Id;
+	}
+
+	uint64 GetLoadOrder() const
+	{
+		return ExportBundles.Num() > 0 ? ExportBundles[0].LoadOrder : 0;
+	}
+
+	uint64 GetExportBundleCount() const
+	{
+		return ExportBundles.Num();
+	}
+
+	uint64 GetImportMapSize() const
+	{
+		return ImportMapSize;
+	}
+
+	uint64 GetExportMapSize() const
+	{
+		return ExportMapSize;
+	}
+	
+	uint64 GetExportBundleArcsSize() const
+	{
+		return ExportBundleArcsSize;
+	}
+
+	uint64 GetExportBundleHeadersSize() const
+	{
+		return ExportBundleHeadersSize;
+	}
+
+	uint64 GetNameCount() const
+	{
+		return NameMapBuilder.GetNameMap().Num();
+	}
+
+	uint64 GetNameMapSize() const
+	{
+		return NameMapSize;
+	}
+
+	const TArray<FPackageId>& GetImportedPackageIds() const
+	{
+		return ImportedPackageIds;
+	}
+
+	const TSet<FPackageId>& GetImportedRedirectedPackageIds() const
+	{
+		return ImportedRedirectedPackageIds;
+	}
+
+	void RedirectFrom(FName SourcePackageName)
+	{
+		SourceName = SourcePackageName;
+	}
+
+private:
+	struct FExportGraphNode;
+
+	struct FExternalDependency
+	{
+		FPackageId PackageId;
+		FPackageObjectIndex GlobalImportIndex;
+		FName DebugImportFullName;
+		FExportBundleEntry::EExportCommandType ExportBundleCommandType;
+		FExportGraphNode* ExportGraphNode = nullptr;
+		int32 ResolvedExportBundleIndex = -1;
+		bool bIsConfirmedMissing = false;
+	};
+
+	struct FExportGraphNode
+	{
+		FPackageStorePackage* Package = nullptr;
+		FExportBundleEntry BundleEntry;
+		TArray<FExternalDependency> ExternalDependencies;
+		int32 ExportBundleIndex = -1;
+		int32 IncomingEdgeCount = 0;
+	};
+
+	struct FExport
+	{
+		FString FullName;
+		FName ObjectName;
+		FPackageObjectIndex GlobalImportIndex;
+		FPackageObjectIndex OuterIndex;
+		FPackageObjectIndex ClassIndex;
+		FPackageObjectIndex SuperIndex;
+		FPackageObjectIndex TemplateIndex;
+		EObjectFlags ObjectFlags = RF_NoFlags;
+		uint64 CookedSerialOffset = uint64(-1);
+		uint64 CookedSerialSize = uint64(-1);
+		bool bNotForClient = false;
+		bool bNotForServer = false;
+		bool bIsPublic = false;
+		FExportGraphNode* CreateNode = nullptr;
+		FExportGraphNode* SerializeNode = nullptr;
+	};
+
+	struct FExportBundle
+	{
+		TArray<FExportBundleEntry> Entries;
+		uint32 LoadOrder = MAX_int32;
+	};
+
+	struct FImport
+	{
+		FString FullName;
+		FPackageId PackageId;
+		FPackageObjectIndex GlobalImportIndex;
+		bool bIsScriptImport = false;
+		bool bIsPackageImport = false;
+	};
+
+	struct FArc
+	{
+		uint32 FromNodeIndex;
+		uint32 ToNodeIndex;
+
+		bool operator==(const FArc& Other) const
+		{
+			return ToNodeIndex == Other.ToNodeIndex && FromNodeIndex == Other.FromNodeIndex;
+		}
+	};
+
+	FPackageId Id;
+	FName Name;
+	FName SourceName;
+	FString Region;
+
+	FPackageFileSummary Summary;
+	TArray<FObjectImport> ObjectImports;
+	TArray<FObjectExport> ObjectExports;
+	TArray<FPackageIndex> PreloadDependencies;
+
+	FPackageStoreNameMapBuilder NameMapBuilder;
+	TArray<FImport> Imports;
+	TArray<FExport> Exports;
+	TArray<FExportGraphNode> ExportGraphNodes;
+	TMap<FPackageId, TArray<FArc>> ExternalArcs;
+	TArray<FExportBundle> ExportBundles;
+
+	TArray<FPackageId> ImportedPackageIds;
+	TArray<FPackageStorePackage*> ImportedWaitingPackages;
+	TArray<FPackageStorePackage*> ImportedByWaitingPackages;
+	TSet<FPackageId> RedirectedToPackageIds;
+	TSet<FPackageId> ImportedRedirectedPackageIds;
+
+	FIoBuffer OptimizedHeaderBuffer;
+
+	uint32 PackageFlags = 0;
+	uint32 CookedHeaderSize = 0;
+	uint64 ExportsSerialSize = 0;
+	uint64 ImportMapSize = 0;
+	uint64 ExportMapSize = 0;
+	uint64 ExportBundleArcsSize = 0;
+	uint64 ExportBundleHeadersSize = 0;
+	uint64 NameMapSize = 0;
+	
+	uint32 UnresolvedImportedPackagesCount = 0;
+	TArray<FExportGraphNode*> NodesWithNoIncomingEdges;
+	bool bTemporaryMark = false;
+	bool bPermanentMark = false;
+
+	bool bIsRedirected = false;
+
+	friend class FPackageStoreOptimizer;
+};
+
+class FPackageStoreOptimizer
+{
+public:
+	uint64 GetTotalPackageCount() const
+	{
+		return TotalPackageCount;
+	}
+
+	uint64 GetTotalExportBundleCount() const
+	{
+		return TotalExportBundleCount;
+	}
+
+	uint64 GetTotalExportBundleEntryCount() const
+	{
+		return TotalExportBundleEntryCount;
+	}
+
+	uint64 GetTotalExportBundleArcCount() const
+	{
+		return TotalExportBundleArcCount;
+	}
+
+	uint64 GetTotalScriptObjectCount() const
+	{
+		return TotalScriptObjectCount;
+	}
+
+	void Initialize(const ITargetPlatform* TargetPlatform);
+
+	void EnableIncrementalResolve()
+	{
+		bIncrementalResolveEnabled = true;
+	}
+
+	FPackageStorePackage* CreatePackageFromCookedHeader(const FName& Name, const FIoBuffer& CookedHeaderBuffer) const;
+	void BeginResolvePackage(FPackageStorePackage* Package);
+	void Flush(bool bAllowMissingImports, TFunction<void(FPackageStorePackage*)> ResolvedPackageCallback = TFunction<void(FPackageStorePackage*)>());
+	FIoBuffer CreatePackageBuffer(const FPackageStorePackage* Package, const FIoBuffer& CookedExportsBuffer, TArray<FFileRegion>* InOutFileRegions) const;
+	FPackageStoreContainerHeaderEntry CreateContainerHeaderEntry(const FPackageStorePackage* Package) const;
+	FContainerHeader CreateContainerHeader(const FIoContainerId& ContainerId, const TArray<FPackageStoreContainerHeaderEntry>& Packages) const;
+	uint64 WriteScriptObjects(FIoStoreWriter* IoStoreWriter) const; // TODO: Merge name map and return FIoBuffer instead
+
+private:
+	struct FScriptObjectData
+	{
+		FName ObjectName;
+		FString FullName;
+		FPackageObjectIndex GlobalIndex;
+		FPackageObjectIndex OuterIndex;
+		FPackageObjectIndex CDOClassIndex;
+	};
+
+	struct FResolvedPublicExport
+	{
+		FPackageObjectIndex GlobalImportIndex;
+		int32 CreateExportBundleIndex = -1;
+		int32 SerializeExportBundleIndex = -1;
+	};
+
+	using FExportGraphEdges = TMultiMap<FPackageStorePackage::FExportGraphNode*, FPackageStorePackage::FExportGraphNode*>;
+
+	void LoadCookedHeader(FPackageStorePackage* Package, const FIoBuffer& CookedHeaderBuffer) const;
+	void ResolveImport(FPackageStorePackage::FImport* Imports, const FObjectImport* ObjectImports, int32 LocalImportIndex) const;
+	void ProcessImports(FPackageStorePackage* Package) const;
+	void ResolveExport(FPackageStorePackage::FExport* Exports, const FObjectExport* ObjectExports, const int32 LocalExportIndex, const FName& PackageName) const;
+	void ProcessExports(FPackageStorePackage* Package) const;
+	void SortPackagesInLoadOrder(TArray<FPackageStorePackage*>& Packages) const;
+	FExportGraphEdges ProcessPreloadDependencies(const TArray<FPackageStorePackage*>& Packages) const;
+	TArray<FPackageStorePackage::FExportGraphNode*> SortExportGraphNodesInLoadOrder(const TArray<FPackageStorePackage*>& Packages, FExportGraphEdges& Edges) const;
+	void CreateExportBundles(TArray<FPackageStorePackage*>& Packages);
+	bool RedirectPackage(const FPackageStorePackage& SourcePackage, FPackageStorePackage& TargetPackage, TMap<FPackageObjectIndex, FPackageObjectIndex>& RedirectedToSourceImportIndexMap) const;
+	void RedirectPackageUnverified(FPackageStorePackage& TargetPackage, TMap<FPackageObjectIndex, FPackageObjectIndex>& RedirectedToSourceImportIndexMap) const;
+	void ProcessRedirects(const TArray<FPackageStorePackage*> Packages) const;
+	void ResolvePackages(TArray<FPackageStorePackage*>& Packages, bool bAllowMissingImports);
+	void FinalizePackageHeader(FPackageStorePackage* Package) const;
+	void FinalizePackage(FPackageStorePackage* Package, bool bAllowMissingImports);
+	bool IsWaitingForImportsRecursive(FPackageStorePackage* Package) const;
+	void FindScriptObjectsRecursive(FPackageObjectIndex OuterIndex, UObject* Object, EObjectMark ExcludedObjectMarks, const ITargetPlatform* TargetPlatform);
+	void FindScriptObjects(const ITargetPlatform* TargetPlatform);
+
+	TMap<FPackageObjectIndex, FScriptObjectData> ScriptObjectsMap;
+	TMap<FPackageId, TArray<FResolvedPublicExport>> ResolvedPublicExportsMap;
+	TMap<FPackageId, FPackageStorePackage*> WaitingPackagesMap;
+	TMultiMap<FPackageId, FPackageStorePackage*> ImportedPackageToWaitingPackagesMap;
+	TArray<FPackageStorePackage*> ResolvedPackages;
+	uint64 TotalPackageCount = 0;
+	uint64 TotalExportBundleCount = 0;
+	uint64 TotalExportBundleEntryCount = 0;
+	uint64 TotalExportBundleArcCount = 0;
+	uint64 TotalScriptObjectCount = 0;
+	uint32 NextLoadOrder = 0;
+	bool bIncrementalResolveEnabled = false;
+};
\ No newline at end of file
