Fix static analysis warnings
#rnx

diff --git a/Engine/Source/Developer/IoStoreUtilities/Private/IoStoreUtilities.cpp b/Engine/Source/Developer/IoStoreUtilities/Private/IoStoreUtilities.cpp
--- a/Engine/Source/Developer/IoStoreUtilities/Private/IoStoreUtilities.cpp
+++ b/Engine/Source/Developer/IoStoreUtilities/Private/IoStoreUtilities.cpp
@@ -1633,6 +1633,7 @@ int32 CreateTarget(const FIoStoreArguments& Arguments, const FIoStoreWriterSetti
 			TArray<FPackageStoreContainerHeaderEntry> ContainerHeaderEntries;
 			for (FContainerTargetFile& TargetFile : ContainerTarget->TargetFiles)
 			{
+				check(TargetFile.Package);
 				if (!TargetFile.bIsBulkData)
 				{
 					FIoWriteOptions WriteOptions;
diff --git a/Engine/Source/Developer/IoStoreUtilities/Private/PackageStoreOptimizer.cpp b/Engine/Source/Developer/IoStoreUtilities/Private/PackageStoreOptimizer.cpp
--- a/Engine/Source/Developer/IoStoreUtilities/Private/PackageStoreOptimizer.cpp
+++ b/Engine/Source/Developer/IoStoreUtilities/Private/PackageStoreOptimizer.cpp
@@ -155,7 +155,8 @@ void FPackageStoreOptimizer::LoadCookedHeader(FPackageStorePackage* Package, con
 
 		FArchive& operator<<(FName& Name)
 		{
-			int32 NameIndex, Number;
+			int32 NameIndex = 0;
+			int32 Number = 0;
 			InnerArchive << NameIndex << Number;
 
 			if (!NameMap.IsValidIndex(NameIndex))
@@ -657,7 +658,12 @@ void FPackageStoreOptimizer::CreateExportBundles(TArray<FPackageStorePackage*>&
 	FPackageStorePackage* LastPackage = nullptr;
 	for (FPackageStorePackage::FExportGraphNode* Node : LoadOrder)
 	{
-		check(Node->Package);
+		check(Node && Node->Package);
+		if (!Node || !Node->Package)
+		{
+			// Needed to please static analysis
+			continue;
+		}
 		FPackageStorePackage::FExportBundle* Bundle;
 		if (Node->Package != LastPackage)
 		{
