Added option to cook for iostore, allowing the new loader to be used with cooked data without having to stage first
#rnx

diff --git a/Engine/Source/Developer/IoStoreUtilities/Private/PackageStoreWriter.cpp b/Engine/Source/Developer/IoStoreUtilities/Private/PackageStoreWriter.cpp
new file mode 100644
--- /dev/null
+++ b/Engine/Source/Developer/IoStoreUtilities/Private/PackageStoreWriter.cpp
@@ -0,0 +1,115 @@
+// Copyright Epic Games, Inc. All Rights Reserved.
+
+#include "PackageStoreWriter.h"
+#include "PackageStoreOptimizer.h"
+#include "Containers/Map.h"
+#include "Serialization/LargeMemoryWriter.h"
+
+DEFINE_LOG_CATEGORY_STATIC(LogPackageStoreWriter, Log, All);
+
+FPackageStoreWriter::FPackageStoreWriter(const FString& InOutputPath, const ITargetPlatform* InTargetPlatform)
+	: PackageStoreOptimizer(new FPackageStoreOptimizer())
+	, IoStoreWriterContext(new FIoStoreWriterContext())
+	, IoStoreWriter(new FIoStoreWriter(*(InOutputPath / TEXT("global"))))
+
+{
+	FIoStoreWriterSettings WriterSettings;
+	WriterSettings.CompressionMethod = NAME_Zlib;
+	FIoStatus IoStatus = IoStoreWriterContext->Initialize(WriterSettings);
+	check(IoStatus.IsOk());
+	FIoContainerSettings ContainerSettings;
+	ContainerSettings.ContainerId = ContainerId;
+	ContainerSettings.ContainerFlags |= EIoContainerFlags::Compressed;
+	ContainerSettings.ContainerFlags |= EIoContainerFlags::Indexed;
+	IoStatus = IoStoreWriter->Initialize(*IoStoreWriterContext, ContainerSettings);
+	check(IoStatus.IsOk());
+	PackageStoreOptimizer->Initialize(InTargetPlatform);
+	PackageStoreOptimizer->EnableIncrementalResolve();
+}
+
+FPackageStoreWriter::~FPackageStoreWriter()
+{
+}
+
+void FPackageStoreWriter::Flush(bool bAllowMissingImports)
+{
+	PackageStoreOptimizer->Flush(bAllowMissingImports, [this](FPackageStorePackage* ResolvedPackage)
+	{
+		FPendingPackage* FindPendingPackage = PendingPackagesMap.Find(ResolvedPackage);
+		check(FindPendingPackage);
+		FIoBuffer PackageBuffer = PackageStoreOptimizer->CreatePackageBuffer(ResolvedPackage, FindPendingPackage->CookedExportsBuffer, &FindPendingPackage->FileRegions);
+		FIoWriteOptions WriteOptions;
+		WriteOptions.FileName = FindPendingPackage->FileName;
+		FIoChunkId ChunkId = CreateIoChunkId(ResolvedPackage->GetId().Value(), 0, EIoChunkType::ExportBundleData);
+		IoStoreWriter->Append(ChunkId, PackageBuffer, WriteOptions);
+		PendingPackagesMap.Remove(ResolvedPackage);
+		CompletedPackages.Add(PackageStoreOptimizer->CreateContainerHeaderEntry(ResolvedPackage));
+		delete ResolvedPackage;
+	});
+}
+
+void FPackageStoreWriter::WritePackage(const FPackageInfo& Info, const FIoBuffer& PackageData, const TArray<FFileRegion>& FileRegions)
+{
+	TRACE_CPUPROFILER_EVENT_SCOPE(WritePackage);
+	FIoBuffer HeaderBuffer = FIoBuffer(PackageData.Data(), Info.HeaderSize, PackageData);
+	FPackageStorePackage* Package = PackageStoreOptimizer->CreatePackageFromCookedHeader(Info.PackageName, HeaderBuffer);
+
+	FPendingPackage PendingPackage = {
+		FIoBuffer(PackageData.Data() + Info.HeaderSize, PackageData.DataSize() - Info.HeaderSize, PackageData),
+		Info.LooseFilePath,
+		FileRegions
+	};
+
+	PendingPackagesMap.Emplace(Package, PendingPackage);
+	PackageStoreOptimizer->BeginResolvePackage(Package);
+	Flush(false);
+}
+
+void FPackageStoreWriter::WriteBulkdata(const FBulkDataInfo& Info, const FIoBuffer& BulkData, const TArray<FFileRegion>& FileRegions)
+{
+	TRACE_CPUPROFILER_EVENT_SCOPE(WriteBulkdata);
+	FIoWriteOptions WriteOptions;
+	WriteOptions.FileName = Info.LooseFilePath;
+	FPackageId PackageId = FPackageId::FromName(Info.PackageName);
+	FIoChunkId ChunkId;
+	if (Info.BulkdataType == FBulkDataInfo::Optional)
+	{
+		ChunkId = CreateIoChunkId(PackageId.Value(), 0, EIoChunkType::OptionalBulkData);
+	}
+	else if (Info.BulkdataType == FBulkDataInfo::Mmap)
+	{
+		WriteOptions.bForceUncompressed = true;
+		WriteOptions.bIsMemoryMapped = true;
+		ChunkId = CreateIoChunkId(PackageId.Value(), 0, EIoChunkType::MemoryMappedBulkData);
+	}
+	else
+	{
+		ChunkId = CreateIoChunkId(PackageId.Value(), 0, EIoChunkType::BulkData);
+	}
+	IoStoreWriter->Append(ChunkId, BulkData, WriteOptions);
+}
+
+void FPackageStoreWriter::Finalize()
+{
+	Flush(true);
+
+	PackageStoreOptimizer->WriteScriptObjects(IoStoreWriter.Get());
+
+	FContainerHeader Header = PackageStoreOptimizer->CreateContainerHeader(ContainerId, CompletedPackages);
+	FLargeMemoryWriter HeaderAr(0, true);
+	HeaderAr << Header;
+	int64 DataSize = HeaderAr.TotalSize();
+	FIoBuffer HeaderBuffer(FIoBuffer::AssumeOwnership, HeaderAr.ReleaseOwnership(), DataSize);
+	FIoWriteOptions WriteOptions;
+	WriteOptions.DebugName = TEXT("ContainerHeader");
+	IoStoreWriter->Append(CreateIoChunkId(ContainerId.Value(), 0, EIoChunkType::ContainerHeader), HeaderBuffer, WriteOptions);
+
+	TIoStatusOr<FIoStoreWriterResult> Result = IoStoreWriter->Flush();
+	check(Result.IsOk());
+
+	UE_LOG(LogPackageStoreWriter, Display, TEXT("Input:\t%d Packages"), PackageStoreOptimizer->GetTotalPackageCount());
+	UE_LOG(LogPackageStoreWriter, Display, TEXT("Output:\t%d Export bundles"), PackageStoreOptimizer->GetTotalExportBundleCount());
+	UE_LOG(LogPackageStoreWriter, Display, TEXT("Output:\t%d Export bundle entries"), PackageStoreOptimizer->GetTotalExportBundleEntryCount());
+	UE_LOG(LogPackageStoreWriter, Display, TEXT("Output:\t%d Export bundle arcs"), PackageStoreOptimizer->GetTotalExportBundleArcCount());
+	UE_LOG(LogPackageStoreWriter, Display, TEXT("Output:\t%d Public runtime script objects"), PackageStoreOptimizer->GetTotalScriptObjectCount());
+}
diff --git a/Engine/Source/Developer/IoStoreUtilities/Public/PackageStoreWriter.h b/Engine/Source/Developer/IoStoreUtilities/Public/PackageStoreWriter.h
new file mode 100644
--- /dev/null
+++ b/Engine/Source/Developer/IoStoreUtilities/Public/PackageStoreWriter.h
@@ -0,0 +1,41 @@
+// Copyright Epic Games, Inc. All Rights Reserved.
+
+#pragma once
+
+#include "UObject/SavePackage.h"
+#include "Serialization/AsyncLoading2.h"
+#include "IO/IoDispatcher.h"
+
+class FIoStoreWriterContext;
+class FIoStoreWriter;
+class FPackageStoreOptimizer;
+class FPackageStorePackage;
+class FPackageStoreContainerHeaderEntry;
+
+class FPackageStoreWriter
+	: public IPackageStoreWriter
+{
+public:
+	IOSTOREUTILITIES_API FPackageStoreWriter(const FString& OutputPath, const ITargetPlatform* TargetPlatform);
+	IOSTOREUTILITIES_API ~FPackageStoreWriter();
+	IOSTOREUTILITIES_API void WritePackage(const FPackageInfo& Info, const FIoBuffer& PackageData, const TArray<FFileRegion>& FileRegions) override;
+	IOSTOREUTILITIES_API void WriteBulkdata(const FBulkDataInfo& Info, const FIoBuffer& BulkData, const TArray<FFileRegion>& FileRegions) override;
+	IOSTOREUTILITIES_API void Finalize() override;
+
+private:
+	struct FPendingPackage
+	{
+		FIoBuffer CookedExportsBuffer;
+		FString FileName;
+		TArray<FFileRegion> FileRegions;
+	};
+
+	void Flush(bool bAllowMissingImports);
+
+	TUniquePtr<FPackageStoreOptimizer> PackageStoreOptimizer;
+	TUniquePtr<FIoStoreWriterContext> IoStoreWriterContext;
+	TUniquePtr<FIoStoreWriter> IoStoreWriter;
+	FIoContainerId ContainerId = FIoContainerId::FromName(TEXT("global"));
+	TMap<FPackageStorePackage*, FPendingPackage> PendingPackagesMap;
+	TArray<FPackageStoreContainerHeaderEntry> CompletedPackages;
+};
\ No newline at end of file
diff --git a/Engine/Source/Editor/UnrealEd/Classes/CookOnTheSide/CookOnTheFlyServer.h b/Engine/Source/Editor/UnrealEd/Classes/CookOnTheSide/CookOnTheFlyServer.h
--- a/Engine/Source/Editor/UnrealEd/Classes/CookOnTheSide/CookOnTheFlyServer.h
+++ b/Engine/Source/Editor/UnrealEd/Classes/CookOnTheSide/CookOnTheFlyServer.h
@@ -70,7 +70,7 @@ enum class ECookByTheBookOptions
 	SkipSoftReferences =				0x00001000, // Don't follow soft references when cooking. Usually not viable for a real cook and the results probably wont load properly, but can be useful for debugging
 	SkipHardReferences =				0x00002000, // Don't follow hard references when cooking. Not viable for a real cook, only useful for debugging
 	FullLoadAndSave =					0x00004000, // Load all packages into memory and save them all at once in one tick for speed reasons. This requires a lot of RAM for large games.
-	PackageStore =						0x00008000, // Cook package header information into a global package store
+	IoStore =							0x00008000, // Cook for IoStore
 	CookAgainstFixedBase =				0x00010000, // If cooking DLC, assume that the base content can not be modified. 
 	DlcLoadMainAssetRegistry =			0x00020000, // If cooking DLC, populate the main game asset registry
 
@@ -543,7 +543,7 @@ public:
 
 	bool IsUsingShaderCodeLibrary() const;
 
-	bool IsUsingPackageStore() const;
+	bool IsUsingIoStore() const;
 
 	/**
 	* Helper function returns if we are in any cook on the fly mode
diff --git a/Engine/Source/Editor/UnrealEd/Private/Commandlets/CookCommandlet.cpp b/Engine/Source/Editor/UnrealEd/Private/Commandlets/CookCommandlet.cpp
--- a/Engine/Source/Editor/UnrealEd/Private/Commandlets/CookCommandlet.cpp
+++ b/Engine/Source/Editor/UnrealEd/Private/Commandlets/CookCommandlet.cpp
@@ -898,7 +898,7 @@ bool UCookCommandlet::CookByTheBook( const TArray<ITargetPlatform*>& Platforms)
 	CookOptions |= Switches.Contains(TEXT("MAPSONLY")) ? ECookByTheBookOptions::MapsOnly : ECookByTheBookOptions::None;
 	CookOptions |= Switches.Contains(TEXT("NODEV")) ? ECookByTheBookOptions::NoDevContent : ECookByTheBookOptions::None;
 	CookOptions |= Switches.Contains(TEXT("FullLoadAndSave")) ? ECookByTheBookOptions::FullLoadAndSave : ECookByTheBookOptions::None;
-	CookOptions |= Switches.Contains(TEXT("PackageStore")) ? ECookByTheBookOptions::PackageStore : ECookByTheBookOptions::None;
+	CookOptions |= Switches.Contains(TEXT("IoStore")) ? ECookByTheBookOptions::IoStore : ECookByTheBookOptions::None;
 	CookOptions |= Switches.Contains(TEXT("NoGameAlwaysCook")) ? ECookByTheBookOptions::NoGameAlwaysCookPackages : ECookByTheBookOptions::None;
 	CookOptions |= Switches.Contains(TEXT("DisableUnsolicitedPackages")) ? (ECookByTheBookOptions::SkipHardReferences | ECookByTheBookOptions::SkipSoftReferences) : ECookByTheBookOptions::None;
 	CookOptions |= Switches.Contains(TEXT("NoDefaultMaps")) ? ECookByTheBookOptions::NoDefaultMaps : ECookByTheBookOptions::None;
diff --git a/Engine/Source/Editor/UnrealEd/Private/CookOnTheFlyServer.cpp b/Engine/Source/Editor/UnrealEd/Private/CookOnTheFlyServer.cpp
--- a/Engine/Source/Editor/UnrealEd/Private/CookOnTheFlyServer.cpp
+++ b/Engine/Source/Editor/UnrealEd/Private/CookOnTheFlyServer.cpp
@@ -119,6 +119,7 @@
 #include "Async/ParallelFor.h"
 
 #include "Commandlets/ShaderPipelineCacheToolsCommandlet.h"
+#include "PackageStoreWriter.h"
 
 #define LOCTEXT_NAMESPACE "Cooker"
 #define REMAPPED_PLUGINS TEXT("RemappedPlugins")
@@ -420,7 +421,7 @@ public:
 	bool							bSkipHardReferences = false;
 	bool							bSkipSoftReferences = false;
 	bool							bFullLoadAndSave = false;
-	bool							bPackageStore = false;
+	bool							bIoStore = false;
 	bool							bCookAgainstFixedBase = false;
 	bool							bDlcLoadMainAssetRegistry = false;
 	TArray<FName>					StartupPackages;
@@ -1013,9 +1014,9 @@ bool UCookOnTheFlyServer::IsUsingShaderCodeLibrary() const
 	return IsCookByTheBookMode();
 }
 
-bool UCookOnTheFlyServer::IsUsingPackageStore() const
+bool UCookOnTheFlyServer::IsUsingIoStore() const
 {
-	return IsCookByTheBookMode() && CookByTheBookOptions->bPackageStore;
+	return IsCookByTheBookMode() && CookByTheBookOptions->bIoStore;
 }
 
 bool UCookOnTheFlyServer::IsCookOnTheFlyMode() const
@@ -7004,7 +7005,11 @@ void UCookOnTheFlyServer::InitializePackageStore(const TArrayView<const ITargetP
 		const FString ResolvedProjectPath = ProjectPathSandbox.Replace(TEXT("[Platform]"), *PlatformString);
 
 		FPackageStoreBulkDataManifest* BulkDataManifest	= bIsDiffOnly == false ? new FPackageStoreBulkDataManifest(ResolvedProjectPath) : nullptr;
-		FLooseFileWriter* LooseFileWriter				= IsUsingPackageStore() ? new FLooseFileWriter() : nullptr;
+		FPackageStoreWriter* PackageStoreWriter = nullptr;
+		if (IsUsingIoStore())
+		{
+			PackageStoreWriter = new FPackageStoreWriter(ResolvedRootPath, TargetPlatform);
+		}
 
 		FConfigFile PlatformEngineIni;
 		FConfigCacheIni::LoadLocalIniFile(PlatformEngineIni, TEXT("Engine"), true, *TargetPlatform->IniPlatformName());
@@ -7012,7 +7017,7 @@ void UCookOnTheFlyServer::InitializePackageStore(const TArrayView<const ITargetP
 		bool bLegacyBulkDataOffsets = false;
 		PlatformEngineIni.GetBool(TEXT("Core.System"), TEXT("LegacyBulkDataOffsets"), bLegacyBulkDataOffsets);
 
-		FSavePackageContext* SavePackageContext			= new FSavePackageContext(LooseFileWriter, BulkDataManifest, bLegacyBulkDataOffsets);
+		FSavePackageContext* SavePackageContext			= new FSavePackageContext(PackageStoreWriter, BulkDataManifest, bLegacyBulkDataOffsets);
 		SavePackageContexts.Add(SavePackageContext);
 	}
 }
@@ -7021,15 +7026,22 @@ void UCookOnTheFlyServer::FinalizePackageStore()
 {
 	UE_SCOPED_HIERARCHICAL_COOKTIMER(FinalizePackageStore);
 
-	UE_LOG(LogCook, Display, TEXT("Saving BulkData manifest(s)..."));
+	UE_LOG(LogCook, Display, TEXT("Finalize package store(s)..."));
 	for (FSavePackageContext* PackageContext : SavePackageContexts)
 	{
-		if (PackageContext != nullptr && PackageContext->BulkDataManifest != nullptr)
+		if (PackageContext != nullptr)
 		{
-			PackageContext->BulkDataManifest->Save();
+			if (PackageContext->BulkDataManifest != nullptr)
+			{
+				PackageContext->BulkDataManifest->Save();
+			}
+			if (PackageContext->PackageStoreWriter != nullptr)
+			{
+				PackageContext->PackageStoreWriter->Finalize();
+			}
 		}
 	}
-	UE_LOG(LogCook, Display, TEXT("Done saving BulkData manifest(s)"));
+	UE_LOG(LogCook, Display, TEXT("Done finalizing package store(s)"));
 
 	ClearPackageStoreContexts();
 }
@@ -7150,7 +7162,7 @@ void UCookOnTheFlyServer::StartCookByTheBook( const FCookByTheBookStartupOptions
 	CookByTheBookOptions->bSkipHardReferences = !!(CookOptions & ECookByTheBookOptions::SkipHardReferences);
 	CookByTheBookOptions->bSkipSoftReferences = !!(CookOptions & ECookByTheBookOptions::SkipSoftReferences);
 	CookByTheBookOptions->bFullLoadAndSave = !!(CookOptions & ECookByTheBookOptions::FullLoadAndSave);
-	CookByTheBookOptions->bPackageStore = !!(CookOptions & ECookByTheBookOptions::PackageStore);
+	CookByTheBookOptions->bIoStore = !!(CookOptions & ECookByTheBookOptions::IoStore);
 	CookByTheBookOptions->bCookAgainstFixedBase = !!(CookOptions & ECookByTheBookOptions::CookAgainstFixedBase);
 	CookByTheBookOptions->bDlcLoadMainAssetRegistry = !!(CookOptions & ECookByTheBookOptions::DlcLoadMainAssetRegistry);
 	CookByTheBookOptions->bErrorOnEngineContentUse = CookByTheBookStartupOptions.bErrorOnEngineContentUse;
diff --git a/Engine/Source/Runtime/CoreUObject/Private/UObject/SavePackage.cpp b/Engine/Source/Runtime/CoreUObject/Private/UObject/SavePackage.cpp
--- a/Engine/Source/Runtime/CoreUObject/Private/UObject/SavePackage.cpp
+++ b/Engine/Source/Runtime/CoreUObject/Private/UObject/SavePackage.cpp
@@ -4458,28 +4458,12 @@ FSavePackageResultStruct UPackage::Save(UPackage* InOuter, UObject* Base, EObjec
 
 									const int32 HeaderSize = Linker->Summary.TotalHeaderSize;
 
-									FPackageStoreWriter::HeaderInfo HeaderInfo;
-									HeaderInfo.PackageName		= InOuter->GetFName();
-									HeaderInfo.LooseFilePath	= Filename;
+									IPackageStoreWriter::FPackageInfo PackageInfo;
+									PackageInfo.PackageName		= InOuter->GetFName();
+									PackageInfo.LooseFilePath	= Filename;
+									PackageInfo.HeaderSize		= HeaderSize;
 
-									SavePackageContext->PackageStoreWriter->WriteHeader(HeaderInfo, FIoBuffer(IoBuffer.Data(), HeaderSize, IoBuffer));
-
-									FPackageStoreWriter::ExportsInfo ExportsInfo;
-									ExportsInfo.LooseFilePath	= Filename;
-									ExportsInfo.PackageName		= InOuter->GetFName();
-									ExportsInfo.RegionsOffset   = HeaderSize;
-
-									const uint8* ExportsData = IoBuffer.Data() + HeaderSize;
-									const int32 ExportCount = Linker->ExportMap.Num();
-
-									ExportsInfo.Exports.Reserve(ExportCount);
-
-									for (const FObjectExport& Export : Linker->ExportMap)
-									{
-										ExportsInfo.Exports.Add(FIoBuffer(IoBuffer.Data() + Export.SerialOffset, Export.SerialSize, IoBuffer));
-									}
-
-									SavePackageContext->PackageStoreWriter->WriteExports(ExportsInfo, FIoBuffer(ExportsData, DataSize - HeaderSize, IoBuffer), Linker->FileRegions);
+									SavePackageContext->PackageStoreWriter->WritePackage(PackageInfo, IoBuffer, Linker->FileRegions);
 								}
 								else
 								{
@@ -4705,78 +4689,7 @@ bool UPackage::SavePackage(UPackage* InOuter, UObject* Base, EObjectFlags TopLev
 // TODO: this should go elsewhere, this file is big enough as it is already
 //
 
-FPackageStoreWriter::FPackageStoreWriter() = default;
-FPackageStoreWriter::~FPackageStoreWriter() = default;
-
-FLooseFileWriter::FLooseFileWriter() = default;
-FLooseFileWriter::~FLooseFileWriter() = default;
-
-void FLooseFileWriter::WriteHeader(const FLooseFileWriter::HeaderInfo& Info, const FIoBuffer& HeaderData)
-{
-	SavePackageUtilities::WriteToFile(Info.LooseFilePath, HeaderData.Data(), HeaderData.DataSize());
-}
-
-void FLooseFileWriter::WriteExports(const FLooseFileWriter::ExportsInfo& Info, const FIoBuffer& ExportsData, const TArray<FFileRegion>& FileRegions)
-{
-	const FString ArchiveFilename = FPaths::ChangeExtension(Info.LooseFilePath, TEXT(".uexp"));
-
-	SavePackageUtilities::WriteToFile(ArchiveFilename, ExportsData.Data(), ExportsData.DataSize());
-
-	if (FileRegions.Num() > 0)
-	{
-		// Adjust regions so they are relative to the start of the uexp file
-		TArray<FFileRegion> FileRegionsCopy = FileRegions;
-		for (FFileRegion& Region : FileRegionsCopy)
-		{
-			Region.Offset -= Info.RegionsOffset;
-		}
-
-		TArray<uint8> Memory;
-		FMemoryWriter Ar(Memory);
-		FFileRegion::SerializeFileRegions(Ar, FileRegionsCopy);
-
-		SavePackageUtilities::WriteToFile(ArchiveFilename + FFileRegion::RegionsFileExtension, Memory.GetData(), Memory.Num());
-	}
-}
-
-void FLooseFileWriter::WriteBulkdata(const FLooseFileWriter::FBulkDataInfo& Info, const FIoBuffer& BulkData, const TArray<FFileRegion>& FileRegions)
-{
-	if (BulkData.DataSize() == 0)
-		return;
-
-	const TCHAR* BulkFileExtension = nullptr;
-
-	switch (Info.BulkdataType)
-	{
-		case FPackageStoreWriter::FBulkDataInfo::Standard:
-			BulkFileExtension = TEXT(".ubulk");		// Regular separate bulk data file
-			break;
-
-		case FPackageStoreWriter::FBulkDataInfo::Mmap:
-			BulkFileExtension = TEXT(".m.ubulk");	// Memory-mapped bulk data
-			break;
-
-		case FPackageStoreWriter::FBulkDataInfo::Optional:
-			BulkFileExtension = TEXT(".uptnl");		// Optional bulk data
-			break;
-
-		default:
-			check(false);
-	}
-
-	const FString ArchiveFilename = FPaths::ChangeExtension(Info.LooseFilePath, BulkFileExtension);
-
-	SavePackageUtilities::WriteToFile(ArchiveFilename, BulkData.Data(), BulkData.DataSize());
-
-	if (FileRegions.Num() > 0)
-	{
-		TArray<uint8> Memory;
-		FMemoryWriter Ar(Memory);
-		FFileRegion::SerializeFileRegions(Ar, const_cast<TArray<FFileRegion>&>(FileRegions));
-
-		SavePackageUtilities::WriteToFile(ArchiveFilename + FFileRegion::RegionsFileExtension, Memory.GetData(), Memory.Num());
-	}
-}
+IPackageStoreWriter::~IPackageStoreWriter() = default;
 
 FSavePackageContext::~FSavePackageContext()
 {
diff --git a/Engine/Source/Runtime/CoreUObject/Private/UObject/SavePackage/SavePackageUtilities.cpp b/Engine/Source/Runtime/CoreUObject/Private/UObject/SavePackage/SavePackageUtilities.cpp
--- a/Engine/Source/Runtime/CoreUObject/Private/UObject/SavePackage/SavePackageUtilities.cpp
+++ b/Engine/Source/Runtime/CoreUObject/Private/UObject/SavePackage/SavePackageUtilities.cpp
@@ -2016,18 +2016,27 @@ void SaveBulkData(FLinkerSave* Linker, const UPackage* InOuter, const TCHAR* Fil
 					return FIoBuffer(FIoBuffer::AssumeOwnership, Writer->ReleaseOwnership(), TotalSize);
 				};
 
-				FPackageStoreWriter::FBulkDataInfo BulkInfo;
+				IPackageStoreWriter::FBulkDataInfo BulkInfo;
 				BulkInfo.PackageName = InOuter->GetFName();
-				BulkInfo.LooseFilePath = Filename;
-				BulkInfo.BulkdataType = FPackageStoreWriter::FBulkDataInfo::Standard;
-
-				SavePackageContext->PackageStoreWriter->WriteBulkdata(BulkInfo, AddSizeAndConvertToIoBuffer(BulkArchive.Get()), BulkArchive->FileRegions);
-
-				BulkInfo.BulkdataType = FPackageStoreWriter::FBulkDataInfo::Optional;
-				SavePackageContext->PackageStoreWriter->WriteBulkdata(BulkInfo, AddSizeAndConvertToIoBuffer(OptionalBulkArchive.Get()), OptionalBulkArchive->FileRegions);
-
-				BulkInfo.BulkdataType = FPackageStoreWriter::FBulkDataInfo::Mmap;
-				SavePackageContext->PackageStoreWriter->WriteBulkdata(BulkInfo, AddSizeAndConvertToIoBuffer(MappedBulkArchive.Get()), MappedBulkArchive->FileRegions);
+				
+				if (BulkArchive->TotalSize())
+				{
+					BulkInfo.BulkdataType = IPackageStoreWriter::FBulkDataInfo::Standard;
+					BulkInfo.LooseFilePath = FPaths::ChangeExtension(Filename, TEXT(".ubulk"));
+					SavePackageContext->PackageStoreWriter->WriteBulkdata(BulkInfo, AddSizeAndConvertToIoBuffer(BulkArchive.Get()), BulkArchive->FileRegions);
+				}
+				if (OptionalBulkArchive->TotalSize())
+				{
+					BulkInfo.BulkdataType = IPackageStoreWriter::FBulkDataInfo::Optional;
+					BulkInfo.LooseFilePath = FPaths::ChangeExtension(Filename, TEXT(".uptnl"));
+					SavePackageContext->PackageStoreWriter->WriteBulkdata(BulkInfo, AddSizeAndConvertToIoBuffer(OptionalBulkArchive.Get()), OptionalBulkArchive->FileRegions);
+				}
+				if (MappedBulkArchive->TotalSize())
+				{
+					BulkInfo.BulkdataType = IPackageStoreWriter::FBulkDataInfo::Mmap;
+					BulkInfo.LooseFilePath = FPaths::ChangeExtension(Filename, TEXT(".m.ubulk"));
+					SavePackageContext->PackageStoreWriter->WriteBulkdata(BulkInfo, AddSizeAndConvertToIoBuffer(MappedBulkArchive.Get()), MappedBulkArchive->FileRegions);
+				}
 			}
 			else
 			{
diff --git a/Engine/Source/Runtime/CoreUObject/Private/UObject/SavePackage2.cpp b/Engine/Source/Runtime/CoreUObject/Private/UObject/SavePackage2.cpp
--- a/Engine/Source/Runtime/CoreUObject/Private/UObject/SavePackage2.cpp
+++ b/Engine/Source/Runtime/CoreUObject/Private/UObject/SavePackage2.cpp
@@ -1783,26 +1783,14 @@ ESavePackageResult FinalizeFile(FStructuredArchive::FRecord& StructuredArchiveRo
 					});
 				}
 
-				FPackageStoreWriter::HeaderInfo HeaderInfo;
-				FPackageStoreWriter::ExportsInfo ExportsInfo;
-
-				ExportsInfo.PackageName = HeaderInfo.PackageName = SaveContext.GetPackage()->GetFName();
-				ExportsInfo.LooseFilePath = HeaderInfo.LooseFilePath = SaveContext.GetFilename();
-
 				const int32 HeaderSize = Linker->Summary.TotalHeaderSize;
-				SavePackageContext->PackageStoreWriter->WriteHeader(HeaderInfo, FIoBuffer(IoBuffer.Data(), HeaderSize, IoBuffer));
-
-				const uint8* ExportsData = IoBuffer.Data() + HeaderSize;
-				const int32 ExportCount = Linker->ExportMap.Num();
 
-				ExportsInfo.Exports.Reserve(ExportCount);
-				ExportsInfo.RegionsOffset = HeaderSize;
+				IPackageStoreWriter::FPackageInfo PackageInfo;
+				PackageInfo.PackageName = SaveContext.GetPackage()->GetFName();
+				PackageInfo.LooseFilePath = SaveContext.GetFilename();
+				PackageInfo.HeaderSize = HeaderSize;
 
-				for (const FObjectExport& Export : Linker->ExportMap)
-				{
-					ExportsInfo.Exports.Add(FIoBuffer(IoBuffer.Data() + Export.SerialOffset, Export.SerialSize, IoBuffer));
-				}
-				SavePackageContext->PackageStoreWriter->WriteExports(ExportsInfo, FIoBuffer(ExportsData, DataSize - HeaderSize, IoBuffer), Linker->FileRegions);
+				SavePackageContext->PackageStoreWriter->WritePackage(PackageInfo, IoBuffer, Linker->FileRegions);
 			}
 			else
 			{
diff --git a/Engine/Source/Runtime/CoreUObject/Public/UObject/SavePackage.h b/Engine/Source/Runtime/CoreUObject/Public/UObject/SavePackage.h
--- a/Engine/Source/Runtime/CoreUObject/Public/UObject/SavePackage.h
+++ b/Engine/Source/Runtime/CoreUObject/Public/UObject/SavePackage.h
@@ -49,34 +49,19 @@ struct FSavePackageArgs
 	FSavePackageContext* SavePackageContext = nullptr;
 };
 
-class FPackageStoreWriter
+class IPackageStoreWriter
 {
 public:
-	COREUOBJECT_API			FPackageStoreWriter();
-	COREUOBJECT_API virtual ~FPackageStoreWriter();
+	COREUOBJECT_API virtual ~IPackageStoreWriter();
 
-	struct HeaderInfo
+	struct FPackageInfo
 	{
 		FName	PackageName;
 		FString	LooseFilePath;
+		uint64  HeaderSize;
 	};
 
-	/** Write 'uasset' data
-	  */
-	virtual void WriteHeader(const HeaderInfo& Info, const FIoBuffer& HeaderData) = 0;
-
-	struct ExportsInfo
-	{
-		FName	PackageName;
-		FString	LooseFilePath;
-		uint64  RegionsOffset;
-
-		TArray<FIoBuffer> Exports;
-	};
-
-	/** Write 'uexp' data
-	  */
-	virtual void WriteExports(const ExportsInfo& Info, const FIoBuffer& ExportsData, const TArray<FFileRegion>& FileRegions) = 0;
+	virtual void WritePackage(const FPackageInfo& Info, const FIoBuffer& PackageData, const TArray<FFileRegion>& FileRegions) = 0;
 
 	struct FBulkDataInfo
 	{
@@ -92,28 +77,15 @@ public:
 		FString	LooseFilePath;
 	};
 
-	/** Write 'ubulk' data
-	  */
 	virtual void WriteBulkdata(const FBulkDataInfo& Info, const FIoBuffer& BulkData, const TArray<FFileRegion>& FileRegions) = 0;
-};
-
-class FLooseFileWriter : public FPackageStoreWriter
-{
-public:
-	COREUOBJECT_API FLooseFileWriter();
-	COREUOBJECT_API ~FLooseFileWriter();
-
-	COREUOBJECT_API virtual void WriteHeader(const HeaderInfo& Info, const FIoBuffer& HeaderData) override;
-	COREUOBJECT_API virtual void WriteExports(const ExportsInfo& Info, const FIoBuffer& ExportsData, const TArray<FFileRegion>& FileRegions) override;
-	COREUOBJECT_API virtual void WriteBulkdata(const FBulkDataInfo& Info, const FIoBuffer& BulkData, const TArray<FFileRegion>& FileRegions) override;
 
-private:
+	virtual void Finalize() = 0;
 };
 
 class FSavePackageContext
 {
 public:
-	FSavePackageContext(FPackageStoreWriter* InPackageStoreWriter, FPackageStoreBulkDataManifest* InBulkDataManifest, bool InbForceLegacyOffsets)
+	FSavePackageContext(IPackageStoreWriter* InPackageStoreWriter, FPackageStoreBulkDataManifest* InBulkDataManifest, bool InbForceLegacyOffsets)
 	: PackageStoreWriter(InPackageStoreWriter) 
 	, BulkDataManifest(InBulkDataManifest)
 	, bForceLegacyOffsets(InbForceLegacyOffsets)
@@ -122,7 +94,7 @@ public:
 
 	COREUOBJECT_API ~FSavePackageContext();
 
-	FPackageStoreWriter* PackageStoreWriter;
+	IPackageStoreWriter* PackageStoreWriter;
 	FPackageStoreBulkDataManifest* BulkDataManifest;
 	bool bForceLegacyOffsets;
 };
diff --git a/Engine/Source/Runtime/SandboxFile/Private/IPlatformFileSandboxWrapper.cpp b/Engine/Source/Runtime/SandboxFile/Private/IPlatformFileSandboxWrapper.cpp
--- a/Engine/Source/Runtime/SandboxFile/Private/IPlatformFileSandboxWrapper.cpp
+++ b/Engine/Source/Runtime/SandboxFile/Private/IPlatformFileSandboxWrapper.cpp
@@ -9,6 +9,7 @@
 #include "Modules/ModuleManager.h"
 #include "HAL/IPlatformFileModule.h"
 #include "Templates/UniquePtr.h"
+#include "IO/IoDispatcherBackend.h"
 
 DEFINE_LOG_CATEGORY(SandboxFile);
 
@@ -153,6 +154,32 @@ bool FSandboxPlatformFile::Initialize(IPlatformFile* Inner, const TCHAR* CmdLine
 			FCommandLine::AddToSubprocessCommandline( *FString::Printf( TEXT("-sandbox=%s"), *SandboxDirectory ) );
 		}
 	}
+
+	FString GlobalUTocPath = FString::Printf(TEXT("%sglobal.utoc"), *SandboxDirectory);
+	if (FileExists(*GlobalUTocPath))
+	{
+		FIoStatus IoDispatcherInitStatus = FIoDispatcher::Initialize();
+		if (IoDispatcherInitStatus.IsOk())
+		{
+			FIoDispatcher& IoDispatcher = FIoDispatcher::Get();
+			TSharedRef<IIoDispatcherFileBackend> IoDispatcherFileBackend = CreateIoDispatcherFileBackend();
+			IoDispatcher.Mount(IoDispatcherFileBackend);
+			TIoStatusOr<FIoContainerId> IoDispatcherMountStatus = IoDispatcherFileBackend->Mount(*FPaths::ChangeExtension(GlobalUTocPath, TEXT("")), 0, FGuid(), FAES::FAESKey());
+			if (IoDispatcherMountStatus.IsOk())
+			{
+				UE_LOG(LogInit, Display, TEXT("Initialized I/O dispatcher"));
+			}
+			else
+			{
+				UE_LOG(LogInit, Error, TEXT("Failed to mount I/O dispatcher container: '%s'"), *IoDispatcherMountStatus.Status().ToString());
+			}
+		}
+		else
+		{
+			UE_LOG(LogInit, Error, TEXT("Failed to initialize I/O dispatcher: '%s'"), *IoDispatcherInitStatus.ToString());
+		}
+	}
+
 	return !!LowerLevel;
 }
 
