IoStore: Fixed progress log messages not outputting correctly to the console
#rnx

diff --git a/Engine/Source/Developer/IoStoreUtilities/Private/IoStoreUtilities.cpp b/Engine/Source/Developer/IoStoreUtilities/Private/IoStoreUtilities.cpp
--- a/Engine/Source/Developer/IoStoreUtilities/Private/IoStoreUtilities.cpp
+++ b/Engine/Source/Developer/IoStoreUtilities/Private/IoStoreUtilities.cpp
@@ -1439,41 +1439,6 @@ private:
 	static constexpr uint64 BufferMemoryLimit = 2ull << 30;
 };
 
-class FIoStoreProgressReporter
-{
-public:
-	FIoStoreProgressReporter(const FIoStoreWriterContext& InWriterContext)
-		: WriterContext(InWriterContext)
-		, StopEvent(FPlatformProcess::GetSynchEventFromPool(false))
-	{
-		ReporterThread = Async(EAsyncExecution::Thread, [this]() { ReporterThreadFunc(); });
-	}
-
-	~FIoStoreProgressReporter()
-	{
-		bStop.Store(true);
-		StopEvent->Trigger();
-		ReporterThread.Wait();
-		FPlatformProcess::ReturnSynchEventToPool(StopEvent);
-	}
-
-private:
-	void ReporterThreadFunc()
-	{
-		while (!bStop.Load())
-		{
-			StopEvent->Wait(FTimespan::FromSeconds(2.0));
-			FIoStoreWriterContext::FProgress Progress = WriterContext.GetProgress();
-			UE_LOG(LogIoStore, Display, TEXT("Hashed, Compressed, Serialized: %lld, %lld, %lld / %lld"), Progress.HashedChunksCount, Progress.CompressedChunksCount, Progress.SerializedChunksCount, Progress.TotalChunksCount);
-		}
-	}
-
-	const FIoStoreWriterContext& WriterContext;
-	TFuture<void> ReporterThread;
-	FEvent* StopEvent;
-	TAtomic<bool> bStop{ false };
-};
-
 int32 CreateTarget(const FIoStoreArguments& Arguments, const FIoStoreWriterSettings& GeneralIoWriterSettings)
 {
 	TGuardValue<int32> GuardAllowUnversionedContentInEditor(GAllowUnversionedContentInEditor, 1);
@@ -1691,19 +1656,26 @@ int32 CreateTarget(const FIoStoreArguments& Arguments, const FIoStoreWriterSetti
 
 	UE_LOG(LogIoStore, Display, TEXT("Serializing container(s)..."));
 
-	FIoStoreProgressReporter* IoStoreProgressReporter = new FIoStoreProgressReporter(*IoStoreWriterContext);
-
 	TArray<FIoStoreWriterResult> IoStoreWriterResults;
 	IoStoreWriterResults.Reserve(IoStoreWriters.Num());
 	for (FIoStoreWriter* IoStoreWriter : IoStoreWriters)
 	{
-		IoStoreWriterResults.Emplace(IoStoreWriter->Flush().ConsumeValueOrDie());
+		TFuture<FIoStoreWriterResult> FlushTask = Async(EAsyncExecution::ThreadPool, [IoStoreWriter]()
+		{
+			return IoStoreWriter->Flush().ConsumeValueOrDie();
+		});
+
+		while (!FlushTask.IsReady())
+		{
+			FlushTask.WaitFor(FTimespan::FromSeconds(2.0));
+			FIoStoreWriterContext::FProgress Progress = IoStoreWriterContext->GetProgress();
+			UE_LOG(LogIoStore, Display, TEXT("Hashed, Compressed, Serialized: %lld, %lld, %lld / %lld"), Progress.HashedChunksCount, Progress.CompressedChunksCount, Progress.SerializedChunksCount, Progress.TotalChunksCount);
+		}
+		IoStoreWriterResults.Emplace(FlushTask.Get());
 		delete IoStoreWriter;
 	}
 	IoStoreWriters.Empty();
 
-	delete IoStoreProgressReporter;
-
 	IOSTORE_CPU_SCOPE(CalculateStats);
 
 	UE_LOG(LogIoStore, Display, TEXT("Calculating stats..."));
