Remove the bulkdata manifest that was generated when cooking.

- The manifest was for a prototype version of the IoStore commandlet and would attempt to record every bulkdata entry during a cook, which would allow us to generate the FioChunkId for each bulkdata entry when building the ucas files.
- Due to changes during development we stopped requiring this data and it was never used in actual production code.
- Since this code was never in use we can just remove it.

* IoStoreUtilities
- No longer attempt to load the manifest
- Removed use of BulkdataTypeToChunkIdType and just use the EIoChunkType directly instead.

#test Full stage of ShooterGame
#rnx

diff --git a/Engine/Source/Developer/IoStoreUtilities/Private/IoStoreUtilities.cpp b/Engine/Source/Developer/IoStoreUtilities/Private/IoStoreUtilities.cpp
--- a/Engine/Source/Developer/IoStoreUtilities/Private/IoStoreUtilities.cpp
+++ b/Engine/Source/Developer/IoStoreUtilities/Private/IoStoreUtilities.cpp
@@ -20,7 +20,6 @@
 #include "Misc/KeyChainUtilities.h"
 #include "Modules/ModuleManager.h"
 #include "Serialization/Archive.h"
-#include "Serialization/BulkDataManifest.h"
 #include "Serialization/JsonReader.h"
 #include "Serialization/JsonSerializer.h"
 #include "Serialization/BufferWriter.h"
@@ -602,22 +601,6 @@ static void CreateDiskLayout(
 	}
 }
 
-static EIoChunkType BulkdataTypeToChunkIdType(FPackageStoreBulkDataManifest::EBulkdataType Type)
-{
-	switch (Type)
-	{
-	case FPackageStoreBulkDataManifest::EBulkdataType::Normal:
-		return EIoChunkType::BulkData;
-	case FPackageStoreBulkDataManifest::EBulkdataType::Optional:
-		return EIoChunkType::OptionalBulkData;
-	case FPackageStoreBulkDataManifest::EBulkdataType::MemoryMapped:
-		return EIoChunkType::MemoryMappedBulkData;
-	default:
-		UE_LOG(LogIoStore, Error, TEXT("Invalid EBulkdataType (%d) found!"), Type);
-		return EIoChunkType::Invalid;
-	}
-}
-
 FContainerTargetSpec* AddContainer(
 	FName Name,
 	TArray<FContainerTargetSpec*>& Containers)
@@ -924,17 +907,17 @@ void InitializeContainerTargetsAndPackages(
 						if (CookedFileStatData->FileExt == FCookedFileStatData::UPtnl)
 						{
 							TargetFile.bIsOptionalBulkData = true;
-							TargetFile.ChunkId = CreateChunkId(Package->GlobalPackageId, 0, BulkdataTypeToChunkIdType(FPackageStoreBulkDataManifest::EBulkdataType::Optional), *TargetFile.TargetPath);
+							TargetFile.ChunkId = CreateChunkId(Package->GlobalPackageId, 0, EIoChunkType::OptionalBulkData, *TargetFile.TargetPath);
 						}
 						else if (CookedFileStatData->FileExt == FCookedFileStatData::UMappedBulk)
 						{
 							TargetFile.bIsMemoryMappedBulkData = true;
 							TargetFile.bForceUncompressed = true;
-							TargetFile.ChunkId = CreateChunkId(Package->GlobalPackageId, 0, BulkdataTypeToChunkIdType(FPackageStoreBulkDataManifest::EBulkdataType::MemoryMapped), *TargetFile.TargetPath);
+							TargetFile.ChunkId = CreateChunkId(Package->GlobalPackageId, 0, EIoChunkType::MemoryMappedBulkData, *TargetFile.TargetPath);
 						}
 						else
 						{
-							TargetFile.ChunkId = CreateChunkId(Package->GlobalPackageId, 0, BulkdataTypeToChunkIdType(FPackageStoreBulkDataManifest::EBulkdataType::Normal), *TargetFile.TargetPath);
+							TargetFile.ChunkId = CreateChunkId(Package->GlobalPackageId, 0, EIoChunkType::BulkData, *TargetFile.TargetPath);
 						}
 						if (Package->FileName.IsEmpty())
 						{
@@ -1447,16 +1430,6 @@ int32 CreateTarget(const FIoStoreArguments& Arguments, const FIoStoreWriterSetti
 {
 	TGuardValue<int32> GuardAllowUnversionedContentInEditor(GAllowUnversionedContentInEditor, 1);
 
-	FPackageStoreBulkDataManifest BulkDataManifest(FString(Arguments.CookedDir) / FApp::GetProjectName());
-	if (!BulkDataManifest.Load())
-	{
-		UE_LOG(LogIoStore, Warning, TEXT("Failed to load Bulk Data manifest %s"), *BulkDataManifest.GetFilename());
-	}
-	else
-	{
-		UE_LOG(LogIoStore, Display, TEXT("Loaded Bulk Data manifest '%s'"), *BulkDataManifest.GetFilename());
-	}
-	
 #if OUTPUT_CHUNKID_DIRECTORY
 	ChunkIdCsv.CreateOutputFile(CookedDir);
 #endif
diff --git a/Engine/Source/Editor/UnrealEd/Private/CookOnTheFlyServer.cpp b/Engine/Source/Editor/UnrealEd/Private/CookOnTheFlyServer.cpp
--- a/Engine/Source/Editor/UnrealEd/Private/CookOnTheFlyServer.cpp
+++ b/Engine/Source/Editor/UnrealEd/Private/CookOnTheFlyServer.cpp
@@ -82,7 +82,6 @@
 #include "Serialization/ArchiveStackTrace.h"
 #include "DistanceFieldAtlas.h"
 #include "Cooker/AsyncIODelete.h"
-#include "Serialization/BulkDataManifest.h"
 #include "MeshCardRepresentation.h"
 #include "Misc/PathViews.h"
 #include "String/Find.h"
@@ -7058,7 +7057,6 @@ void UCookOnTheFlyServer::InitializePackageStore(const TArrayView<const ITargetP
 		const FString ResolvedRootPath = RootPathSandbox.Replace(TEXT("[Platform]"), *PlatformString);
 		const FString ResolvedProjectPath = ProjectPathSandbox.Replace(TEXT("[Platform]"), *PlatformString);
 
-		FPackageStoreBulkDataManifest* BulkDataManifest	= bIsDiffOnly == false ? new FPackageStoreBulkDataManifest(ResolvedProjectPath) : nullptr;
 		FPackageStoreWriter* PackageStoreWriter = nullptr;
 		if (IsUsingIoStore())
 		{
@@ -7071,7 +7069,7 @@ void UCookOnTheFlyServer::InitializePackageStore(const TArrayView<const ITargetP
 		bool bLegacyBulkDataOffsets = false;
 		PlatformEngineIni.GetBool(TEXT("Core.System"), TEXT("LegacyBulkDataOffsets"), bLegacyBulkDataOffsets);
 
-		FSavePackageContext* SavePackageContext			= new FSavePackageContext(PackageStoreWriter, BulkDataManifest, bLegacyBulkDataOffsets);
+		FSavePackageContext* SavePackageContext = new FSavePackageContext(PackageStoreWriter, bLegacyBulkDataOffsets);
 		SavePackageContexts.Add(SavePackageContext);
 	}
 }
@@ -7085,10 +7083,6 @@ void UCookOnTheFlyServer::FinalizePackageStore()
 	{
 		if (PackageContext != nullptr)
 		{
-			if (PackageContext->BulkDataManifest != nullptr)
-			{
-				PackageContext->BulkDataManifest->Save();
-			}
 			if (PackageContext->PackageStoreWriter != nullptr)
 			{
 				PackageContext->PackageStoreWriter->Finalize();
diff --git a/Engine/Source/Runtime/CoreUObject/Private/Serialization/BulkDataManifest.cpp b/Engine/Source/Runtime/CoreUObject/Private/Serialization/BulkDataManifest.cpp
deleted file mode 100644
diff --git a/Engine/Source/Runtime/CoreUObject/Private/UObject/SavePackage.cpp b/Engine/Source/Runtime/CoreUObject/Private/UObject/SavePackage.cpp
--- a/Engine/Source/Runtime/CoreUObject/Private/UObject/SavePackage.cpp
+++ b/Engine/Source/Runtime/CoreUObject/Private/UObject/SavePackage.cpp
@@ -65,7 +65,6 @@
 #include "Serialization/ArchiveUObjectFromStructuredArchive.h"
 #include "Serialization/UnversionedPropertySerialization.h"
 #include "UObject/AsyncWorkSequence.h"
-#include "Serialization/BulkDataManifest.h"
 #include "Misc/ScopeExit.h"
 #include "Misc/PackageAccessTracking.h"
 #include "Misc/PackageAccessTrackingOps.h"
@@ -4711,7 +4710,6 @@ IPackageStoreWriter::~IPackageStoreWriter() = default;
 FSavePackageContext::~FSavePackageContext()
 {
 	delete PackageStoreWriter;
-	delete BulkDataManifest;
 }
 
 #endif	// UE_WITH_SAVEPACKAGE
diff --git a/Engine/Source/Runtime/CoreUObject/Private/UObject/SavePackage/SavePackageUtilities.cpp b/Engine/Source/Runtime/CoreUObject/Private/UObject/SavePackage/SavePackageUtilities.cpp
--- a/Engine/Source/Runtime/CoreUObject/Private/UObject/SavePackage/SavePackageUtilities.cpp
+++ b/Engine/Source/Runtime/CoreUObject/Private/UObject/SavePackage/SavePackageUtilities.cpp
@@ -15,7 +15,6 @@
 #include "Misc/ScopedSlowTask.h"
 #include "SaveContext.h"
 #include "Serialization/BulkData.h"
-#include "Serialization/BulkDataManifest.h"
 #include "Serialization/LargeMemoryWriter.h"
 #include "UObject/AsyncWorkSequence.h"
 #include "UObject/Class.h"
@@ -1967,11 +1966,7 @@ ESavePackageResult SaveBulkData(FLinkerSave* Linker, const UPackage* InOuter, co
 	if (bRequestSaveByReference)
 	{
 		const TCHAR* FailureReason = nullptr;
-		if (SavePackageContext != nullptr && (SavePackageContext->BulkDataManifest != nullptr))
-		{
-			FailureReason = TEXT("SAVE_BulkDataByReference is incompatible with BulkDataManifest");
-		}
-		else if (Linker->bUpdatingLoadedPath)
+		if (Linker->bUpdatingLoadedPath)
 		{
 			FailureReason = TEXT("SAVE_BulkDataByReference is incompatible with bUpdatingLoadedPath");
 		}
@@ -2134,30 +2129,6 @@ ESavePackageResult SaveBulkData(FLinkerSave* Linker, const UPackage* InOuter, co
 			BulkDataOffsetInFile = (BulkDataFlags & BULKDATA_NoOffsetFixUp) == 0 ? BulkStartOffset - StartOfBulkDataArea : BulkStartOffset;
 			BulkDataSizeOnDisk = BulkEndOffset - BulkStartOffset;
 
-			if (SavePackageContext != nullptr && SavePackageContext->BulkDataManifest != nullptr)
-			{
-				auto BulkDataTypeFromFlags = [](uint32 BulkDataFlags)
-				{
-					if (BulkDataFlags & BULKDATA_MemoryMappedPayload)
-					{
-						return FPackageStoreBulkDataManifest::EBulkdataType::MemoryMapped;
-					}
-					else if (BulkDataFlags & BULKDATA_OptionalPayload)
-					{
-						return FPackageStoreBulkDataManifest::EBulkdataType::Optional;
-					}
-					else
-					{
-						check(!(BulkDataFlags & BULKDATA_WorkspaceDomainPayload));
-						return FPackageStoreBulkDataManifest::EBulkdataType::Normal;
-					}
-				};
-
-				const FPackageStoreBulkDataManifest::EBulkdataType Type = BulkDataTypeFromFlags(BulkDataStorageInfo.BulkDataFlags);
-				SavePackageContext->BulkDataManifest->AddFileAccess(Filename, Type,
-					BulkDataOffsetInFile /* Use the serialized offset as the ChunkId */, BulkStartOffset, BulkDataSizeOnDisk);
-			}
-
 			if (bUseFileRegions && BulkDataStorageInfo.BulkDataFileRegionType != EFileRegionType::None && BulkDataSizeOnDisk > 0)
 			{
 				TargetRegions->Add(FFileRegion(BulkStartOffset, BulkDataSizeOnDisk, BulkDataStorageInfo.BulkDataFileRegionType));
diff --git a/Engine/Source/Runtime/CoreUObject/Public/Serialization/BulkDataManifest.h b/Engine/Source/Runtime/CoreUObject/Public/Serialization/BulkDataManifest.h
deleted file mode 100644
diff --git a/Engine/Source/Runtime/CoreUObject/Public/UObject/SavePackage.h b/Engine/Source/Runtime/CoreUObject/Public/UObject/SavePackage.h
--- a/Engine/Source/Runtime/CoreUObject/Public/UObject/SavePackage.h
+++ b/Engine/Source/Runtime/CoreUObject/Public/UObject/SavePackage.h
@@ -18,7 +18,6 @@ class FArchive;
 class FIoBuffer;
 struct FObjectSaveContextData;
 class FPackagePath;
-class FPackageStoreBulkDataManifest;
 class FSavePackageContext;
 class FArchiveDiffMap;
 class FOutputDevice;
@@ -87,9 +86,8 @@ public:
 class FSavePackageContext
 {
 public:
-	FSavePackageContext(IPackageStoreWriter* InPackageStoreWriter, FPackageStoreBulkDataManifest* InBulkDataManifest, bool InbForceLegacyOffsets)
+	FSavePackageContext(IPackageStoreWriter* InPackageStoreWriter, bool InbForceLegacyOffsets)
 	: PackageStoreWriter(InPackageStoreWriter) 
-	, BulkDataManifest(InBulkDataManifest)
 	, bForceLegacyOffsets(InbForceLegacyOffsets)
 	{
 	}
@@ -97,7 +95,6 @@ public:
 	COREUOBJECT_API ~FSavePackageContext();
 
 	IPackageStoreWriter* PackageStoreWriter;
-	FPackageStoreBulkDataManifest* BulkDataManifest;
 	bool bForceLegacyOffsets;
 };
 
