IoStore: Don't attempt to serialize package summaries from empty files
#rnx

diff --git a/Engine/Source/Developer/IoStoreUtilities/Private/IoStoreUtilities.cpp b/Engine/Source/Developer/IoStoreUtilities/Private/IoStoreUtilities.cpp
--- a/Engine/Source/Developer/IoStoreUtilities/Private/IoStoreUtilities.cpp
+++ b/Engine/Source/Developer/IoStoreUtilities/Private/IoStoreUtilities.cpp
@@ -736,8 +736,15 @@ static void ParsePackageAssets(TArray<FLegacyCookedPackage*>& Packages, const FP
 			uint8* PackageBuffer = PackageAssetBuffers[Index];
 			FLegacyCookedPackage* Package = Packages[Index];
 
-			FIoBuffer CookedHeaderBuffer = FIoBuffer(FIoBuffer::Wrap, PackageBuffer, Package->UAssetSize);
-			Package->OptimizedPackage = PackageStoreOptimizer.CreatePackageFromCookedHeader(Package->PackageName, CookedHeaderBuffer);
+			if (Package->UAssetSize)
+			{
+				FIoBuffer CookedHeaderBuffer = FIoBuffer(FIoBuffer::Wrap, PackageBuffer, Package->UAssetSize);
+				Package->OptimizedPackage = PackageStoreOptimizer.CreatePackageFromCookedHeader(Package->PackageName, CookedHeaderBuffer);
+			}
+			else
+			{
+				Package->OptimizedPackage = PackageStoreOptimizer.CreateMissingPackage(Package->PackageName);
+			}
 			check(Package->OptimizedPackage->GetId() == Package->GlobalPackageId);
 		}, EParallelForFlags::Unbalanced);
 	}
diff --git a/Engine/Source/Developer/IoStoreUtilities/Private/PackageStoreOptimizer.cpp b/Engine/Source/Developer/IoStoreUtilities/Private/PackageStoreOptimizer.cpp
--- a/Engine/Source/Developer/IoStoreUtilities/Private/PackageStoreOptimizer.cpp
+++ b/Engine/Source/Developer/IoStoreUtilities/Private/PackageStoreOptimizer.cpp
@@ -92,6 +92,15 @@ void FPackageStoreOptimizer::Initialize(const ITargetPlatform* TargetPlatform)
 	FindScriptObjects(TargetPlatform);
 }
 
+FPackageStorePackage* FPackageStoreOptimizer::CreateMissingPackage(const FName& Name) const
+{
+	FPackageStorePackage* Package = new FPackageStorePackage();
+	Package->Name = Name;
+	Package->Id = FPackageId::FromName(Name);
+	Package->SourceName = *RemapLocalizationPathIfNeeded(Name.ToString(), &Package->Region);
+	return Package;
+}
+
 FPackageStorePackage* FPackageStoreOptimizer::CreatePackageFromCookedHeader(const FName& Name, const FIoBuffer& CookedHeaderBuffer) const
 {
 	FPackageStorePackage* Package = new FPackageStorePackage();
diff --git a/Engine/Source/Developer/IoStoreUtilities/Private/PackageStoreOptimizer.h b/Engine/Source/Developer/IoStoreUtilities/Private/PackageStoreOptimizer.h
--- a/Engine/Source/Developer/IoStoreUtilities/Private/PackageStoreOptimizer.h
+++ b/Engine/Source/Developer/IoStoreUtilities/Private/PackageStoreOptimizer.h
@@ -319,6 +319,7 @@ public:
 		bIncrementalResolveEnabled = true;
 	}
 
+	FPackageStorePackage* CreateMissingPackage(const FName& Name) const;
 	FPackageStorePackage* CreatePackageFromCookedHeader(const FName& Name, const FIoBuffer& CookedHeaderBuffer) const;
 	void BeginResolvePackage(FPackageStorePackage* Package);
 	void Flush(bool bAllowMissingImports, TFunction<void(FPackageStorePackage*)> ResolvedPackageCallback = TFunction<void(FPackageStorePackage*)>());
