Added support for storing shaders in iostore.
Opt-In with core setting AllowShadersInIoStore

diff --git a/Engine/Source/Developer/IoStoreUtilities/IoStoreUtilities.Build.cs b/Engine/Source/Developer/IoStoreUtilities/IoStoreUtilities.Build.cs
--- a/Engine/Source/Developer/IoStoreUtilities/IoStoreUtilities.Build.cs
+++ b/Engine/Source/Developer/IoStoreUtilities/IoStoreUtilities.Build.cs
@@ -15,6 +15,7 @@ public class IoStoreUtilities : ModuleRules
             "CoreUObject",
             "Projects",
             "AssetRegistry",
+            "RenderCore",
         });
         PrivateDependencyModuleNames.Add("PakFile");
         PrivateDependencyModuleNames.Add("Json");
diff --git a/Engine/Source/Developer/IoStoreUtilities/Private/IoStoreUtilities.cpp b/Engine/Source/Developer/IoStoreUtilities/Private/IoStoreUtilities.cpp
--- a/Engine/Source/Developer/IoStoreUtilities/Private/IoStoreUtilities.cpp
+++ b/Engine/Source/Developer/IoStoreUtilities/Private/IoStoreUtilities.cpp
@@ -54,6 +54,7 @@
 #include "Algo/MaxElement.h"
 #include "Algo/StableSort.h"
 #include "PackageStoreOptimizer.h"
+#include "ShaderCodeLibrary.h"
 
 //PRAGMA_DISABLE_OPTIMIZATION
 
@@ -284,8 +285,8 @@ struct FContainerSourceSpec
 
 struct FCookedFileStatData
 {
-	enum EFileExt { UMap, UAsset, UExp, UBulk, UPtnl, UMappedBulk };
-	enum EFileType { PackageHeader, PackageData, BulkData };
+	enum EFileExt { UMap, UAsset, UExp, UBulk, UPtnl, UMappedBulk, UShaderByteCode };
+	enum EFileType { PackageHeader, PackageData, BulkData, ShaderLibrary };
 
 	int64 FileSize = 0;
 	EFileType FileType = PackageHeader;
@@ -298,6 +299,16 @@ using FCookedFileStatMap = TMap<FString, FCookedFileStatData>;
 
 struct FContainerTargetSpec;
 
+struct FShaderInfo
+{
+	FIoChunkId ChunkId;
+	FIoChunkId LibraryChunkId;
+	FContainerTargetSpec* ContainerTarget = nullptr;
+	uint64 Size = 0;
+	uint32 ReferencedByPackagesCount = 0;
+	uint32 LibraryOrder = 0;
+};
+
 struct FLegacyCookedPackage
 {
 	FPackageId GlobalPackageId;
@@ -308,6 +319,17 @@ struct FLegacyCookedPackage
 	uint64 UExpSize = 0;
 	uint64 DiskLayoutOrder = -1;
 	FPackageStorePackage* OptimizedPackage = nullptr;
+	TSet<FShaderInfo*> Shaders;
+};
+
+enum class EContainerChunkType
+{
+	ShaderCodeLibrary,
+	ShaderCode,
+	PackageData,
+	BulkData,
+	OptionalBulkData,
+	MemoryMappedBulkData
 };
 
 struct FContainerTargetFile
@@ -315,15 +337,14 @@ struct FContainerTargetFile
 	FContainerTargetSpec* ContainerTarget = nullptr;
 	FLegacyCookedPackage* Package = nullptr;
 	FString NormalizedSourcePath;
+	TOptional<FIoBuffer> SourceBuffer;
 	FString TargetPath;
 	FString DestinationPath;
 	uint64 SourceSize = 0;
 	uint64 IdealOrder = 0;
 	FIoChunkId ChunkId;
 	TArray<uint8> PackageHeaderData;
-	bool bIsBulkData = false;
-	bool bIsOptionalBulkData = false;
-	bool bIsMemoryMappedBulkData = false;
+	EContainerChunkType ChunkType;
 	bool bForceUncompressed = false;
 
 	TArray<FFileRegion> FileRegions;
@@ -387,7 +408,8 @@ struct FContainerTargetSpec
 	TArray<FContainerTargetFile> TargetFiles;
 	TArray<TUniquePtr<FIoStoreReader>> PatchSourceReaders;
 	EIoContainerFlags ContainerFlags = EIoContainerFlags::None;
-	uint32 PackageCount = 0;
+	TArray<FLegacyCookedPackage*> Packages;
+	TArray<FShaderInfo> Shaders;
 	bool bGenerateDiffPatch = false;
 };
 
@@ -652,34 +674,115 @@ static void CreateDiskLayout(
 	{
 		TArray<FContainerTargetFile*> SortedTargetFiles;
 		SortedTargetFiles.Reserve(ContainerTarget->TargetFiles.Num());
+		TMap<FIoChunkId, FContainerTargetFile*> ShaderTargetFilesMap;
+		ShaderTargetFilesMap.Reserve(ContainerTarget->Shaders.Num());
 		for (FContainerTargetFile& TargetFile : ContainerTarget->TargetFiles)
 		{
-			SortedTargetFiles.Add(&TargetFile);
+			if (TargetFile.ChunkType == EContainerChunkType::ShaderCode)
+			{
+				ShaderTargetFilesMap.Add(TargetFile.ChunkId, &TargetFile);
+			}
+			else
+			{
+				SortedTargetFiles.Add(&TargetFile);
+			}
 		}
 		Algo::Sort(SortedTargetFiles, [](const FContainerTargetFile* A, const FContainerTargetFile* B)
 		{
-			if (A->bIsMemoryMappedBulkData != B->bIsMemoryMappedBulkData)
+			if (A->ChunkType != B->ChunkType)
 			{
-				// Put all memory mapped bulkdata last
-				return B->bIsMemoryMappedBulkData;
+				return A->ChunkType < B->ChunkType;
 			}
-			else if (A->bIsBulkData != B->bIsBulkData)
+			if (A->ChunkType == EContainerChunkType::ShaderCodeLibrary)
 			{
-				// Put packages before bulkdata
-				return B->bIsBulkData;
+				return A->TargetPath < B->TargetPath;
 			}
-			else if (A->Package != B->Package)
+			if (A->Package != B->Package)
 			{
 				return A->Package->DiskLayoutOrder < B->Package->DiskLayoutOrder;
 			}
-			else if (A->bIsOptionalBulkData != B->bIsOptionalBulkData)
-			{
-				// Put each ubulk right before it's corresponding uptnl
-				return B->bIsOptionalBulkData;
-			}
 			check(A == B)
 			return false;
 		});
+
+		TArray<const FShaderInfo*> SharedShaders;
+		TSet<const FShaderInfo*> AddedSharedShaders;
+		int32 Index = 0;
+		int32 ShaderCodeInsertionIndex = -1;
+		while (Index < SortedTargetFiles.Num())
+		{
+			FContainerTargetFile* TargetFile = SortedTargetFiles[Index];
+			if (ShaderCodeInsertionIndex < 0 && TargetFile->ChunkType != EContainerChunkType::ShaderCodeLibrary)
+			{
+				ShaderCodeInsertionIndex = Index;
+			}
+			if (TargetFile->ChunkType == EContainerChunkType::PackageData)
+			{
+				TArray<FContainerTargetFile*, TInlineAllocator<1024>> PackageUniqueShaders;
+				for (FShaderInfo* Shader : TargetFile->Package->Shaders)
+				{
+					if (Shader->ContainerTarget == ContainerTarget)
+					{
+						check(Shader->ReferencedByPackagesCount > 0);
+						if (Shader->ReferencedByPackagesCount == 1)
+						{
+							FContainerTargetFile* ShaderTargetFile = ShaderTargetFilesMap.FindRef(Shader->ChunkId);
+							check(ShaderTargetFile);
+							PackageUniqueShaders.Add(ShaderTargetFile);
+						}
+						else if (!AddedSharedShaders.Contains(Shader))
+						{
+							SharedShaders.Add(Shader);
+							AddedSharedShaders.Add(Shader);
+						}
+					}
+				}
+				SortedTargetFiles.Insert(PackageUniqueShaders, Index + 1);
+				Index += PackageUniqueShaders.Num();
+			}
+			++Index;
+		}
+		if (ShaderCodeInsertionIndex < 0)
+		{
+			ShaderCodeInsertionIndex = 0;
+		}
+
+		TArray<const FShaderInfo*> GlobalShaders;
+		for (const FShaderInfo& Shader : ContainerTarget->Shaders)
+		{
+			if (Shader.ReferencedByPackagesCount == 0)
+			{
+				GlobalShaders.Add(&Shader);
+			}
+		}
+		if (!GlobalShaders.IsEmpty())
+		{
+			TArray<FContainerTargetFile*> GlobalShaderTargetFiles;
+			GlobalShaderTargetFiles.Reserve(GlobalShaders.Num());
+			for (const FShaderInfo* ShaderInfo : GlobalShaders)
+			{
+				FContainerTargetFile* ShaderTargetFile = ShaderTargetFilesMap.FindRef(ShaderInfo->ChunkId);
+				check(ShaderTargetFile);
+				GlobalShaderTargetFiles.Add(ShaderTargetFile);
+			}
+			SortedTargetFiles.Insert(GlobalShaderTargetFiles, ShaderCodeInsertionIndex);
+			ShaderCodeInsertionIndex += GlobalShaderTargetFiles.Num();
+		}
+		if (!SharedShaders.IsEmpty())
+		{
+			TArray<FContainerTargetFile*> SharedShaderTargetFiles;
+			SharedShaderTargetFiles.Reserve(SharedShaders.Num());
+			for (const FShaderInfo* ShaderInfo : SharedShaders)
+			{
+				FContainerTargetFile* ShaderTargetFile = ShaderTargetFilesMap.FindRef(ShaderInfo->ChunkId);
+				check(ShaderTargetFile);
+				SharedShaderTargetFiles.Add(ShaderTargetFile);
+			}
+			SortedTargetFiles.Insert(SharedShaderTargetFiles, ShaderCodeInsertionIndex);
+		}
+
+		check(SortedTargetFiles.Num() == ContainerTarget->TargetFiles.Num());
+
 		uint64 IdealOrder = 0;
 		for (FContainerTargetFile* TargetFile : SortedTargetFiles)
 		{
@@ -877,6 +980,139 @@ TArray<TUniquePtr<FIoStoreReader>> CreatePatchSourceReaders(const TArray<FString
 	return Readers;
 }
 
+void ProcessShaderLibraries(const FIoStoreArguments& Arguments, TArray<FContainerTargetSpec*>& ContainerTargets)
+{
+	IOSTORE_CPU_SCOPE(ProcessShaderLibraries);
+
+	TMap<FIoChunkId, FSHAHash> AllShaderCodeHashes;
+	uint64 LibraryOrder = 0;
+	for (FContainerTargetSpec* ContainerTarget : ContainerTargets)
+	{
+		TMap<FSHAHash, TArray<FIoChunkId>> ShaderChunkIdsByShaderMapHash;
+		TSet<FIoChunkId> AddedShaderChunkIds;
+		TArray<FContainerTargetFile> ShaderTargetFiles;
+		TArray<FShaderInfo> Shaders;
+		for (FContainerTargetFile& TargetFile : ContainerTarget->TargetFiles)
+		{
+			if (TargetFile.ChunkType == EContainerChunkType::ShaderCodeLibrary)
+			{
+				TArray<TTuple<FSHAHash, TArray<FIoChunkId>>> ShaderMaps;
+				TMap<FIoChunkId, FIoBuffer> ShaderCodeMap;
+				TTuple<FIoChunkId, FIoBuffer> LibraryChunk;
+				TArray<TTuple<FIoChunkId, FIoBuffer, FSHAHash>> CodeChunks;
+				TArray<TTuple<FSHAHash, TSet<FName>>> ShaderMapAssetAssociations;
+				if (!FShaderLibraryCooker::ConvertToIoStoreShaderLibrary(*TargetFile.NormalizedSourcePath, LibraryChunk, CodeChunks, ShaderMaps, ShaderMapAssetAssociations))
+				{
+					UE_LOG(LogIoStore, Warning, TEXT("Failed converting shader library '%s'"), *TargetFile.NormalizedSourcePath);
+					continue;
+				}
+				TargetFile.ChunkId = LibraryChunk.Key;
+				TargetFile.SourceBuffer.Emplace(LibraryChunk.Value);
+				for (const TTuple<FIoChunkId, FIoBuffer, FSHAHash>& CodeChunk : CodeChunks)
+				{
+					ShaderCodeMap.Add(CodeChunk.Get<0>(), CodeChunk.Get<1>());
+					FSHAHash* FindShaderCodeHash = AllShaderCodeHashes.Find(CodeChunk.Get<0>());
+					if (FindShaderCodeHash)
+					{
+						UE_CLOG(*FindShaderCodeHash != CodeChunk.Get<2>(), LogIoStore, Fatal, TEXT("Shader code chunk id collision"));
+					}
+					else
+					{
+						AllShaderCodeHashes.Add(CodeChunk.Get<0>(), CodeChunk.Get<2>());
+					}
+				}
+				TMap<FName, TSet<FSHAHash>> PackageNameToShaderMaps;
+				for (const TTuple<FSHAHash, TSet<FName>>& ShaderMapAssetAssociation : ShaderMapAssetAssociations)
+				{
+					for (const FName& PackageName : ShaderMapAssetAssociation.Value)
+					{
+						TSet<FSHAHash>& PackageShaderMaps = PackageNameToShaderMaps.FindOrAdd(PackageName);
+						PackageShaderMaps.Add(ShaderMapAssetAssociation.Key);
+					}
+				}
+				for (FLegacyCookedPackage* Package : ContainerTarget->Packages)
+				{
+					TSet<FSHAHash>* FindShaderMapHashes = PackageNameToShaderMaps.Find(Package->PackageName);
+					if (FindShaderMapHashes)
+					{
+						for (const FSHAHash& ShaderMapHash : *FindShaderMapHashes)
+						{
+							Package->OptimizedPackage->AddShaderMapHash(ShaderMapHash);
+						}
+					}
+				}
+
+				for (TTuple<FSHAHash, TArray<FIoChunkId>>& ShaderMap : ShaderMaps)
+				{
+					for (const FIoChunkId& ShaderChunkId : ShaderMap.Value)
+					{
+						if (!AddedShaderChunkIds.Contains(ShaderChunkId))
+						{
+							FContainerTargetFile& ShaderTargetFile = ShaderTargetFiles.AddDefaulted_GetRef();
+							ShaderTargetFile.ContainerTarget = ContainerTarget;
+							ShaderTargetFile.ChunkId = ShaderChunkId;
+							ShaderTargetFile.ChunkType = EContainerChunkType::ShaderCode;
+							ShaderTargetFile.bForceUncompressed = true;
+							
+							FShaderInfo& ShaderInfo = Shaders.AddDefaulted_GetRef();
+							ShaderInfo.ChunkId = ShaderChunkId;
+							ShaderInfo.LibraryChunkId = TargetFile.ChunkId;
+							ShaderInfo.ContainerTarget = ContainerTarget;
+							ShaderInfo.LibraryOrder = LibraryOrder++;
+
+							FIoBuffer* FindCodeIoBuffer = ShaderCodeMap.Find(ShaderChunkId);
+							check(FindCodeIoBuffer);
+							ShaderTargetFile.SourceBuffer.Emplace(*FindCodeIoBuffer);
+							ShaderTargetFile.SourceSize = FindCodeIoBuffer->DataSize();
+							ShaderInfo.Size = FindCodeIoBuffer->DataSize();
+							
+							AddedShaderChunkIds.Add(ShaderChunkId);
+						}
+					}
+					ShaderChunkIdsByShaderMapHash.Add(ShaderMap.Key, MoveTemp(ShaderMap.Value));
+				}
+			}
+		}
+
+		check(ShaderTargetFiles.Num() == Shaders.Num());
+		
+		if (ShaderTargetFiles.Num() == 0)
+		{
+			continue;
+		}
+
+		ContainerTarget->TargetFiles.Append(ShaderTargetFiles);
+		ContainerTarget->Shaders.Append(Shaders);
+
+		TMap<FIoChunkId, FShaderInfo*> ShaderInfoByChunkIdMap;
+		ShaderInfoByChunkIdMap.Reserve(ContainerTarget->Shaders.Num());
+		for (FShaderInfo& ShaderInfo : ContainerTarget->Shaders)
+		{
+			ShaderInfoByChunkIdMap.Add(ShaderInfo.ChunkId, &ShaderInfo);
+		}
+
+		for (FLegacyCookedPackage* Package : ContainerTarget->Packages)
+		{
+			for (const FSHAHash& ShaderMapHash : Package->OptimizedPackage->GetShaderMapHashes())
+			{
+				const TArray<FIoChunkId>* FindChunkIds = ShaderChunkIdsByShaderMapHash.Find(ShaderMapHash);
+				if (!FindChunkIds)
+				{
+					UE_LOG(LogIoStore, Warning, TEXT("Package referencing unknown shader map"));
+					continue;
+				}
+				for (const FIoChunkId& ShaderChunkId : *FindChunkIds)
+				{
+					FShaderInfo* ShaderInfo = ShaderInfoByChunkIdMap.FindRef(ShaderChunkId);
+					check(ShaderInfo);
+					Package->Shaders.Add(ShaderInfo);
+					++ShaderInfo->ReferencedByPackagesCount;
+				}
+			}
+		}
+	}
+}
+
 void InitializeContainerTargetsAndPackages(
 	const FIoStoreArguments& Arguments,
 	TArray<FLegacyCookedPackage*>& Packages,
@@ -960,6 +1196,31 @@ void InitializeContainerTargetsAndPackages(
 				}
 
 				FString RelativeFileName = ConvertCookedPathToRelativePath(SourceFile.NormalizedPath);
+				if (CookedFileStatData->FileType == FCookedFileStatData::ShaderLibrary)
+				{
+					FContainerTargetFile& TargetFile = ContainerTarget->TargetFiles.AddDefaulted_GetRef();
+					TargetFile.ChunkType = EContainerChunkType::ShaderCodeLibrary;
+					TargetFile.ContainerTarget = ContainerTarget;
+					TargetFile.SourceSize = uint64(CookedFileStatData->FileSize);
+					TargetFile.NormalizedSourcePath = NormalizedSourcePath;
+					TargetFile.TargetPath = MoveTemp(RelativeFileName);
+					TargetFile.DestinationPath = SourceFile.DestinationPath;
+					if (SourceFile.bNeedsCompression)
+					{
+						ContainerTarget->ContainerFlags |= EIoContainerFlags::Compressed;
+					}
+					else
+					{
+						TargetFile.bForceUncompressed = true;
+					}
+					if (SourceFile.bNeedsEncryption)
+					{
+						ContainerTarget->ContainerFlags |= EIoContainerFlags::Encrypted;
+					}
+					continue;
+				}
+				
+				
 				FLegacyCookedPackage* Package = nullptr;
 				const bool bIsMemoryMappedBulkData = CookedFileStatData->FileExt == FCookedFileStatData::EFileExt::UMappedBulk;
 
@@ -997,20 +1258,20 @@ void InitializeContainerTargetsAndPackages(
 
 					if (CookedFileStatData->FileType == FCookedFileStatData::BulkData)
 					{
-						TargetFile.bIsBulkData = true;
 						if (CookedFileStatData->FileExt == FCookedFileStatData::UPtnl)
 						{
-							TargetFile.bIsOptionalBulkData = true;
+							TargetFile.ChunkType = EContainerChunkType::OptionalBulkData;
 							TargetFile.ChunkId = CreateChunkId(Package->GlobalPackageId, 0, EIoChunkType::OptionalBulkData, *TargetFile.TargetPath);
 						}
 						else if (CookedFileStatData->FileExt == FCookedFileStatData::UMappedBulk)
 						{
-							TargetFile.bIsMemoryMappedBulkData = true;
+							TargetFile.ChunkType = EContainerChunkType::MemoryMappedBulkData;
 							TargetFile.bForceUncompressed = true;
 							TargetFile.ChunkId = CreateChunkId(Package->GlobalPackageId, 0, EIoChunkType::MemoryMappedBulkData, *TargetFile.TargetPath);
 						}
 						else
 						{
+							TargetFile.ChunkType = EContainerChunkType::BulkData;
 							TargetFile.ChunkId = CreateChunkId(Package->GlobalPackageId, 0, EIoChunkType::BulkData, *TargetFile.TargetPath);
 						}
 						if (Package->FileName.IsEmpty())
@@ -1021,8 +1282,8 @@ void InitializeContainerTargetsAndPackages(
 					else
 					{
 						check(CookedFileStatData->FileType == FCookedFileStatData::PackageData);
-
-						++ContainerTarget->PackageCount;
+						TargetFile.ChunkType = EContainerChunkType::PackageData;
+						ContainerTarget->Packages.Add(Package);
 
 						Package->FileName = SourceFile.NormalizedPath; // .uasset path
 						Package->UAssetSize = OriginalCookedFileStatData->FileSize;
@@ -1141,7 +1402,7 @@ void LogContainerPackageInfo(const TArray<FContainerTargetSpec*>& ContainerTarge
 	for (const FContainerTargetSpec* ContainerTarget : ContainerTargets)
 	{
 		uint64 StoreSize = ContainerTarget->Header.StoreEntries.Num();
-		uint64 PackageCount = ContainerTarget->PackageCount;
+		uint64 PackageCount = ContainerTarget->Packages.Num();
 		uint64 LocalizedPackageCount = 0;
 
 		for (const auto& KV : ContainerTarget->Header.CulturePackageMap)
@@ -1189,23 +1450,18 @@ void LogContainerPackageInfo(const TArray<FContainerTargetSpec*>& ContainerTarge
 	for (const FContainerTargetSpec* ContainerTarget : ContainerTargets)
 	{
 		uint64 HeaderSize = 0;
-		uint64 SummarySize = ContainerTarget->PackageCount * sizeof(FPackageSummary);
+		uint64 SummarySize = ContainerTarget->Packages.Num() * sizeof(FPackageSummary);
 		uint64 UGraphSize = 0;
 		uint64 ImportMapSize = 0;
 		uint64 ExportMapSize = 0;
 		uint64 NameMapSize = 0;
 
-		for (const FContainerTargetFile& TargetFile : ContainerTarget->TargetFiles)
+		for (const FLegacyCookedPackage* Package : ContainerTarget->Packages)
 		{
-			if (TargetFile.bIsBulkData)
-			{
-				continue;
-			}
-
-			UGraphSize += TargetFile.Package->OptimizedPackage->GetExportBundleArcsSize();
-			ImportMapSize += TargetFile.Package->OptimizedPackage->GetImportMapSize();
-			ExportMapSize += TargetFile.Package->OptimizedPackage->GetExportMapSize();
-			NameMapSize += TargetFile.Package->OptimizedPackage->GetNameMapSize();
+			UGraphSize += Package->OptimizedPackage->GetExportBundleArcsSize();
+			ImportMapSize += Package->OptimizedPackage->GetImportMapSize();
+			ExportMapSize += Package->OptimizedPackage->GetExportMapSize();
+			NameMapSize += Package->OptimizedPackage->GetNameMapSize();
 		}
 
 		HeaderSize = SummarySize + UGraphSize + ImportMapSize + ExportMapSize + NameMapSize;
@@ -1324,7 +1580,7 @@ private:
 		void AsyncReadCallback()
 		{
 			//TRACE_CPUPROFILER_EVENT_SCOPE(ReadCallback);
-			if (!TargetFile.bIsBulkData)
+			if (TargetFile.ChunkType == EContainerChunkType::PackageData)
 			{
 				*SourceBuffer = Manager.PackageStoreOptimizer.CreatePackageBuffer(TargetFile.Package->OptimizedPackage, *SourceBuffer, bHasUpdatedExportBundleRegions ? nullptr : &FileRegions);
 				bHasUpdatedExportBundleRegions = true;
@@ -1608,6 +1864,9 @@ int32 CreateTarget(const FIoStoreArguments& Arguments, const FIoStoreWriterSetti
 
 	ParsePackageAssets(Packages, PackageStoreOptimizer);
 
+	UE_LOG(LogIoStore, Display, TEXT("Processing shader libraries..."));
+	ProcessShaderLibraries(Arguments, ContainerTargets);
+
 	if (Arguments.IsDLC() && Arguments.bRemapPluginContentToGame)
 	{
 		for (FLegacyCookedPackage* Package : Packages)
@@ -1638,14 +1897,21 @@ int32 CreateTarget(const FIoStoreArguments& Arguments, const FIoStoreWriterSetti
 		{
 			for (FContainerTargetFile& TargetFile : ContainerTarget->TargetFiles)
 			{
-				if (TargetFile.bIsBulkData)
+				if (TargetFile.ChunkType != EContainerChunkType::PackageData)
 				{
 					FIoWriteOptions WriteOptions;
 					WriteOptions.DebugName = *TargetFile.TargetPath;
 					WriteOptions.bForceUncompressed = TargetFile.bForceUncompressed;
-					WriteOptions.bIsMemoryMapped = TargetFile.bIsMemoryMappedBulkData;
+					WriteOptions.bIsMemoryMapped = TargetFile.ChunkType == EContainerChunkType::MemoryMappedBulkData;
 					WriteOptions.FileName = TargetFile.DestinationPath;
-					ContainerTarget->IoStoreWriter->Append(TargetFile.ChunkId, WriteRequestManager.Read(TargetFile), WriteOptions);
+					if (TargetFile.SourceBuffer.IsSet())
+					{
+						ContainerTarget->IoStoreWriter->Append(TargetFile.ChunkId, *TargetFile.SourceBuffer, WriteOptions, TargetFile.IdealOrder);
+					}
+					else
+					{
+						ContainerTarget->IoStoreWriter->Append(TargetFile.ChunkId, WriteRequestManager.Read(TargetFile), WriteOptions);
+					}
 				}
 			}
 		}
@@ -1669,14 +1935,21 @@ int32 CreateTarget(const FIoStoreArguments& Arguments, const FIoStoreWriterSetti
 			TArray<FPackageStoreContainerHeaderEntry> ContainerHeaderEntries;
 			for (FContainerTargetFile& TargetFile : ContainerTarget->TargetFiles)
 			{
-				check(TargetFile.Package);
-				if (!TargetFile.bIsBulkData)
+				if (TargetFile.ChunkType == EContainerChunkType::PackageData)
 				{
+					check(TargetFile.Package);
 					FIoWriteOptions WriteOptions;
 					WriteOptions.DebugName = *TargetFile.TargetPath;
 					WriteOptions.bForceUncompressed = TargetFile.bForceUncompressed;
 					WriteOptions.FileName = TargetFile.DestinationPath;
-					ContainerTarget->IoStoreWriter->Append(TargetFile.ChunkId, WriteRequestManager.Read(TargetFile), WriteOptions);
+					if (TargetFile.SourceBuffer.IsSet())
+					{
+						ContainerTarget->IoStoreWriter->Append(TargetFile.ChunkId, *TargetFile.SourceBuffer, WriteOptions, TargetFile.IdealOrder);
+					}
+					else
+					{
+						ContainerTarget->IoStoreWriter->Append(TargetFile.ChunkId, WriteRequestManager.Read(TargetFile), WriteOptions);
+					}
 					ContainerHeaderEntries.Add(PackageStoreOptimizer.CreateContainerHeaderEntry(TargetFile.Package->OptimizedPackage));
 				}
 			}
@@ -1765,6 +2038,28 @@ int32 CreateTarget(const FIoStoreArguments& Arguments, const FIoStoreWriterSetti
 		ImportedPackagesCount += PackageImportedPackagesCount;
 		NoImportedPackagesCount += PackageImportedPackagesCount == 0;
 	}
+	
+	uint64 GlobalShaderSize = 0;
+	uint64 SharedShaderSize = 0;
+	uint64 UniqueShaderSize = 0;
+	for (const FContainerTargetSpec* ContainerTarget : ContainerTargets)
+	{
+		for (const FShaderInfo& ShaderInfo : ContainerTarget->Shaders)
+		{
+			if (ShaderInfo.ReferencedByPackagesCount == 0)
+			{
+				GlobalShaderSize += ShaderInfo.Size;
+			}
+			else if (ShaderInfo.ReferencedByPackagesCount == 1)
+			{
+				UniqueShaderSize += ShaderInfo.Size;
+			}
+			else
+			{
+				SharedShaderSize += ShaderInfo.Size;
+			}
+		}
+	}
 
 	LogWriterResults(IoStoreWriterResults);
 	LogContainerPackageInfo(ContainerTargets);
@@ -1772,6 +2067,9 @@ int32 CreateTarget(const FIoStoreArguments& Arguments, const FIoStoreWriterSetti
 	UE_LOG(LogIoStore, Display, TEXT("Input:  %8.2lf MB UExp"), (double)UExpSize / 1024.0 / 1024.0);
 	UE_LOG(LogIoStore, Display, TEXT("Input:  %8.2lf MB UAsset"), (double)UAssetSize / 1024.0 / 1024.0);
 	UE_LOG(LogIoStore, Display, TEXT("Input:  %8d Packages"), Packages.Num());
+	UE_LOG(LogIoStore, Display, TEXT("Input:  %8.2f MB Global shaders"), (double)GlobalShaderSize / 1024.0 / 1024.0);
+	UE_LOG(LogIoStore, Display, TEXT("Input:  %8.2f MB Shared shaders"), (double)SharedShaderSize / 1024.0 / 1024.0);
+	UE_LOG(LogIoStore, Display, TEXT("Input:  %8.2f MB Unique shaders"), (double)UniqueShaderSize / 1024.0 / 1024.0);
 	UE_LOG(LogIoStore, Display, TEXT(""));
 	UE_LOG(LogIoStore, Display, TEXT("Output: %8llu Export bundles"), PackageStoreOptimizer.GetTotalExportBundleCount());
 	UE_LOG(LogIoStore, Display, TEXT("Output: %8llu Export bundle entries"), PackageStoreOptimizer.GetTotalExportBundleEntryCount());
@@ -3232,9 +3530,10 @@ public:
 	virtual bool Visit(const TCHAR* FilenameOrDirectory, const FFileStatData& StatData)
 	{
 		// Should match FCookedFileStatData::EFileExt
-		static const TCHAR* Extensions[] = { TEXT("umap"), TEXT("uasset"), TEXT("uexp"), TEXT("ubulk"), TEXT("uptnl"), TEXT("m.ubulk") };
+		static const TCHAR* Extensions[] = { TEXT("umap"), TEXT("uasset"), TEXT("uexp"), TEXT("ubulk"), TEXT("uptnl"), TEXT("m.ubulk"), TEXT("ushaderbytecode") };
 		static const int32 NumPackageExtensions = 2;
 		static const int32 UExpExtensionIndex = 2;
+		static const int32 UShaderByteCodeExtensionIndex = 6;
 
 		if (StatData.bIsDirectory)
 		{
@@ -3297,6 +3596,10 @@ public:
 		{
 			CookedFileStatData.FileType = FCookedFileStatData::PackageData;
 		}
+		else if (ExtIndex == UShaderByteCodeExtensionIndex)
+		{
+			CookedFileStatData.FileType = FCookedFileStatData::ShaderLibrary;
+		}
 		else
 		{
 			CookedFileStatData.FileType = FCookedFileStatData::BulkData;
diff --git a/Engine/Source/Developer/IoStoreUtilities/Private/PackageStoreOptimizer.cpp b/Engine/Source/Developer/IoStoreUtilities/Private/PackageStoreOptimizer.cpp
--- a/Engine/Source/Developer/IoStoreUtilities/Private/PackageStoreOptimizer.cpp
+++ b/Engine/Source/Developer/IoStoreUtilities/Private/PackageStoreOptimizer.cpp
@@ -12,6 +12,7 @@
 #include "UObject/Object.h"
 #include "Serialization/MemoryReader.h"
 #include "UObject/Package.h"
+#include "Misc/SecureHash.h"
 
 DEFINE_LOG_CATEGORY_STATIC(LogPackageStoreOptimizer, Log, All);
 
@@ -1566,6 +1567,7 @@ FPackageStoreContainerHeaderEntry FPackageStoreOptimizer::CreateContainerHeaderE
 	Result.ExportBundleCount = Package->ExportBundles.Num();
 	Result.LoadOrder = Package->GetLoadOrder();
 	Result.ImportedPackageIds = Package->ImportedPackageIds;
+	Result.ShaderMapHashes = Package->ShaderMapHashes.Array();
 	Result.bIsRedirected = Package->bIsRedirected;
 	return Result;
 }
@@ -1638,6 +1640,14 @@ FContainerHeader FPackageStoreOptimizer::CreateContainerHeader(const FIoContaine
 			check(ImportedPackageId.IsValid());
 			StoreDataArchive << ImportedPackageId;
 		}
+		
+		// ShaderMapHashes
+		const TArray<FSHAHash>& ShaderMapHashes = Package->ShaderMapHashes;
+		SerializePackageEntryCArrayHeader(ShaderMapHashes.Num());
+		for (const FSHAHash& ShaderMapHash : ShaderMapHashes)
+		{
+			StoreDataArchive << const_cast<FSHAHash&>(ShaderMapHash);
+		}
 	}
 
 	check(StoreTocArchive.TotalSize() == StoreTocSize);
diff --git a/Engine/Source/Developer/IoStoreUtilities/Private/PackageStoreOptimizer.h b/Engine/Source/Developer/IoStoreUtilities/Private/PackageStoreOptimizer.h
--- a/Engine/Source/Developer/IoStoreUtilities/Private/PackageStoreOptimizer.h
+++ b/Engine/Source/Developer/IoStoreUtilities/Private/PackageStoreOptimizer.h
@@ -104,6 +104,7 @@ private:
 	int32 ExportBundleCount = 0;
 	uint32 LoadOrder = 0;
 	TArray<FPackageId> ImportedPackageIds;
+	TArray<FSHAHash> ShaderMapHashes;
 	bool bIsRedirected = false;
 
 	friend class FPackageStoreOptimizer;
@@ -172,6 +173,16 @@ public:
 		SourceName = SourcePackageName;
 	}
 
+	void AddShaderMapHash(const FSHAHash& ShaderMapHash)
+	{
+		ShaderMapHashes.Add(ShaderMapHash);
+	}
+
+	const TSet<FSHAHash>& GetShaderMapHashes() const
+	{
+		return ShaderMapHashes;
+	}
+
 private:
 	struct FExportGraphNode;
 
@@ -262,6 +273,7 @@ private:
 	TArray<FPackageStorePackage*> ImportedByWaitingPackages;
 	TSet<FPackageId> RedirectedToPackageIds;
 	TSet<FPackageId> ImportedRedirectedPackageIds;
+	TSet<FSHAHash> ShaderMapHashes;
 
 	FIoBuffer OptimizedHeaderBuffer;
 
diff --git a/Engine/Source/Programs/AutomationTool/Scripts/CopyBuildToStagingDirectory.Automation.cs b/Engine/Source/Programs/AutomationTool/Scripts/CopyBuildToStagingDirectory.Automation.cs
--- a/Engine/Source/Programs/AutomationTool/Scripts/CopyBuildToStagingDirectory.Automation.cs
+++ b/Engine/Source/Programs/AutomationTool/Scripts/CopyBuildToStagingDirectory.Automation.cs
@@ -2736,6 +2736,13 @@ public partial class Project : CommandUtils
 							bAllowBulkDataInIoStore = true; // Default is to allow it in the IoStore
 						}
 
+						bool bAllowShadersInIoStore;
+						if (!PlatformEngineConfig.GetBool("Core.System", "AllowShadersInIoStore", out bAllowShadersInIoStore))
+						{
+							bAllowShadersInIoStore = false;
+						}
+
+
 						UnrealPakResponseFile = new Dictionary<string, string>();
 						Dictionary<string, string> IoStoreResponseFile = new Dictionary<string, string>();
 						foreach (var Entry in PakParams.UnrealPakResponseFile)
@@ -2764,6 +2771,17 @@ public partial class Project : CommandUtils
 									UnrealPakResponseFile.Add(Entry.Key, Entry.Value);
 								}
 							}
+							else if (Path.GetExtension(Entry.Key).Contains(".ushaderbytecode"))
+							{
+								if (bAllowShadersInIoStore)
+								{
+									IoStoreResponseFile.Add(Entry.Key, Entry.Value);
+								}
+								else
+								{
+									UnrealPakResponseFile.Add(Entry.Key, Entry.Value);
+								}
+							}
 							else if (!Path.GetExtension(Entry.Key).Contains(".uexp"))
 							{
 								// all other extensions go in Pak not iostore
diff --git a/Engine/Source/Runtime/Core/Private/IO/IoStore.cpp b/Engine/Source/Runtime/Core/Private/IO/IoStore.cpp
--- a/Engine/Source/Runtime/Core/Private/IO/IoStore.cpp
+++ b/Engine/Source/Runtime/Core/Private/IO/IoStore.cpp
@@ -1352,12 +1352,13 @@ void FIoStoreWriter::EnableDiskLayoutOrdering(const TArray<TUniquePtr<FIoStoreRe
 	Impl->EnableDiskLayoutOrdering(PatchSourceReaders);
 }
 
-void FIoStoreWriter::Append(const FIoChunkId& ChunkId, FIoBuffer Chunk, const FIoWriteOptions& WriteOptions)
+void FIoStoreWriter::Append(const FIoChunkId& ChunkId, FIoBuffer Chunk, const FIoWriteOptions& WriteOptions, uint64 OrderHint)
 {
 	struct FWriteRequest
 		: IIoStoreWriteRequest
 	{
-		FWriteRequest(FIoBuffer InSourceBuffer)
+		FWriteRequest(FIoBuffer InSourceBuffer, uint64 InOrderHint)
+			: OrderHint(InOrderHint)
 		{
 			SourceBuffer = InSourceBuffer;
 			SourceBuffer.MakeOwned();
@@ -1381,7 +1382,7 @@ void FIoStoreWriter::Append(const FIoChunkId& ChunkId, FIoBuffer Chunk, const FI
 
 		uint64 GetOrderHint() override
 		{
-			return MAX_uint64;
+			return OrderHint;
 		}
 
 		TArrayView<const FFileRegion> GetRegions()
@@ -1390,9 +1391,10 @@ void FIoStoreWriter::Append(const FIoChunkId& ChunkId, FIoBuffer Chunk, const FI
 		}
 
 		FIoBuffer SourceBuffer;
+		uint64 OrderHint;
 	};
 
-	Append(ChunkId, new FWriteRequest(Chunk), WriteOptions);
+	Append(ChunkId, new FWriteRequest(Chunk, OrderHint), WriteOptions);
 }
 
 void FIoStoreWriter::Append(const FIoChunkId& ChunkId, IIoStoreWriteRequest* Request, const FIoWriteOptions& WriteOptions)
diff --git a/Engine/Source/Runtime/Core/Private/Misc/CoreDelegates.cpp b/Engine/Source/Runtime/Core/Private/Misc/CoreDelegates.cpp
--- a/Engine/Source/Runtime/Core/Private/Misc/CoreDelegates.cpp
+++ b/Engine/Source/Runtime/Core/Private/Misc/CoreDelegates.cpp
@@ -192,6 +192,9 @@ FCoreDelegates::FOnGCFinishDestroyTimeExtended FCoreDelegates::OnGCFinishDestroy
 
 FCoreDelegates::FAccesExtraBinaryConfigData FCoreDelegates::AccessExtraBinaryConfigData;
 
+FCoreDelegates::FPreloadPackageShaderMaps FCoreDelegates::PreloadPackageShaderMaps;
+FCoreDelegates::FReleasePreloadedPackageShaderMaps FCoreDelegates::ReleasePreloadedPackageShaderMaps;
+
 FCoreDelegates::FApplicationNetworkInitializationChanged FCoreDelegates::ApplicationNetworkInitializationChanged;
 
 PRAGMA_DISABLE_DEPRECATION_WARNINGS
diff --git a/Engine/Source/Runtime/Core/Public/IO/IoDispatcher.h b/Engine/Source/Runtime/Core/Public/IO/IoDispatcher.h
--- a/Engine/Source/Runtime/Core/Public/IO/IoDispatcher.h
+++ b/Engine/Source/Runtime/Core/Public/IO/IoDispatcher.h
@@ -704,7 +704,9 @@ enum class EIoChunkType : uint8
 	LoaderInitialLoadMeta,
 	LoaderGlobalNames,
 	LoaderGlobalNameHashes,
-	ContainerHeader
+	ContainerHeader,
+	ShaderCodeLibrary,
+	ShaderCode
 };
 
 /**
@@ -1150,7 +1152,7 @@ public:
 
 	UE_NODISCARD CORE_API FIoStatus	Initialize(const FIoStoreWriterContext& Context, const FIoContainerSettings& ContainerSettings);
 	CORE_API void EnableDiskLayoutOrdering(const TArray<TUniquePtr<FIoStoreReader>>& PatchSourceReaders = TArray<TUniquePtr<FIoStoreReader>>());
-	CORE_API void Append(const FIoChunkId& ChunkId, FIoBuffer Chunk, const FIoWriteOptions& WriteOptions);
+	CORE_API void Append(const FIoChunkId& ChunkId, FIoBuffer Chunk, const FIoWriteOptions& WriteOptions, uint64 OrderHint = MAX_uint64);
 	CORE_API void Append(const FIoChunkId& ChunkId, IIoStoreWriteRequest* Request, const FIoWriteOptions& WriteOptions);
 	UE_NODISCARD CORE_API TIoStatusOr<FIoStoreWriterResult> Flush();
 
diff --git a/Engine/Source/Runtime/Core/Public/Misc/CoreDelegates.h b/Engine/Source/Runtime/Core/Public/Misc/CoreDelegates.h
--- a/Engine/Source/Runtime/Core/Public/Misc/CoreDelegates.h
+++ b/Engine/Source/Runtime/Core/Public/Misc/CoreDelegates.h
@@ -10,6 +10,7 @@
 #include "Math/IntVector.h"
 #include "Misc/AES.h"
 #include "GenericPlatform/GenericPlatformFile.h"
+#include "Async/TaskGraphInterfaces.h"
 
 class AActor;
 class Error;
@@ -618,6 +619,12 @@ public:
 	DECLARE_MULTICAST_DELEGATE_OneParam(FAccesExtraBinaryConfigData, FExtraBinaryConfigData&);
 	static FAccesExtraBinaryConfigData AccessExtraBinaryConfigData;
 
+	using FAttachShaderReadRequestFunc = TFunctionRef<class FIoRequest(const class FIoChunkId&, FGraphEventRef)>;
+	DECLARE_DELEGATE_TwoParams(FPreloadPackageShaderMaps, TArrayView<const FSHAHash>, FAttachShaderReadRequestFunc);
+	static FPreloadPackageShaderMaps PreloadPackageShaderMaps;
+	DECLARE_DELEGATE_OneParam(FReleasePreloadedPackageShaderMaps, TArrayView<const FSHAHash>);
+	static FReleasePreloadedPackageShaderMaps ReleasePreloadedPackageShaderMaps;
+
 private:
 
 	// Callbacks for hotfixes
diff --git a/Engine/Source/Runtime/CoreUObject/Private/Serialization/AsyncLoading2.cpp b/Engine/Source/Runtime/CoreUObject/Private/Serialization/AsyncLoading2.cpp
--- a/Engine/Source/Runtime/CoreUObject/Private/Serialization/AsyncLoading2.cpp
+++ b/Engine/Source/Runtime/CoreUObject/Private/Serialization/AsyncLoading2.cpp
@@ -1953,7 +1953,7 @@ private:
 	TMap<FPackageId, FAsyncPackage2*> AsyncPackageLookup;
 
 	TQueue<FAsyncPackage2*, EQueueMode::Mpsc> ExternalReadQueue;
-	TAtomic<int32> WaitingForIoBundleCounter{ 0 };
+	TAtomic<int32> PendingIoRequestsCounter{ 0 };
 	
 	/** List of all pending package requests */
 	TSet<int32> PendingRequests;
@@ -2590,8 +2590,8 @@ bool FAsyncLoadingThread2::CreateAsyncPackagesFromQueue(FAsyncLoadingThreadState
 
 void FAsyncLoadingThread2::AddBundleIoRequest(FAsyncPackage2* Package)
 {
-	int32 LocalWaitingForIoBundleCount = WaitingForIoBundleCounter.IncrementExchange() + 1;
-	TRACE_COUNTER_SET(AsyncLoadingPendingIoRequests, LocalWaitingForIoBundleCount);
+	int32 LocalPendingIoRequestsCounter = PendingIoRequestsCounter.IncrementExchange() + 1;
+	TRACE_COUNTER_SET(AsyncLoadingPendingIoRequests, LocalPendingIoRequestsCounter);
 	WaitingIoRequests.HeapPush({ Package });
 }
 
@@ -2639,11 +2639,31 @@ void FAsyncLoadingThread2::StartBundleIoRequests()
 					TEXT("Failed reading chunk for package: %s"), *Result.Status().ToString());
 				Package->bLoadHasFailed = true;
 			}
-			int32 LocalWaitingForIoBundleCounter = Package->AsyncLoadingThread.WaitingForIoBundleCounter.DecrementExchange() - 1;
-			TRACE_COUNTER_SET(AsyncLoadingPendingIoRequests, LocalWaitingForIoBundleCounter);
+			int32 LocalPendingIoRequestsCounter = Package->AsyncLoadingThread.PendingIoRequestsCounter.DecrementExchange() - 1;
+			TRACE_COUNTER_SET(AsyncLoadingPendingIoRequests, LocalPendingIoRequestsCounter);
 			TRACE_COUNTER_ADD(AsyncLoadingTotalLoaded, Package->IoBuffer.DataSize());
 			Package->GetPackageNode(EEventLoadNode2::Package_ProcessSummary).ReleaseBarrier();
 		});
+
+		if (Package->Desc.StoreEntry->ShaderMapHashes.Num())
+		{
+			TRACE_CPUPROFILER_EVENT_SCOPE(StartShaderMapRequests);
+			auto ReadShaderMapFunc = [Package, &IoBatch](const FIoChunkId& ChunkId, FGraphEventRef GraphEvent)
+			{
+				Package->GetPackageNode(Package_ExportsSerialized).AddBarrier();
+				int32 LocalPendingIoRequestsCounter = Package->AsyncLoadingThread.PendingIoRequestsCounter.IncrementExchange() + 1;
+				TRACE_COUNTER_SET(AsyncLoadingPendingIoRequests, LocalPendingIoRequestsCounter);
+				return IoBatch.ReadWithCallback(ChunkId, FIoReadOptions(), Package->Desc.Priority,
+					[Package, GraphEvent](TIoStatusOr<FIoBuffer> Result)
+					{
+						GraphEvent->DispatchSubsequents();
+						int32 LocalPendingIoRequestsCounter = Package->AsyncLoadingThread.PendingIoRequestsCounter.DecrementExchange() - 1;
+						TRACE_COUNTER_SET(AsyncLoadingPendingIoRequests, LocalPendingIoRequestsCounter);
+						Package->GetPackageNode(Package_ExportsSerialized).ReleaseBarrier();
+					});
+			};
+			FCoreDelegates::PreloadPackageShaderMaps.ExecuteIfBound(MakeArrayView(Package->Desc.StoreEntry->ShaderMapHashes.Data(), Package->Desc.StoreEntry->ShaderMapHashes.Num()), ReadShaderMapFunc);
+		}
 	}
 	IoBatch.Issue();
 
@@ -3798,6 +3818,12 @@ EAsyncPackageState::Type FAsyncPackage2::Event_ExportsDone(FAsyncLoadingThreadSt
 		PackageRef.SetAllPublicExportsLoaded();
 	}
 
+	if (Package->Desc.StoreEntry->ShaderMapHashes.Num())
+	{
+		TRACE_CPUPROFILER_EVENT_SCOPE(ReleasePreloadedShaderMaps);
+		FCoreDelegates::ReleasePreloadedPackageShaderMaps.ExecuteIfBound(MakeArrayView(Package->Desc.StoreEntry->ShaderMapHashes.Data(), Package->Desc.StoreEntry->ShaderMapHashes.Num()));
+	}
+
 	Package->AsyncPackageLoadingState = EAsyncPackageLoadingState2::PostLoad;
 	Package->GetExportBundleNode(EEventLoadNode2::ExportBundle_PostLoad, 0).ReleaseBarrier();
 	return EAsyncPackageState::Complete;
@@ -4752,7 +4778,7 @@ uint32 FAsyncLoadingThread2::Run()
 			if (!bDidSomething)
 			{
 				FAsyncPackage2* Package = nullptr;
-				if (WaitingForIoBundleCounter.Load() > 0)
+				if (PendingIoRequestsCounter.Load() > 0)
 				{
 					TRACE_CPUPROFILER_EVENT_SCOPE(AsyncLoadingTime);
 					TRACE_CPUPROFILER_EVENT_SCOPE(WaitingForIo);
diff --git a/Engine/Source/Runtime/CoreUObject/Public/Serialization/PackageStore.h b/Engine/Source/Runtime/CoreUObject/Public/Serialization/PackageStore.h
--- a/Engine/Source/Runtime/CoreUObject/Public/Serialization/PackageStore.h
+++ b/Engine/Source/Runtime/CoreUObject/Public/Serialization/PackageStore.h
@@ -44,6 +44,7 @@ struct FPackageStoreEntry
 	uint32 LoadOrder;
 	uint32 Pad;
 	TPackageStoreEntryCArrayView<FPackageId> ImportedPackages;
+	TPackageStoreEntryCArrayView<FSHAHash> ShaderMapHashes;
 };
 
 /**
diff --git a/Engine/Source/Runtime/RHI/Public/RHIResources.h b/Engine/Source/Runtime/RHI/Public/RHIResources.h
--- a/Engine/Source/Runtime/RHI/Public/RHIResources.h
+++ b/Engine/Source/Runtime/RHI/Public/RHIResources.h
@@ -24,6 +24,7 @@
 #include "Serialization/MemoryImage.h"
 #include "Experimental/Containers/HazardPointer.h"
 #include "Containers/ClosableMpscQueue.h"
+#include "Misc/CoreDelegates.h"
 
 #ifndef RHI_WANT_RESOURCE_INFO
 #define RHI_WANT_RESOURCE_INFO 0
@@ -2708,6 +2709,7 @@ public:
 	virtual int32 FindShaderIndex(const FSHAHash& Hash) = 0;
 	virtual bool PreloadShader(int32 ShaderIndex, FGraphEventArray& OutCompletionEvents) { return false; }
 	virtual bool PreloadShaderMap(int32 ShaderMapIndex, FGraphEventArray& OutCompletionEvents) { return false; }
+	virtual bool PreloadShaderMap(int32 ShaderMapIndex, FCoreDelegates::FAttachShaderReadRequestFunc AttachShaderReadRequestFunc) { return false; }
 	virtual void ReleasePreloadedShader(int32 ShaderIndex) {}
 
 	virtual TRefCountPtr<FRHIShader> CreateShader(int32 ShaderIndex) { return nullptr; }
diff --git a/Engine/Source/Runtime/RenderCore/Private/ShaderCodeArchive.cpp b/Engine/Source/Runtime/RenderCore/Private/ShaderCodeArchive.cpp
--- a/Engine/Source/Runtime/RenderCore/Private/ShaderCodeArchive.cpp
+++ b/Engine/Source/Runtime/RenderCore/Private/ShaderCodeArchive.cpp
@@ -9,11 +9,14 @@
 #include "Policies/PrettyJsonPrintPolicy.h"
 #include "Serialization/JsonSerializer.h"
 #include "Misc/FileHelper.h"
+#include "Serialization/MemoryReader.h"
 #if WITH_EDITOR
 #include "Misc/StringBuilder.h"
 #include "Templates/Greater.h"
 #endif
 
+uint32 GIoStoreShaderCodeArchiveVersion = 1;
+
 int32 GShaderCodeLibraryAsyncLoadingPriority = int32(AIOP_Normal);
 static FAutoConsoleVariableRef CVarShaderCodeLibraryAsyncLoadingPriority(
 	TEXT("r.ShaderCodeLibrary.DefaultAsyncIOPriority"),
@@ -30,6 +33,23 @@ static FAutoConsoleVariableRef CVarShaderCodeLibraryAsyncLoadingAllowDontCache(
 	ECVF_Default
 );
 
+FIoChunkId GetShaderCodeArchiveChunkId(const FString& LibraryName, FName FormatName)
+{
+	FString Name = FString::Printf(TEXT("%s-%s"), *LibraryName, *FormatName.ToString());
+	Name.ToLowerInline();
+	uint64 Hash = CityHash64(reinterpret_cast<const char*>(*Name), Name.Len() * sizeof(TCHAR));
+	return CreateIoChunkId(Hash, 0, EIoChunkType::ShaderCodeLibrary);
+}
+
+FIoChunkId GetShaderCodeChunkId(const FSHAHash& ShaderHash)
+{
+	uint8 Data[12];
+	FMemory::Memcpy(Data, ShaderHash.Hash, 11);
+	*reinterpret_cast<uint8*>(&Data[11]) = static_cast<uint8>(EIoChunkType::ShaderCode);
+	FIoChunkId ChunkId;
+	ChunkId.Set(Data, 12);
+	return ChunkId;
+}
 
 int32 FSerializedShaderArchive::FindShaderMapWithKey(const FSHAHash& Hash, uint32 Key) const
 {
@@ -997,3 +1017,309 @@ TRefCountPtr<FRHIShader> FShaderCodeArchive::CreateShader(int32 Index)
 
 	return Shader;
 }
+
+FIoStoreShaderCodeArchive* FIoStoreShaderCodeArchive::Create(EShaderPlatform InPlatform, const FString& InLibraryName, FIoDispatcher& InIoDispatcher)
+{
+	const FName PlatformName = LegacyShaderPlatformToShaderFormat(InPlatform);
+	FIoChunkId ChunkId = GetShaderCodeArchiveChunkId(InLibraryName, PlatformName);
+	if (InIoDispatcher.DoesChunkExist(ChunkId))
+	{
+		FIoBatch IoBatch = InIoDispatcher.NewBatch();
+		FIoRequest IoRequest = IoBatch.Read(ChunkId, FIoReadOptions(), IoDispatcherPriority_Max);
+		FEvent* Event = FPlatformProcess::GetSynchEventFromPool();
+		IoBatch.IssueAndTriggerEvent(Event);
+		Event->Wait();
+		FPlatformProcess::ReturnSynchEventToPool(Event);
+		const FIoBuffer& IoBuffer = IoRequest.GetResult().ValueOrDie();
+		FMemoryReaderView Ar(MakeArrayView(IoBuffer.Data(), IoBuffer.DataSize()));
+		uint32 Version = 0;
+		Ar << Version;
+		if (Version == GIoStoreShaderCodeArchiveVersion)
+		{
+			FIoStoreShaderCodeArchive* Library = new FIoStoreShaderCodeArchive(InPlatform, InLibraryName, InIoDispatcher);
+			Ar << Library->ShaderMapHashes;
+			Ar << Library->ShaderHashes;
+			Ar << Library->ShaderMapEntries;
+			Ar << Library->ShaderEntries;
+			Ar << Library->ShaderIndices;
+			Library->ShaderPreloads.SetNum(Library->GetNumShaders());
+			{
+				const uint32 HashSize = FMath::Min<uint32>(0x10000, 1u << FMath::CeilLogTwo(Library->ShaderMapHashes.Num()));
+				Library->ShaderMapHashTable.Clear(HashSize, Library->ShaderMapHashes.Num());
+				for (int32 Index = 0; Index < Library->ShaderMapHashes.Num(); ++Index)
+				{
+					const uint32 Key = GetTypeHash(Library->ShaderMapHashes[Index]);
+					Library->ShaderMapHashTable.Add(Key, Index);
+				}
+			}
+			{
+				const uint32 HashSize = FMath::Min<uint32>(0x10000, 1u << FMath::CeilLogTwo(Library->ShaderHashes.Num()));
+				Library->ShaderHashTable.Clear(HashSize, Library->ShaderHashes.Num());
+				for (int32 Index = 0; Index < Library->ShaderHashes.Num(); ++Index)
+				{
+					const uint32 Key = GetTypeHash(Library->ShaderHashes[Index]);
+					Library->ShaderHashTable.Add(Key, Index);
+				}
+			}
+			
+			UE_LOG(LogShaderLibrary, Display, TEXT("Using IoDispatcher for shader code library %s. Total %d unique shaders."), *InLibraryName, Library->ShaderEntries.Num());
+			INC_DWORD_STAT_BY(STAT_Shaders_ShaderResourceMemory, Library->GetSizeBytes());
+			return Library;
+		}
+	}
+	return nullptr;
+}
+
+FIoStoreShaderCodeArchive::FIoStoreShaderCodeArchive(EShaderPlatform InPlatform, const FString& InLibraryName, FIoDispatcher& InIoDispatcher)
+	: FRHIShaderLibrary(InPlatform, InLibraryName)
+	, IoDispatcher(InIoDispatcher)
+{
+}
+
+FIoStoreShaderCodeArchive::~FIoStoreShaderCodeArchive()
+{
+	DEC_DWORD_STAT_BY(STAT_Shaders_ShaderResourceMemory, GetSizeBytes());
+	Teardown();
+}
+
+void FIoStoreShaderCodeArchive::Teardown()
+{
+	for (int32 ShaderIndex = 0; ShaderIndex < ShaderPreloads.Num(); ++ShaderIndex)
+	{
+		FShaderPreloadEntry& ShaderPreloadEntry = ShaderPreloads[ShaderIndex];
+		check(!ShaderPreloadEntry.NumRefs);
+	}
+}
+
+bool FIoStoreShaderCodeArchive::PreloadShader(int32 ShaderIndex, FGraphEventArray& OutCompletionEvents)
+{
+	LLM_SCOPE(ELLMTag::Shaders);
+
+	FWriteScopeLock Lock(ShaderPreloadLock);
+
+	FShaderPreloadEntry& ShaderPreloadEntry = ShaderPreloads[ShaderIndex];
+	const uint32 ShaderNumRefs = ShaderPreloadEntry.NumRefs++;
+	if (ShaderNumRefs == 0u)
+	{
+		check(!ShaderPreloadEntry.PreloadEvent);
+
+		const FIoStoreShaderCodeEntry& ShaderEntry = ShaderEntries[ShaderIndex];
+		ShaderPreloadEntry.FramePreloadStarted = GFrameNumber;
+
+		ShaderPreloadEntry.PreloadEvent = FGraphEvent::CreateGraphEvent();
+		FIoBatch IoBatch = IoDispatcher.NewBatch();
+		IoBatch.Read(GetShaderCodeChunkId(ShaderHashes[ShaderIndex]), FIoReadOptions(), IoDispatcherPriority_Medium);
+		IoBatch.IssueAndDispatchSubsequents(ShaderPreloadEntry.PreloadEvent);
+		INC_DWORD_STAT_BY(STAT_Shaders_ShaderPreloadMemory, ShaderEntry.CompressedSize);
+	}
+
+	if (ShaderPreloadEntry.PreloadEvent && !ShaderPreloadEntry.PreloadEvent->IsComplete())
+	{
+		OutCompletionEvents.Add(ShaderPreloadEntry.PreloadEvent);
+	}
+	return true;
+}
+
+bool FIoStoreShaderCodeArchive::PreloadShaderMap(int32 ShaderMapIndex, FGraphEventArray& OutCompletionEvents)
+{
+	LLM_SCOPE(ELLMTag::Shaders);
+
+	const FIoStoreShaderMapEntry& ShaderMapEntry = ShaderMapEntries[ShaderMapIndex];
+	const uint32 FrameNumber = GFrameNumber;
+	uint32 PreloadMemory = 0u;
+	uint32 PreloadCount = 0u;
+
+	FWriteScopeLock Lock(ShaderPreloadLock);
+
+	for (uint32 i = 0u; i < ShaderMapEntry.NumShaders; ++i)
+	{
+		const int32 ShaderIndex = ShaderIndices[ShaderMapEntry.ShaderIndicesOffset + i];
+		FShaderPreloadEntry& ShaderPreloadEntry = ShaderPreloads[ShaderIndex];
+		const uint32 ShaderNumRefs = ShaderPreloadEntry.NumRefs++;
+		if (ShaderNumRefs == 0u)
+		{
+			check(!ShaderPreloadEntry.PreloadEvent);
+			const FIoStoreShaderCodeEntry& ShaderEntry = ShaderEntries[ShaderIndex];
+			ShaderPreloadEntry.FramePreloadStarted = FrameNumber;
+			PreloadMemory += ShaderEntry.CompressedSize;
+			++PreloadCount;
+			ShaderPreloadEntry.PreloadEvent = FGraphEvent::CreateGraphEvent();
+			FIoBatch IoBatch = IoDispatcher.NewBatch();
+			ShaderPreloadEntry.IoRequest = IoBatch.Read(GetShaderCodeChunkId(ShaderHashes[ShaderIndex]), FIoReadOptions(), IoDispatcherPriority_Medium);
+			IoBatch.IssueAndDispatchSubsequents(ShaderPreloadEntry.PreloadEvent);
+		}
+		if (ShaderPreloadEntry.PreloadEvent && !ShaderPreloadEntry.PreloadEvent->IsComplete())
+		{
+			OutCompletionEvents.Add(ShaderPreloadEntry.PreloadEvent);
+		}
+	}
+
+	INC_DWORD_STAT_BY(STAT_Shaders_ShaderPreloadMemory, PreloadMemory);
+
+	return true;
+}
+
+bool FIoStoreShaderCodeArchive::PreloadShaderMap(int32 ShaderMapIndex, FCoreDelegates::FAttachShaderReadRequestFunc AttachShaderReadRequestFunc)
+{
+	LLM_SCOPE(ELLMTag::Shaders);
+
+	const FIoStoreShaderMapEntry& ShaderMapEntry = ShaderMapEntries[ShaderMapIndex];
+	const uint32 FrameNumber = GFrameNumber;
+	uint32 PreloadMemory = 0u;
+
+	FWriteScopeLock Lock(ShaderPreloadLock);
+
+	for (uint32 i = 0u; i < ShaderMapEntry.NumShaders; ++i)
+	{
+		const int32 ShaderIndex = ShaderIndices[ShaderMapEntry.ShaderIndicesOffset + i];
+		FShaderPreloadEntry& ShaderPreloadEntry = ShaderPreloads[ShaderIndex];
+		const uint32 ShaderNumRefs = ShaderPreloadEntry.NumRefs++;
+		if (ShaderNumRefs == 0u)
+		{
+			check(!ShaderPreloadEntry.PreloadEvent);
+			const FIoStoreShaderCodeEntry& ShaderEntry = ShaderEntries[ShaderIndex];
+			ShaderPreloadEntry.FramePreloadStarted = FrameNumber;
+			PreloadMemory += ShaderEntry.CompressedSize;
+			ShaderPreloadEntry.PreloadEvent = FGraphEvent::CreateGraphEvent();
+			ShaderPreloadEntry.IoRequest = AttachShaderReadRequestFunc(GetShaderCodeChunkId(ShaderHashes[ShaderIndex]), ShaderPreloadEntry.PreloadEvent);
+		}
+	}
+	INC_DWORD_STAT_BY(STAT_Shaders_ShaderPreloadMemory, PreloadMemory);
+	return true;
+}
+
+bool FIoStoreShaderCodeArchive::ReleaseRef(int32 ShaderIndex)
+{
+	FShaderPreloadEntry& ShaderPreloadEntry = ShaderPreloads[ShaderIndex];
+	const uint32 ShaderNumRefs = ShaderPreloadEntry.NumRefs--;
+	check(ShaderNumRefs > 0u);
+	if (ShaderNumRefs == 1u)
+	{
+		ShaderPreloadEntry.IoRequest = FIoRequest();
+		ShaderPreloadEntry.PreloadEvent.SafeRelease();
+		const FIoStoreShaderCodeEntry& ShaderEntry = ShaderEntries[ShaderIndex];
+		DEC_DWORD_STAT_BY(STAT_Shaders_ShaderPreloadMemory, ShaderEntry.CompressedSize);
+		return true;
+	}
+	return false;
+}
+
+void FIoStoreShaderCodeArchive::ReleasePreloadedShader(int32 ShaderIndex)
+{
+	FWriteScopeLock Lock(ShaderPreloadLock);
+	ReleaseRef(ShaderIndex);
+}
+
+int32 FIoStoreShaderCodeArchive::FindShaderMapIndex(const FSHAHash& Hash)
+{
+	const uint32 Key = GetTypeHash(Hash);
+	for (uint32 Index = ShaderMapHashTable.First(Key); ShaderMapHashTable.IsValid(Index); Index = ShaderMapHashTable.Next(Index))
+	{
+		if (ShaderMapHashes[Index] == Hash)
+		{
+			return Index;
+		}
+	}
+	return INDEX_NONE;
+}
+
+int32 FIoStoreShaderCodeArchive::FindShaderIndex(const FSHAHash& Hash)
+{
+	const uint32 Key = GetTypeHash(Hash);
+	for (uint32 Index = ShaderHashTable.First(Key); ShaderHashTable.IsValid(Index); Index = ShaderHashTable.Next(Index))
+	{
+		if (ShaderHashes[Index] == Hash)
+		{
+			return Index;
+		}
+	}
+	return INDEX_NONE;
+}
+
+TRefCountPtr<FRHIShader> FIoStoreShaderCodeArchive::CreateShader(int32 Index)
+{
+	LLM_SCOPE(ELLMTag::Shaders);
+	TRefCountPtr<FRHIShader> Shader;
+
+	const FIoStoreShaderCodeEntry& ShaderEntry = ShaderEntries[Index];
+	FShaderPreloadEntry& ShaderPreloadEntry = ShaderPreloads[Index];
+
+	FGraphEventRef Event;
+	{
+		FWriteScopeLock Lock(ShaderPreloadLock);
+		uint32 ShaderNumRefs = ShaderPreloadEntry.NumRefs++;
+		if (ShaderNumRefs == 0u)
+		{
+			check(!ShaderPreloadEntry.PreloadEvent);
+			ShaderPreloadEntry.PreloadEvent = FGraphEvent::CreateGraphEvent();
+			FIoBatch IoBatch = IoDispatcher.NewBatch();
+			ShaderPreloadEntry.IoRequest = IoBatch.Read(GetShaderCodeChunkId(ShaderHashes[Index]), FIoReadOptions(), IoDispatcherPriority_Max);
+			IoBatch.IssueAndDispatchSubsequents(ShaderPreloadEntry.PreloadEvent);
+		}
+		else if (!ShaderPreloadEntry.IoRequest.Status().IsCompleted())
+		{
+			ShaderPreloadEntry.IoRequest.UpdatePriority(IoDispatcherPriority_Max);
+		}
+		Event = ShaderPreloadEntry.PreloadEvent;
+	}
+
+	const bool bNeededToWait = !Event->IsComplete();
+	if (bNeededToWait)
+	{
+		TRACE_CPUPROFILER_EVENT_SCOPE(BlockingShaderLoad);
+		UE_LOG(LogShaderLibrary, Warning, TEXT("Blocking wait for shader preload, NumRefs: %d, FramePreloadStarted: %d"), ShaderPreloadEntry.NumRefs, ShaderPreloadEntry.FramePreloadStarted);
+		FTaskGraphInterface::Get().WaitUntilTaskCompletes(Event);
+	}
+
+	const uint8* ShaderCode = ShaderPreloadEntry.IoRequest.GetResult().ValueOrDie().Data();
+
+	FMemStackBase& MemStack = FMemStack::Get();
+	FMemMark Mark(MemStack);
+	if (ShaderEntry.UncompressedSize != ShaderEntry.CompressedSize)
+	{
+		void* UncompressedCode = MemStack.Alloc(ShaderEntry.UncompressedSize, 16);
+		const bool bDecompressResult = FCompression::UncompressMemory(GetShaderCompressionFormat(), UncompressedCode, ShaderEntry.UncompressedSize, ShaderCode, ShaderEntry.CompressedSize);
+		check(bDecompressResult);
+		ShaderCode = (uint8*)UncompressedCode;
+	}
+
+	const auto ShaderCodeView = MakeArrayView(ShaderCode, ShaderEntry.UncompressedSize);
+	const FSHAHash& ShaderHash = ShaderHashes[Index];
+	switch (ShaderEntry.Frequency)
+	{
+	case SF_Vertex: Shader = RHICreateVertexShader(ShaderCodeView, ShaderHash); break;
+	case SF_Mesh: Shader = RHICreateMeshShader(ShaderCodeView, ShaderHash); break;
+	case SF_Amplification: Shader = RHICreateAmplificationShader(ShaderCodeView, ShaderHash); break;
+	case SF_Pixel: Shader = RHICreatePixelShader(ShaderCodeView, ShaderHash); break;
+	case SF_Geometry: Shader = RHICreateGeometryShader(ShaderCodeView, ShaderHash); break;
+	case SF_Compute: Shader = RHICreateComputeShader(ShaderCodeView, ShaderHash); break;
+	case SF_RayGen: case SF_RayMiss: case SF_RayHitGroup: case SF_RayCallable:
+#if RHI_RAYTRACING
+		if (GRHISupportsRayTracing)
+		{
+			Shader = RHICreateRayTracingShader(ShaderCodeView, ShaderHash, ShaderEntry.GetFrequency());
+		}
+#endif // RHI_RAYTRACING
+		break;
+	default: checkNoEntry(); break;
+	}
+
+	{
+		FWriteScopeLock Lock(ShaderPreloadLock);
+		const uint32 ShaderNumRefs = ShaderPreloadEntry.NumRefs--;
+		check(ShaderNumRefs > 0u);
+		if (ShaderNumRefs == 1u)
+		{
+			ShaderPreloadEntry.PreloadEvent.SafeRelease();
+			ShaderPreloadEntry.IoRequest = FIoRequest();
+		}
+	}
+
+	if (Shader)
+	{
+		Shader->SetHash(ShaderHash);
+	}
+
+	return Shader;
+}
+
diff --git a/Engine/Source/Runtime/RenderCore/Private/ShaderCodeLibrary.cpp b/Engine/Source/Runtime/RenderCore/Private/ShaderCodeLibrary.cpp
--- a/Engine/Source/Runtime/RenderCore/Private/ShaderCodeLibrary.cpp
+++ b/Engine/Source/Runtime/RenderCore/Private/ShaderCodeLibrary.cpp
@@ -33,6 +33,8 @@ ShaderCodeLibrary.cpp: Bound shader state cache implementation.
 #include "Misc/ConfigCacheIni.h"
 #include "Misc/CommandLine.h"
 #include "ProfilingDebugging/LoadTimeTracker.h"
+#include "Serialization/LargeMemoryWriter.h"
+#include "Serialization/MemoryReader.h"
 
 #if WITH_EDITORONLY_DATA
 #include "Modules/ModuleManager.h"
@@ -59,6 +61,7 @@ DEFINE_LOG_CATEGORY(LogShaderLibrary);
 
 static uint32 GShaderCodeArchiveVersion = 2;
 static uint32 GShaderPipelineArchiveVersion = 1;
+extern uint32 GIoStoreShaderCodeArchiveVersion;
 
 static FString ShaderExtension = TEXT(".ushaderbytecode");
 static FString ShaderAssetInfoExtension = TEXT(".assetinfo.json");
@@ -734,6 +737,11 @@ public:
 			Library = RHICreateShaderLibrary(InShaderPlatform, ShaderCodeDir, InLibraryName);
 		}
 
+		if (!Library && FIoDispatcher::IsInitialized())
+		{
+			Library = FIoStoreShaderCodeArchive::Create(InShaderPlatform, InLibraryName, FIoDispatcher::Get());
+		}
+
 		if (!Library)
 		{
 			const FName PlatformName = LegacyShaderPlatformToShaderFormat(InShaderPlatform);
@@ -919,6 +927,24 @@ public:
 		}
 	}
 
+	void PreloadPackageShaderMap(int32 ShaderMapIndex, FCoreDelegates::FAttachShaderReadRequestFunc AttachShaderReadRequestFunc)
+	{
+		FRWScopeLock Locker(ResourceLock, SLT_Write);
+		FShaderMapResource_SharedCode*& Resource = Resources[ShaderMapIndex];
+		if (!Resource)
+		{
+			Resource = new FShaderMapResource_SharedCode(this, ShaderMapIndex);
+			Resource->bShaderMapPreloaded = Library->PreloadShaderMap(ShaderMapIndex, AttachShaderReadRequestFunc);
+		}
+		Resource->AddRef();
+	}
+
+	void ReleasePreloadedPackageShaderMap(int32 ShaderMapIndex)
+	{
+		FRWScopeLock Locker(ResourceLock, SLT_Write);
+		Resources[ShaderMapIndex]->Release();
+	}
+
 public:
 	FRHIShaderLibraryRef Library;
 
@@ -1327,6 +1353,33 @@ struct FEditorShaderCodeArchive
 		}
 	}
 
+	static void SaveIoStoreShaderCodeArchive(const FSerializedShaderArchive& SerializedShaders, FArchive& OutLibraryAr)
+	{
+		OutLibraryAr << GIoStoreShaderCodeArchiveVersion;
+		OutLibraryAr << const_cast<FSerializedShaderArchive&>(SerializedShaders).ShaderMapHashes;
+		OutLibraryAr << const_cast<FSerializedShaderArchive&>(SerializedShaders).ShaderHashes;
+		TArray<FIoStoreShaderMapEntry> ShaderMapEntries;
+		ShaderMapEntries.Reserve(SerializedShaders.ShaderMapEntries.Num());
+		for (const FShaderMapEntry& ShaderMapEntry : SerializedShaders.ShaderMapEntries)
+		{
+			FIoStoreShaderMapEntry& IoStoreShaderMapEntry = ShaderMapEntries.AddDefaulted_GetRef();
+			IoStoreShaderMapEntry.ShaderIndicesOffset = ShaderMapEntry.ShaderIndicesOffset;
+			IoStoreShaderMapEntry.NumShaders = ShaderMapEntry.NumShaders;
+		}
+		OutLibraryAr << ShaderMapEntries;
+		TArray<FIoStoreShaderCodeEntry> ShaderEntries;
+		ShaderEntries.Reserve(SerializedShaders.ShaderEntries.Num());
+		for (const FShaderCodeEntry& ShaderEntry : SerializedShaders.ShaderEntries)
+		{
+			FIoStoreShaderCodeEntry& IoStoreShaderEntry = ShaderEntries.AddDefaulted_GetRef();
+			IoStoreShaderEntry.UncompressedSize = ShaderEntry.UncompressedSize;
+			IoStoreShaderEntry.CompressedSize = ShaderEntry.Size;
+			IoStoreShaderEntry.Frequency = ShaderEntry.Frequency;
+		}
+		OutLibraryAr << ShaderEntries;
+		OutLibraryAr << const_cast<FSerializedShaderArchive&>(SerializedShaders).ShaderIndices;
+	}
+
 	bool Finalize(FString OutputDir, const FString& MetaOutputDir, bool bSaveOnlyAssetInfo = false, TArray<FString>* OutputFilenames = nullptr)
 	{
 		check(LibraryName.Len() > 0);
@@ -1851,10 +1904,16 @@ public:
 			ECVF_Default
 			);
 #endif
+
+		FCoreDelegates::PreloadPackageShaderMaps.BindRaw(this, &FShaderLibrariesCollection::PreloadPackageShaderMapsDelegate);
+		FCoreDelegates::ReleasePreloadedPackageShaderMaps.BindRaw(this, &FShaderLibrariesCollection::ReleasePreloadedPackageShaderMapsDelegate);
 	}
 
 	~FShaderLibrariesCollection()
 	{
+		FCoreDelegates::PreloadPackageShaderMaps.Unbind();
+		FCoreDelegates::ReleasePreloadedPackageShaderMaps.Unbind();
+
 #if WITH_EDITOR
 		for (uint32 i = 0; i < EShaderPlatform::SP_NumPlatforms; i++)
 		{
@@ -2501,6 +2560,33 @@ public:
 		}
 	}
 #endif// WITH_EDITOR
+
+	void PreloadPackageShaderMapsDelegate(TArrayView<const FSHAHash> ShaderMapHashes, FCoreDelegates::FAttachShaderReadRequestFunc AttachShaderReadRequestFunc)
+	{
+		for (const FSHAHash& ShaderMapHash : ShaderMapHashes)
+		{
+			int32 ShaderMapIndex;
+			FShaderLibraryInstance* LibraryInstance = FindShaderLibraryForShaderMap(ShaderMapHash, ShaderMapIndex);
+			if (LibraryInstance)
+			{
+				LibraryInstance->PreloadPackageShaderMap(ShaderMapIndex, AttachShaderReadRequestFunc);
+			}
+		}
+	}
+
+	void ReleasePreloadedPackageShaderMapsDelegate(TArrayView<const FSHAHash> ShaderMapHashes)
+	{
+		for (const FSHAHash& ShaderMapHash : ShaderMapHashes)
+		{
+			int32 ShaderMapIndex;
+			FShaderLibraryInstance* LibraryInstance = FindShaderLibraryForShaderMap(ShaderMapHash, ShaderMapIndex);
+			if (LibraryInstance)
+			{
+				LibraryInstance->ReleasePreloadedPackageShaderMap(ShaderMapIndex);
+			}
+		}
+	}
+
 };
 
 static FSharedShaderCodeRequest OnSharedShaderCodeRequest;
@@ -2960,6 +3046,90 @@ bool FShaderLibraryCooker::SaveShaderLibraryChunk(int32 ChunkId, const TSet<FNam
 	return bResult;
 }
 
+bool FShaderLibraryCooker::ConvertToIoStoreShaderLibrary(
+	const TCHAR* FileName,
+	TTuple<FIoChunkId, FIoBuffer>& OutLibraryIoChunk,
+	TArray<TTuple<FIoChunkId, FIoBuffer, FSHAHash>>& OutCodeIoChunks,
+	TArray<TTuple<FSHAHash, TArray<FIoChunkId>>>& OutShaderMaps,
+	TArray<TTuple<FSHAHash, TSet<FName>>>& OutShaderMapAssetAssociations)
+{
+	TArray<FString> Components;
+	if (FPaths::GetBaseFilename(FileName).ParseIntoArray(Components, TEXT("-")) != 3)
+	{
+		UE_LOG(LogShaderLibrary, Error, TEXT("Invalid shader code library file name '%s'."), FileName);
+		return false;
+	}
+	FString LibraryName = Components[1];
+	FName FormatName(Components[2]);
+	
+	TUniquePtr<FArchive> LibraryAr(IFileManager::Get().CreateFileReader(FileName));
+	if (!LibraryAr)
+	{
+		UE_LOG(LogShaderLibrary, Error, TEXT("Missing shader code library file '%s'."), FileName);
+		return false;
+	}
+
+	uint32 Version;
+	(*LibraryAr) << Version;
+	if (Version != GShaderCodeArchiveVersion)
+	{
+		UE_LOG(LogShaderLibrary, Error, TEXT("Invalid shader code library chunk version in file '%s'."), FileName);
+		return false;
+	}
+
+	FSerializedShaderArchive SerializedShaders;
+	(*LibraryAr) << SerializedShaders;
+	int64 OffsetToShaderCode = LibraryAr->Tell();
+
+	FString AssetInfoFileName = GetShaderAssetInfoFilename(FPaths::GetPath(FileName), LibraryName, FormatName);
+	if (!SerializedShaders.LoadAssetInfo(AssetInfoFileName))
+	{
+		UE_LOG(LogShaderLibrary, Error, TEXT("Failed loading asset asset info file '%s'"), *AssetInfoFileName);
+		return false;
+	}
+
+	FLargeMemoryWriter IoStoreLibraryAr(0, true);
+	FEditorShaderCodeArchive::SaveIoStoreShaderCodeArchive(SerializedShaders, IoStoreLibraryAr);
+	OutLibraryIoChunk.Key = GetShaderCodeArchiveChunkId(LibraryName, FormatName);
+	int64 TotalSize = IoStoreLibraryAr.TotalSize();
+	OutLibraryIoChunk.Value = FIoBuffer(FIoBuffer::AssumeOwnership, IoStoreLibraryAr.ReleaseOwnership(), TotalSize);
+
+	int32 ShaderCount = SerializedShaders.ShaderEntries.Num();
+	OutCodeIoChunks.Reserve(ShaderCount);
+	for (int32 ShaderIndex = 0; ShaderIndex < ShaderCount; ++ShaderIndex)
+	{
+		const FSHAHash& ShaderHash = SerializedShaders.ShaderHashes[ShaderIndex];
+		const FShaderCodeEntry& ShaderEntry = SerializedShaders.ShaderEntries[ShaderIndex];
+		TTuple<FIoChunkId, FIoBuffer, FSHAHash>& OutCodeIoChunk = OutCodeIoChunks.Emplace_GetRef(GetShaderCodeChunkId(ShaderHash), ShaderEntry.Size, ShaderHash);
+		LibraryAr->Seek(OffsetToShaderCode + ShaderEntry.Offset);
+		LibraryAr->Serialize(OutCodeIoChunk.Get<1>().Data(), ShaderEntry.Size);
+	}
+
+	int32 ShaderMapCount = SerializedShaders.ShaderMapHashes.Num();
+	OutShaderMaps.Reserve(ShaderMapCount);
+	for (int32 ShaderMapIndex = 0; ShaderMapIndex < ShaderMapCount; ++ShaderMapIndex)
+	{
+		TTuple<FSHAHash, TArray<FIoChunkId>>& OutShaderMap = OutShaderMaps.AddDefaulted_GetRef();
+		OutShaderMap.Key = SerializedShaders.ShaderMapHashes[ShaderMapIndex];
+		const FShaderMapEntry& ShaderMapEntry = SerializedShaders.ShaderMapEntries[ShaderMapIndex];
+		int32 LookupIndexEnd = ShaderMapEntry.ShaderIndicesOffset + ShaderMapEntry.NumShaders;
+		OutShaderMap.Value.Reserve(ShaderMapEntry.NumShaders);
+		for (int32 ShaderLookupIndex = ShaderMapEntry.ShaderIndicesOffset; ShaderLookupIndex < LookupIndexEnd; ++ShaderLookupIndex)
+		{
+			int32 ShaderIndex = SerializedShaders.ShaderIndices[ShaderLookupIndex];
+			OutShaderMap.Value.Add(GetShaderCodeChunkId(SerializedShaders.ShaderHashes[ShaderIndex]));
+		}
+	}
+
+	OutShaderMapAssetAssociations.Reserve(SerializedShaders.ShaderCodeToAssets.Num());
+	for (auto& KV : SerializedShaders.ShaderCodeToAssets)
+	{
+		OutShaderMapAssetAssociations.Emplace(KV.Key, MoveTemp(KV.Value));
+	}
+
+	return true;
+}
+
 #endif// WITH_EDITOR
 
 void FShaderCodeLibrary::SafeAssignHash(FRHIShader* InShader, const FSHAHash& Hash)
diff --git a/Engine/Source/Runtime/RenderCore/Public/ShaderCodeArchive.h b/Engine/Source/Runtime/RenderCore/Public/ShaderCodeArchive.h
--- a/Engine/Source/Runtime/RenderCore/Public/ShaderCodeArchive.h
+++ b/Engine/Source/Runtime/RenderCore/Public/ShaderCodeArchive.h
@@ -7,6 +7,10 @@
 #include "FileCache/FileCache.h"
 #include "Containers/HashTable.h"
 #include "Shader.h"
+#include "IO/IoDispatcher.h"
+
+FIoChunkId GetShaderCodeArchiveChunkId(const FString& LibraryName, FName FormatName);
+FIoChunkId GetShaderCodeChunkId(const FSHAHash& ShaderHash);
 
 struct FShaderMapEntry
 {
@@ -261,3 +265,96 @@ protected:
 	TArray<FShaderPreloadEntry> ShaderPreloads;
 	FRWLock ShaderPreloadLock;
 };
+
+struct FIoStoreShaderMapEntry
+{
+	uint32 ShaderIndicesOffset = 0u;
+	uint32 NumShaders = 0u;
+
+	friend FArchive& operator <<(FArchive& Ar, FIoStoreShaderMapEntry& Ref)
+	{
+		return Ar << Ref.ShaderIndicesOffset << Ref.NumShaders;
+	}
+};
+
+struct FIoStoreShaderCodeEntry
+{
+	uint32 UncompressedSize = 0;
+	uint32 CompressedSize = 0;
+	uint8 Frequency;
+
+	EShaderFrequency GetFrequency() const
+	{
+		return (EShaderFrequency)Frequency;
+	}
+
+	friend FArchive& operator <<(FArchive& Ar, FIoStoreShaderCodeEntry& Ref)
+	{
+		return Ar << Ref.UncompressedSize << Ref.CompressedSize << Ref.Frequency;
+	}
+};
+
+class FIoStoreShaderCodeArchive : public FRHIShaderLibrary
+{
+public:
+	static FIoStoreShaderCodeArchive* Create(EShaderPlatform InPlatform, const FString& InLibraryName, FIoDispatcher& InIoDispatcher);
+
+	virtual ~FIoStoreShaderCodeArchive();
+
+	virtual bool IsNativeLibrary() const override { return false; }
+
+	uint32 GetSizeBytes() const
+	{
+		return sizeof(*this) +
+			ShaderMapHashes.GetAllocatedSize() +
+			ShaderMapEntries.GetAllocatedSize() +
+			ShaderHashes.GetAllocatedSize() +
+			ShaderEntries.GetAllocatedSize() +
+			ShaderIndices.GetAllocatedSize() +
+			ShaderPreloads.GetAllocatedSize();
+	}
+
+	virtual int32 GetNumShaders() const override { return ShaderEntries.Num(); }
+	virtual int32 GetNumShaderMaps() const override { return ShaderMapEntries.Num(); }
+	virtual int32 GetNumShadersForShaderMap(int32 ShaderMapIndex) const override { return ShaderMapEntries[ShaderMapIndex].NumShaders; }
+
+	virtual int32 GetShaderIndex(int32 ShaderMapIndex, int32 i) const override
+	{
+		const FIoStoreShaderMapEntry& ShaderMapEntry = ShaderMapEntries[ShaderMapIndex];
+		return ShaderIndices[ShaderMapEntry.ShaderIndicesOffset + i];
+	}
+
+	virtual int32 FindShaderMapIndex(const FSHAHash& Hash) override;
+	virtual int32 FindShaderIndex(const FSHAHash& Hash) override;
+	virtual bool PreloadShader(int32 ShaderIndex, FGraphEventArray& OutCompletionEvents) override;
+	virtual bool PreloadShaderMap(int32 ShaderMapIndex, FGraphEventArray& OutCompletionEvents) override;
+	virtual bool PreloadShaderMap(int32 ShaderMapIndex, FCoreDelegates::FAttachShaderReadRequestFunc AttachShaderReadRequestFunc) override;
+	virtual void ReleasePreloadedShader(int32 ShaderIndex) override;
+	virtual TRefCountPtr<FRHIShader> CreateShader(int32 Index) override;
+	virtual void Teardown() override;
+
+private:
+	struct FShaderPreloadEntry
+	{
+		FGraphEventRef PreloadEvent;
+		FIoRequest IoRequest;
+		uint32 FramePreloadStarted = ~0u;
+		uint32 NumRefs = 0u;
+	};
+
+	FIoStoreShaderCodeArchive(EShaderPlatform InPlatform, const FString& InLibraryName, FIoDispatcher& InIoDispatcher);
+	bool ReleaseRef(int32 ShaderIndex);
+
+	FIoDispatcher& IoDispatcher;
+
+	TArray<FSHAHash> ShaderMapHashes;
+	TArray<FSHAHash> ShaderHashes;
+	TArray<FIoStoreShaderMapEntry> ShaderMapEntries;
+	TArray<FIoStoreShaderCodeEntry> ShaderEntries;
+	TArray<uint32> ShaderIndices;
+	FHashTable ShaderMapHashTable;
+	FHashTable ShaderHashTable;
+	TArray<FShaderPreloadEntry> ShaderPreloads;
+	FRWLock ShaderPreloadLock;
+};
+
diff --git a/Engine/Source/Runtime/RenderCore/Public/ShaderCodeLibrary.h b/Engine/Source/Runtime/RenderCore/Public/ShaderCodeLibrary.h
--- a/Engine/Source/Runtime/RenderCore/Public/ShaderCodeLibrary.h
+++ b/Engine/Source/Runtime/RenderCore/Public/ShaderCodeLibrary.h
@@ -17,6 +17,8 @@ class FShaderPipeline;
 class FShaderMapResource;
 class FShaderMapResourceCode;
 using FShaderMapAssetPaths = TSet<FName>;
+class FIoChunkId;
+class FIoBuffer;
 
 struct RENDERCORE_API FShaderCodeLibraryPipeline
 {
@@ -302,5 +304,12 @@ struct RENDERCORE_API FShaderLibraryCooker
 
 	// Create a smaller 'patch' library that only contains data from 'NewMetaDataDir' not contained in any of 'OldMetaDataDirs'
 	static bool CreatePatchLibrary(TArray<FString> const& OldMetaDataDirs, FString const& NewMetaDataDir, FString const& OutDir, bool bNativeFormat, bool bNeedsDeterministicOrder);
+
+	static bool ConvertToIoStoreShaderLibrary(
+		const TCHAR* FileName,
+		TTuple<FIoChunkId, FIoBuffer>& OutLibraryIoChunk,
+		TArray<TTuple<FIoChunkId, FIoBuffer, FSHAHash>>& OutCodeIoChunks,
+		TArray<TTuple<FSHAHash, TArray<FIoChunkId>>>& OutShaderMaps,
+		TArray<TTuple<FSHAHash, TSet<FName>>>& OutShaderMapAssetAssociations);
 };
 #endif
