Integrating //UE5/Dev-Cooker @ 16678003 to //UE5/Main (Zousar.Shaker-YEG-0943-Quaternary)

diff --git a/Engine/Source/Developer/IoStoreUtilities/Private/IoStoreUtilities.cpp b/Engine/Source/Developer/IoStoreUtilities/Private/IoStoreUtilities.cpp
--- a/Engine/Source/Developer/IoStoreUtilities/Private/IoStoreUtilities.cpp
+++ b/Engine/Source/Developer/IoStoreUtilities/Private/IoStoreUtilities.cpp
@@ -3,7 +3,7 @@
 #include "IoStoreUtilities.h"
 
 #if !UE_BUILD_SHIPPING
-
+#include "FilePackageStoreWriter.h"
 #include "HAL/FileManager.h"
 #include "HAL/PlatformFileManager.h"
 #include "Hash/CityHash.h"
@@ -30,7 +30,8 @@
 #include "Serialization/MemoryReader.h"
 #include "Serialization/AsyncLoading2.h"
 #include "Serialization/ArrayReader.h"
-#include "Serialization/PackageStore.h"
+#include "IO/PackageStore.h"
+#include "Serialization/FilePackageStore.h"
 #include "UObject/Class.h"
 #include "UObject/NameBatchSerialization.h"
 #include "UObject/PackageFileSummary.h"
@@ -55,6 +56,9 @@
 #include "Algo/StableSort.h"
 #include "PackageStoreOptimizer.h"
 #include "ShaderCodeLibrary.h"
+#include "ZenStoreHttpClient.h"
+#include "IPlatformFilePak.h"
+#include "ZenStoreWriter.h"
 
 //PRAGMA_DISABLE_OPTIMIZATION
 
@@ -320,6 +324,7 @@ struct FLegacyCookedPackage
 	uint64 DiskLayoutOrder = -1;
 	FPackageStorePackage* OptimizedPackage = nullptr;
 	TSet<FShaderInfo*> Shaders;
+	const FPackageStoreEntryResource* PackageStoreEntry = nullptr;
 };
 
 enum class EContainerChunkType
@@ -350,6 +355,248 @@ struct FContainerTargetFile
 	TArray<FFileRegion> FileRegions;
 };
 
+class FCookedPackageStore
+{
+public:
+	FCookedPackageStore()
+	{
+	}
+
+	FIoStatus Load(const FString& CookedDirectory)
+	{
+		FIoStatus Status = PackageStoreManifest.Load(*(CookedDirectory / FApp::GetProjectName() / TEXT("MetaData") / TEXT("packagestore.manifest")));
+		if (!Status.IsOk())
+		{
+			return Status;
+		}
+
+		if (const FPackageStoreManifest::FZenServerInfo* ZenServerInfo = PackageStoreManifest.ReadZenServerInfo())
+		{
+			DataSource = MakeUnique<FZenStoreDataSource>(*ZenServerInfo);
+		}
+		else
+		{
+			TUniquePtr<FIoStoreReaderDataSource> IoStoreReaderDataSource = MakeUnique<FIoStoreReaderDataSource>();
+			Status = IoStoreReaderDataSource->Initialize(*(CookedDirectory / TEXT("global")));
+			if (!Status.IsOk())
+			{
+				return Status;
+			}
+			DataSource = MoveTemp(IoStoreReaderDataSource);
+		}
+
+		FIoContainerId ContainerId = FIoContainerId::FromName(TEXT("global"));
+
+		TIoStatusOr<FIoBuffer> HeaderBuffer = DataSource->ReadChunk(CreateIoChunkId(ContainerId.Value(), 0, EIoChunkType::ContainerHeader), FIoReadOptions());
+		if (!HeaderBuffer.IsOk())
+		{
+			return FIoStatusBuilder(EIoErrorCode::ReadError) << TEXT("Failed to read container header");
+		}
+
+		FMemoryReaderView Ar(MakeArrayView(HeaderBuffer.ValueOrDie().Data(), HeaderBuffer.ValueOrDie().DataSize()));
+		FContainerHeader ContainerHeader;
+		Ar << ContainerHeader;
+
+		PackageEntries.Reserve(ContainerHeader.PackageCount);
+		TArrayView<const FFilePackageStoreEntry> FilePackageEntries = MakeArrayView<const FFilePackageStoreEntry>(reinterpret_cast<const FFilePackageStoreEntry*>(ContainerHeader.StoreEntries.GetData()), ContainerHeader.PackageCount);
+
+		int32 Idx = 0;
+		for (const FFilePackageStoreEntry& FilePackageEntry : FilePackageEntries)
+		{
+			FPackageStoreEntryResource& Entry = PackageEntries.AddDefaulted_GetRef();
+			
+			Entry.ExportInfo = FPackageStoreExportInfo
+			{
+				FilePackageEntry.ExportBundlesSize,
+				FilePackageEntry.ExportCount,
+				FilePackageEntry.ExportBundleCount,
+				FilePackageEntry.LoadOrder
+			};
+			Entry.ImportedPackageIds = MakeArrayView<const FPackageId>(FilePackageEntry.ImportedPackages.Data(), FilePackageEntry.ImportedPackages.Num());
+
+			PackageIdToEntry.Add(ContainerHeader.PackageIds[Idx++], &Entry);
+		}
+
+		for (const FPackageStoreManifest::FFileInfo& FileInfo : PackageStoreManifest.GetFiles())
+		{
+			FilenameToChunkIdMap.Add(FileInfo.FileName, FileInfo.ChunkId);
+		}
+
+		for (const FPackageStoreManifest::FPackageInfo& PackageInfo : PackageStoreManifest.GetPackages())
+		{
+			FName PackageName = FName(PackageInfo.PackageName);
+			PackageNameToPackageInfoMap.Add(PackageName, &PackageInfo);
+			ChunkIdToPackageNameMap.Add(PackageInfo.PackageChunkId, PackageName);
+			for (const FIoChunkId& ChunkId : PackageInfo.BulkDataChunkIds)
+			{
+				ChunkIdToPackageNameMap.Add(ChunkId, PackageName);
+			}
+		}
+
+		return FIoStatus::Ok;
+	}
+
+	FIoChunkId GetChunkIdFromFileName(const FString& Filename) const
+	{
+		return FilenameToChunkIdMap.FindRef(Filename);
+	}
+
+	FName GetPackageNameFromChunkId(const FIoChunkId& ChunkId) const
+	{
+		return ChunkIdToPackageNameMap.FindRef(ChunkId);
+	}
+
+	const FPackageStoreManifest::FPackageInfo* GetPackageInfoFromPackageName(FName PackageName)
+	{
+		return PackageNameToPackageInfoMap.FindRef(PackageName);
+	}
+
+	const FPackageStoreEntryResource* GetPackageStoreEntry(FPackageId PackageId) const
+	{
+		return PackageIdToEntry.FindRef(PackageId);
+	}
+
+	TIoStatusOr<uint64> GetChunkSize(const FIoChunkId& ChunkId)
+	{
+		return DataSource->GetChunkSize(ChunkId);
+	}
+
+	void ReadChunkAsync(const FIoChunkId& ChunkId, TFunction<void(TIoStatusOr<FIoBuffer>)>&& Callback)
+	{
+		class FReadChunkTask
+			: public FNonAbandonableTask
+		{
+		public:
+			FReadChunkTask(IDataSource* InDataSource, const FIoChunkId& InChunkId, TFunction<void(TIoStatusOr<FIoBuffer>)>&& InCallback)
+				: DataSource(InDataSource)
+				, ChunkId(InChunkId)
+				, Callback(MoveTemp(InCallback))
+			{
+			}
+
+			void DoWork()
+			{
+				Callback(DataSource->ReadChunk(ChunkId, FIoReadOptions()));
+			}
+
+			TStatId GetStatId() const
+			{
+				return TStatId();
+			}
+
+		private:
+			IDataSource* DataSource;
+			FIoChunkId ChunkId;
+			TFunction<void(TIoStatusOr<FIoBuffer>)> Callback;
+		};
+
+		(new FAutoDeleteAsyncTask<FReadChunkTask>(DataSource.Get(), ChunkId, MoveTemp(Callback)))->StartBackgroundTask();
+	}
+	
+	TIoStatusOr<FIoBuffer> ReadPackageHeader(FPackageId PackageId)
+	{
+		if (const FPackageStoreEntryResource* Entry = PackageIdToEntry.FindRef(PackageId))
+		{
+			FIoReadOptions ReadOptions;
+			ReadOptions.SetRange(0, 64 << 10);
+			
+			TIoStatusOr<FIoBuffer> Status = DataSource->ReadChunk(CreateIoChunkId(PackageId.Value(), 0, EIoChunkType::ExportBundleData), ReadOptions);
+			if (!Status.IsOk())
+			{
+				return Status;
+			}
+
+			FIoBuffer Buffer = Status.ConsumeValueOrDie();
+			uint32 HeaderSize = reinterpret_cast<const FPackageSummary*>(Buffer.Data())->HeaderSize;
+			if (HeaderSize > Buffer.DataSize())
+			{
+				ReadOptions.SetRange(0, HeaderSize);
+
+				Status = DataSource->ReadChunk(CreateIoChunkId(PackageId.Value(), 0, EIoChunkType::ExportBundleData), ReadOptions);
+				if (!Status.IsOk())
+				{
+					return Status;
+				}
+				Buffer = Status.ConsumeValueOrDie();
+			}
+			
+			return FIoBuffer(Buffer.Data(), HeaderSize, Buffer);
+		}
+
+		return FIoStatus(EIoErrorCode::NotFound);
+	}
+
+private:
+	class IDataSource
+	{
+	public:
+		virtual ~IDataSource() = default;
+		virtual TIoStatusOr<uint64> GetChunkSize(const FIoChunkId& ChunkId) = 0;
+		virtual TIoStatusOr<FIoBuffer> ReadChunk(const FIoChunkId& ChunkId, const FIoReadOptions& ReadOptions) = 0;
+	};
+
+	class FIoStoreReaderDataSource
+		: public IDataSource
+	{
+	public:
+		FIoStatus Initialize(const TCHAR* ContainerPath)
+		{
+			using FEncryptionKeys = TMap<FGuid, FAES::FAESKey>;
+			return IoStoreReader.Initialize(ContainerPath, FEncryptionKeys());
+		}
+
+		virtual TIoStatusOr<uint64> GetChunkSize(const FIoChunkId& ChunkId) override
+		{
+			TIoStatusOr<FIoStoreTocChunkInfo> ChunkInfo = IoStoreReader.GetChunkInfo(ChunkId);
+			if (!ChunkInfo.IsOk())
+			{
+				return ChunkInfo.Status();
+			}
+			return ChunkInfo.ValueOrDie().Size;
+		}
+
+		virtual TIoStatusOr<FIoBuffer> ReadChunk(const FIoChunkId& ChunkId, const FIoReadOptions& ReadOptions) override
+		{
+			return IoStoreReader.Read(ChunkId, ReadOptions);
+		}
+
+	private:
+		FIoStoreReader IoStoreReader;
+	};
+
+	class FZenStoreDataSource
+		: public IDataSource
+	{
+	public:
+		FZenStoreDataSource(const FPackageStoreManifest::FZenServerInfo& ZenServerInfo)
+		{
+			ZenStoreClient = MakeUnique<UE::FZenStoreHttpClient>(ZenServerInfo.HostName, ZenServerInfo.Port);
+			ZenStoreClient->InitializeReadOnly(ZenServerInfo.ProjectId, ZenServerInfo.OplogId);
+		}
+
+		virtual TIoStatusOr<uint64> GetChunkSize(const FIoChunkId& ChunkId) override
+		{
+			return ZenStoreClient->GetChunkSize(ChunkId);
+		}
+
+		virtual TIoStatusOr<FIoBuffer> ReadChunk(const FIoChunkId& ChunkId, const FIoReadOptions& ReadOptions) override
+		{
+			return ZenStoreClient->ReadChunk(ChunkId, ReadOptions.GetOffset(), ReadOptions.GetSize());
+		}
+
+	private:
+		TUniquePtr<UE::FZenStoreHttpClient> ZenStoreClient;
+	};
+	
+	TUniquePtr<IDataSource> DataSource;
+	FPackageStoreManifest PackageStoreManifest;
+	TArray<FPackageStoreEntryResource> PackageEntries;
+	TMap<FPackageId, const FPackageStoreEntryResource*> PackageIdToEntry;
+	TMap<FString, FIoChunkId> FilenameToChunkIdMap;
+	TMap<FIoChunkId, FName> ChunkIdToPackageNameMap;
+	TMap<FName, const FPackageStoreManifest::FPackageInfo*> PackageNameToPackageInfoMap;
+};
+
 struct FFileOrderMap
 {
 	TMap<FName, int64> PackageNameToOrder;
@@ -381,10 +628,12 @@ struct FIoStoreArguments
 	FString BasedOnReleaseVersionPath;
 	FAssetRegistryState ReleaseAssetRegistry;
 	FReleasedPackages ReleasedPackages;
+	TUniquePtr<FCookedPackageStore> PackageStore;
 	bool bSign = false;
 	bool bRemapPluginContentToGame = false;
 	bool bCreateDirectoryIndex = true;
 	bool bClusterByOrderFilePriority = false;
+	bool bFileRegions = false;
 
 	bool ShouldCreateContainers() const
 	{
@@ -454,15 +703,6 @@ FChunkIdCsv ChunkIdCsv;
 
 #endif
 
-static FIoChunkId CreateChunkId(FPackageId GlobalPackageId, uint16 ChunkIndex, EIoChunkType ChunkType, const TCHAR* DebugString)
-{
-	FIoChunkId ChunkId = CreateIoChunkId(GlobalPackageId.Value(), ChunkIndex, ChunkType);
-#if OUTPUT_CHUNKID_DIRECTORY
-	ChunkIdCsv.AddChunk(GlobalPackageId, 0, ChunkIndex, (uint8)ChunkType, GetTypeHash(ChunkId), DebugString);
-#endif
-	return ChunkId;
-}
-
 // If bClusterByOrderFilePriority is false
 //		Order packages by the order of OrderMaps with their associated priority 
 //		e.g. Order map 1 with priority 0 (A, B, C) 
@@ -612,10 +852,7 @@ static void AssignPackagesDiskOrder(
 				if (!AssignedPackages.Contains(PackageToProcess))
 				{
 					AssignedPackages.Add(PackageToProcess);
-					if (PackageToProcess->OptimizedPackage->GetExportBundleCount())
-					{
-						Cluster->Packages.Add(PackageToProcess);
-					}
+					Cluster->Packages.Add(PackageToProcess);
 					TArray<FPackageId> AllReferencedPackageIds;
 					AllReferencedPackageIds.Append(PackageToProcess->OptimizedPackage->GetImportedPackageIds());
 					AllReferencedPackageIds.Append(PackageToProcess->OptimizedPackage->GetImportedRedirectedPackageIds().Array());
@@ -694,9 +931,7 @@ static void CreateDiskLayout(
 				return A->ChunkType < B->ChunkType;
 			}
 			if (A->ChunkType == EContainerChunkType::ShaderCodeLibrary)
-			{
 				return A->TargetPath < B->TargetPath;
-			}
 			if (A->Package != B->Package)
 			{
 				return A->Package->DiskLayoutOrder < B->Package->DiskLayoutOrder;
@@ -819,50 +1054,58 @@ FContainerTargetSpec* AddContainer(
 
 FLegacyCookedPackage* FindOrAddPackage(
 	const FIoStoreArguments& Arguments,
-	const TCHAR* RelativeFileName,
+	const FName& PackageName,
 	TArray<FLegacyCookedPackage*>& Packages,
 	FPackageNameMap& PackageNameMap,
-	FPackageIdMap& PackageIdMap,
-	FPackageStoreOptimizer& PackageStoreOptimizer)
+	FPackageIdMap& PackageIdMap)
 {
-	FString PackageName;
-	FString ErrorMessage;
-	if (!FPackageName::TryConvertFilenameToLongPackageName(RelativeFileName, PackageName, &ErrorMessage))
-	{
-		UE_LOG(LogIoStore, Warning, TEXT("Failed to obtain package name from file name '%s'"), *ErrorMessage);
-		return nullptr;
-	}
-
-	FName PackageFName = *PackageName;
-
-	FLegacyCookedPackage* Package = PackageNameMap.FindRef(PackageFName);
+	FLegacyCookedPackage* Package = PackageNameMap.FindRef(PackageName);
 	if (!Package)
 	{
-		FPackageId PackageId = FPackageId::FromName(PackageFName);
+		FPackageId PackageId = FPackageId::FromName(PackageName);
 		if (FLegacyCookedPackage* FindById = PackageIdMap.FindRef(PackageId))
 		{
-			UE_LOG(LogIoStore, Fatal, TEXT("Package name hash collision \"%s\" and \"%s"), *FindById->PackageName.ToString(), *PackageFName.ToString());
+			UE_LOG(LogIoStore, Fatal, TEXT("Package name hash collision \"%s\" and \"%s"), *FindById->PackageName.ToString(), *PackageName.ToString());
 		}
 
 		if (const FName* ReleasedPackageName = Arguments.ReleasedPackages.PackageIdToName.Find(PackageId))
 		{
-			UE_LOG(LogIoStore, Fatal, TEXT("Package name hash collision \"%s\" and \"%s"), *ReleasedPackageName->ToString(), *PackageFName.ToString());
+			UE_LOG(LogIoStore, Fatal, TEXT("Package name hash collision \"%s\" and \"%s"), *ReleasedPackageName->ToString(), *PackageName.ToString());
 		}
 
 		Package = new FLegacyCookedPackage();
-		Package->PackageName = PackageFName;
+		Package->PackageName = PackageName;
 		Package->GlobalPackageId = PackageId;
 		Packages.Add(Package);
-		PackageNameMap.Add(PackageFName, Package);
+		PackageNameMap.Add(PackageName, Package);
 		PackageIdMap.Add(PackageId, Package);
 	}
 
 	return Package;
 }
 
-static void ParsePackageAssets(TArray<FLegacyCookedPackage*>& Packages, const FPackageStoreOptimizer& PackageStoreOptimizer)
+FLegacyCookedPackage* FindOrAddPackage(
+	const FIoStoreArguments& Arguments,
+	const TCHAR* RelativeFileName,
+	TArray<FLegacyCookedPackage*>& Packages,
+	FPackageNameMap& PackageNameMap,
+	FPackageIdMap& PackageIdMap)
 {
-	IOSTORE_CPU_SCOPE(ParsePackageAssets);
+	FString PackageName;
+	FString ErrorMessage;
+	if (!FPackageName::TryConvertFilenameToLongPackageName(RelativeFileName, PackageName, &ErrorMessage))
+	{
+		UE_LOG(LogIoStore, Warning, TEXT("Failed to obtain package name from file name '%s'"), *ErrorMessage);
+		return nullptr;
+	}
+
+	FName PackageFName = *PackageName;
+	return FindOrAddPackage(Arguments, PackageFName, Packages, PackageNameMap, PackageIdMap);
+}
+
+static void ParsePackageAssetsFromFiles(TArray<FLegacyCookedPackage*>& Packages, const FPackageStoreOptimizer& PackageStoreOptimizer)
+{
+	IOSTORE_CPU_SCOPE(ParsePackageAssetsFromFiles);
 	UE_LOG(LogIoStore, Display, TEXT("Parsing packages..."));
 
 	TAtomic<int32> ReadCount {0};
@@ -898,6 +1141,10 @@ static void ParsePackageAssets(TArray<FLegacyCookedPackage*>& Packages, const FP
 		{
 			TRACE_CPUPROFILER_EVENT_SCOPE(ReadUAssetFile);
 			FLegacyCookedPackage* Package = Packages[Index];
+			if (!Package->UAssetSize)
+			{
+				return;
+			}
 			uint8* Buffer = PackageAssetBuffers[Index];
 			IFileHandle* FileHandle = FPlatformFileManager::Get().GetPlatformFile().OpenRead(*Package->FileName);
 			if (FileHandle)
@@ -944,6 +1191,27 @@ static void ParsePackageAssets(TArray<FLegacyCookedPackage*>& Packages, const FP
 	FMemory::Free(UAssetMemory);
 }
 
+static void ParsePackageAssetsFromPackageStore(FCookedPackageStore& PackageStore, TArray<FLegacyCookedPackage*>& Packages, const FPackageStoreOptimizer& PackageStoreOptimizer)
+{
+	IOSTORE_CPU_SCOPE(ParsePackageAssetsFromPackageStore);
+	UE_LOG(LogIoStore, Display, TEXT("Parsing packages..."));
+
+	const int32 TotalPackageCount = Packages.Num();
+	
+	ParallelFor(TotalPackageCount, [
+		&PackageStore,
+		&PackageStoreOptimizer,
+		&Packages](int32 Index)
+	{
+		FLegacyCookedPackage* Package = Packages[Index];
+		TIoStatusOr<FIoBuffer> HeaderBuffer = PackageStore.ReadPackageHeader(Package->GlobalPackageId);
+		Package->UExpSize = Package->UAssetSize - HeaderBuffer.ValueOrDie().DataSize();
+		Package->UAssetSize = HeaderBuffer.ValueOrDie().DataSize();
+		Package->OptimizedPackage = PackageStoreOptimizer.CreatePackageFromPackageStoreHeader(Package->PackageName, HeaderBuffer.ValueOrDie(), *Package->PackageStoreEntry);
+		check(Package->OptimizedPackage->GetId() == Package->GlobalPackageId);
+	}, EParallelForFlags::Unbalanced);
+}
+
 TUniquePtr<FIoStoreReader> CreateIoStoreReader(const TCHAR* Path, const FKeyChain& KeyChain)
 {
 	TUniquePtr<FIoStoreReader> IoStoreReader(new FIoStoreReader());
@@ -1155,6 +1423,179 @@ void InitializeContainerTargetsAndPackages(
 		return RelativeFileName;
 	};
 
+	auto CreateTargetFileFromCookedFile = [
+		&Arguments,
+		&ConvertCookedPathToRelativePath,
+		&Packages,
+		&PackageNameMap,
+		&PackageIdMap](const FContainerSourceFile& SourceFile, const FString& RelativeFileName, FContainerTargetFile& OutTargetFile) -> bool
+	{
+		const FCookedFileStatData* OriginalCookedFileStatData = Arguments.CookedFileStatMap.Find(SourceFile.NormalizedPath);
+		if (!OriginalCookedFileStatData)
+		{
+			UE_LOG(LogIoStore, Warning, TEXT("File not found: '%s'"), *SourceFile.NormalizedPath);
+			return false;
+		}
+
+		const FCookedFileStatData* CookedFileStatData = OriginalCookedFileStatData;
+		if (CookedFileStatData->FileType == FCookedFileStatData::PackageHeader)
+		{
+			OutTargetFile.NormalizedSourcePath = FPaths::ChangeExtension(SourceFile.NormalizedPath, TEXT(".uexp"));
+			CookedFileStatData = Arguments.CookedFileStatMap.Find(OutTargetFile.NormalizedSourcePath);
+			if (!CookedFileStatData)
+			{
+				UE_LOG(LogIoStore, Warning, TEXT("File not found: '%s'"), *OutTargetFile.NormalizedSourcePath);
+				return false;
+			}
+		}
+		else
+		{
+			OutTargetFile.NormalizedSourcePath = SourceFile.NormalizedPath;
+		}
+		OutTargetFile.SourceSize = uint64(CookedFileStatData->FileSize);
+		
+		if (CookedFileStatData->FileType == FCookedFileStatData::ShaderLibrary)
+		{
+			OutTargetFile.ChunkType = EContainerChunkType::ShaderCodeLibrary;
+		}
+		else if (CookedFileStatData->FileExt == FCookedFileStatData::EFileExt::UMappedBulk)
+		{
+			OutTargetFile.ChunkType = EContainerChunkType::MemoryMappedBulkData;
+			FString TmpFileName = FString(RelativeFileName.Len() - 8, GetData(RelativeFileName)) + TEXT(".ubulk");
+			OutTargetFile.Package = FindOrAddPackage(Arguments, *TmpFileName, Packages, PackageNameMap, PackageIdMap);
+		}
+		else
+		{
+			OutTargetFile.Package = FindOrAddPackage(Arguments, *RelativeFileName, Packages, PackageNameMap, PackageIdMap);
+			if (CookedFileStatData->FileExt == FCookedFileStatData::UPtnl)
+			{
+				OutTargetFile.ChunkType = EContainerChunkType::OptionalBulkData;
+			}
+			else if (CookedFileStatData->FileType == FCookedFileStatData::BulkData)
+			{
+				OutTargetFile.ChunkType = EContainerChunkType::BulkData;
+			}
+			else
+			{
+				OutTargetFile.ChunkType = EContainerChunkType::PackageData;
+			}
+		}
+
+		switch (OutTargetFile.ChunkType)
+		{
+		case EContainerChunkType::OptionalBulkData:
+			OutTargetFile.ChunkId = CreateIoChunkId(OutTargetFile.Package->GlobalPackageId.Value(), 0, EIoChunkType::OptionalBulkData);
+			break;
+		case EContainerChunkType::MemoryMappedBulkData:
+			OutTargetFile.ChunkId = CreateIoChunkId(OutTargetFile.Package->GlobalPackageId.Value(), 0, EIoChunkType::MemoryMappedBulkData);
+			break;
+		case EContainerChunkType::BulkData:
+			OutTargetFile.ChunkId = CreateIoChunkId(OutTargetFile.Package->GlobalPackageId.Value(), 0, EIoChunkType::BulkData);
+			break;
+		case EContainerChunkType::PackageData:
+			OutTargetFile.Package->FileName = SourceFile.NormalizedPath; // .uasset path
+			OutTargetFile.Package->UAssetSize = OriginalCookedFileStatData->FileSize;
+			OutTargetFile.Package->UExpSize = CookedFileStatData->FileSize;
+			OutTargetFile.ChunkId = CreateIoChunkId(OutTargetFile.Package->GlobalPackageId.Value(), 0, EIoChunkType::ExportBundleData);
+			break;
+		}
+
+		// Only keep the regions for the file if neither compression nor encryption are enabled, otherwise the regions will be meaningless.
+		if (!SourceFile.bNeedsCompression && !SourceFile.bNeedsEncryption)
+		{
+			OutTargetFile.FileRegions = CookedFileStatData->FileRegions;
+		}
+
+		return true;
+	};
+
+	auto CreateTargetFileFromPackageStore = [
+		&Arguments,
+		&ConvertCookedPathToRelativePath,
+		&Packages,
+		&PackageNameMap,
+		&PackageIdMap](const FContainerSourceFile& SourceFile, const FString& RelativeFileName, FContainerTargetFile& OutTargetFile) -> bool
+	{
+		FCookedPackageStore& PackageStore = *Arguments.PackageStore;
+		
+		OutTargetFile.NormalizedSourcePath = SourceFile.NormalizedPath;
+		
+		if (SourceFile.NormalizedPath.EndsWith(TEXT(".ushaderbytecode")))
+		{
+			int64 FileSize = IFileManager::Get().FileSize(*SourceFile.NormalizedPath);
+			if (FileSize == INDEX_NONE)
+			{
+				UE_LOG(LogIoStore, Warning, TEXT("File not found: '%s'"), *SourceFile.NormalizedPath);
+				return false;
+			}
+			OutTargetFile.ChunkType = EContainerChunkType::ShaderCodeLibrary;
+			OutTargetFile.SourceSize = FileSize;
+			return true;
+		}
+
+		OutTargetFile.ChunkId = PackageStore.GetChunkIdFromFileName(SourceFile.NormalizedPath);
+		if (!OutTargetFile.ChunkId.IsValid())
+		{
+			UE_LOG(LogIoStore, Warning, TEXT("File not found in manifest: '%s'"), *SourceFile.NormalizedPath);
+			return false;
+		}
+
+		TIoStatusOr<uint64> ChunkSize = PackageStore.GetChunkSize(OutTargetFile.ChunkId);
+		if (!ChunkSize.IsOk())
+		{
+			UE_LOG(LogIoStore, Warning, TEXT("Chunk size not found for: '%s'"), *SourceFile.NormalizedPath);
+			return false;
+		}
+		OutTargetFile.SourceSize = ChunkSize.ValueOrDie();
+
+		FName PackageName = PackageStore.GetPackageNameFromChunkId(OutTargetFile.ChunkId);
+		if (!PackageName.IsNone())
+		{
+			OutTargetFile.Package = FindOrAddPackage(Arguments, PackageName, Packages, PackageNameMap, PackageIdMap);
+		}
+
+		if (SourceFile.NormalizedPath.EndsWith(TEXT(".m.ubulk")))
+		{
+			OutTargetFile.ChunkType = EContainerChunkType::MemoryMappedBulkData;
+		}
+		else if (SourceFile.NormalizedPath.EndsWith(TEXT(".ubulk")))
+		{
+			OutTargetFile.ChunkType = EContainerChunkType::BulkData;
+		}
+		else if (SourceFile.NormalizedPath.EndsWith(TEXT(".uptnl")))
+		{
+			OutTargetFile.ChunkType = EContainerChunkType::OptionalBulkData;
+		}
+		else
+		{
+			check(OutTargetFile.Package);
+			OutTargetFile.Package->PackageStoreEntry = PackageStore.GetPackageStoreEntry(OutTargetFile.Package->GlobalPackageId);
+			if (!OutTargetFile.Package->PackageStoreEntry)
+			{
+				UE_LOG(LogIoStore, Warning, TEXT("Failed to find package store entry for package: '%s'"), *PackageName.ToString());
+				return false;
+			}
+			OutTargetFile.ChunkType = EContainerChunkType::PackageData;
+			OutTargetFile.Package->FileName = SourceFile.NormalizedPath;
+			OutTargetFile.Package->UAssetSize = OutTargetFile.SourceSize;
+		}
+		
+		// Only keep the regions for the file if neither compression nor encryption are enabled, otherwise the regions will be meaningless.
+		if (Arguments.bFileRegions && !SourceFile.bNeedsCompression && !SourceFile.bNeedsEncryption)
+		{
+			FString RegionsFilename = OutTargetFile.ChunkType == EContainerChunkType::PackageData
+				? FPaths::ChangeExtension(SourceFile.NormalizedPath, FString(TEXT(".uexp")) + FFileRegion::RegionsFileExtension)
+				: SourceFile.NormalizedPath + FFileRegion::RegionsFileExtension;
+			TUniquePtr<FArchive> RegionsFile(IFileManager::Get().CreateFileReader(*RegionsFilename));
+			if (RegionsFile.IsValid())
+			{
+				FFileRegion::SerializeFileRegions(*RegionsFile.Get(), OutTargetFile.FileRegions);
+			}
+		}
+		
+		return true;
+	};
+	
 	for (const FContainerSourceSpec& ContainerSource : Arguments.Containers)
 	{
 		FContainerTargetSpec* ContainerTarget = AddContainer(ContainerSource.Name, ContainerTargets);
@@ -1176,127 +1617,40 @@ void InitializeContainerTargetsAndPackages(
 			IOSTORE_CPU_SCOPE(ProcessSourceFiles);
 			for (const FContainerSourceFile& SourceFile : ContainerSource.SourceFiles)
 			{
-				const FCookedFileStatData* OriginalCookedFileStatData = Arguments.CookedFileStatMap.Find(SourceFile.NormalizedPath);
-				if (!OriginalCookedFileStatData)
-				{
-					UE_LOG(LogIoStore, Warning, TEXT("File not found: '%s'"), *SourceFile.NormalizedPath);
-					continue;
-				}
-				const FCookedFileStatData* CookedFileStatData = OriginalCookedFileStatData;
-				FString NormalizedSourcePath = SourceFile.NormalizedPath;
-				if (CookedFileStatData->FileType == FCookedFileStatData::PackageHeader)
-				{
-					NormalizedSourcePath = FPaths::ChangeExtension(SourceFile.NormalizedPath, TEXT(".uexp"));
-					CookedFileStatData = Arguments.CookedFileStatMap.Find(NormalizedSourcePath);
-					if (!CookedFileStatData)
-					{
-						UE_LOG(LogIoStore, Warning, TEXT("File not found: '%s'"), *NormalizedSourcePath);
-						continue;
-					}
-				}
-
 				FString RelativeFileName = ConvertCookedPathToRelativePath(SourceFile.NormalizedPath);
-				if (CookedFileStatData->FileType == FCookedFileStatData::ShaderLibrary)
+
+				FContainerTargetFile TargetFile;
+				bool bIsValidTargetFile = Arguments.PackageStore.IsValid()
+					? CreateTargetFileFromPackageStore(SourceFile, RelativeFileName, TargetFile)
+					: CreateTargetFileFromCookedFile(SourceFile, RelativeFileName, TargetFile);
+
+				if (!bIsValidTargetFile)
 				{
-					FContainerTargetFile& TargetFile = ContainerTarget->TargetFiles.AddDefaulted_GetRef();
-					TargetFile.ChunkType = EContainerChunkType::ShaderCodeLibrary;
-					TargetFile.ContainerTarget = ContainerTarget;
-					TargetFile.SourceSize = uint64(CookedFileStatData->FileSize);
-					TargetFile.NormalizedSourcePath = NormalizedSourcePath;
-					TargetFile.TargetPath = MoveTemp(RelativeFileName);
-					TargetFile.DestinationPath = SourceFile.DestinationPath;
-					if (SourceFile.bNeedsCompression)
-					{
-						ContainerTarget->ContainerFlags |= EIoContainerFlags::Compressed;
-					}
-					else
-					{
-						TargetFile.bForceUncompressed = true;
-					}
-					if (SourceFile.bNeedsEncryption)
-					{
-						ContainerTarget->ContainerFlags |= EIoContainerFlags::Encrypted;
-					}
 					continue;
 				}
-				
-				
-				FLegacyCookedPackage* Package = nullptr;
-				const bool bIsMemoryMappedBulkData = CookedFileStatData->FileExt == FCookedFileStatData::EFileExt::UMappedBulk;
 
-				if (bIsMemoryMappedBulkData)
+				TargetFile.ContainerTarget			= ContainerTarget;
+				TargetFile.TargetPath				= MoveTemp(RelativeFileName);
+				TargetFile.DestinationPath			= SourceFile.DestinationPath;
+				TargetFile.bForceUncompressed		= !SourceFile.bNeedsCompression;
+				
+				if (SourceFile.bNeedsCompression)
 				{
-					FString TmpFileName = FString(RelativeFileName.Len() - 8, GetData(RelativeFileName)) + TEXT(".ubulk");
-					Package = FindOrAddPackage(Arguments, *TmpFileName, Packages, PackageNameMap, PackageIdMap, PackageStoreOptimizer);
+					ContainerTarget->ContainerFlags |= EIoContainerFlags::Compressed;
 				}
-				else
+
+				if (SourceFile.bNeedsEncryption)
 				{
-					Package = FindOrAddPackage(Arguments, *RelativeFileName, Packages, PackageNameMap, PackageIdMap, PackageStoreOptimizer);
+					ContainerTarget->ContainerFlags |= EIoContainerFlags::Encrypted;
 				}
 
-				if (Package)
+				if (TargetFile.ChunkType == EContainerChunkType::PackageData)
 				{
-					FContainerTargetFile& TargetFile = ContainerTarget->TargetFiles.AddDefaulted_GetRef();
-					TargetFile.ContainerTarget = ContainerTarget;
-					TargetFile.SourceSize = uint64(CookedFileStatData->FileSize);
-					TargetFile.NormalizedSourcePath = NormalizedSourcePath;
-					TargetFile.TargetPath = MoveTemp(RelativeFileName);
-					TargetFile.DestinationPath = SourceFile.DestinationPath;
-					TargetFile.Package = Package;
-					if (SourceFile.bNeedsCompression)
-					{
-						ContainerTarget->ContainerFlags |= EIoContainerFlags::Compressed;
-					}
-					else
-					{
-						TargetFile.bForceUncompressed = true;
-					}
-					if (SourceFile.bNeedsEncryption)
-					{
-						ContainerTarget->ContainerFlags |= EIoContainerFlags::Encrypted;
-					}
-
-					if (CookedFileStatData->FileType == FCookedFileStatData::BulkData)
-					{
-						if (CookedFileStatData->FileExt == FCookedFileStatData::UPtnl)
-						{
-							TargetFile.ChunkType = EContainerChunkType::OptionalBulkData;
-							TargetFile.ChunkId = CreateChunkId(Package->GlobalPackageId, 0, EIoChunkType::OptionalBulkData, *TargetFile.TargetPath);
-						}
-						else if (CookedFileStatData->FileExt == FCookedFileStatData::UMappedBulk)
-						{
-							TargetFile.ChunkType = EContainerChunkType::MemoryMappedBulkData;
-							TargetFile.bForceUncompressed = true;
-							TargetFile.ChunkId = CreateChunkId(Package->GlobalPackageId, 0, EIoChunkType::MemoryMappedBulkData, *TargetFile.TargetPath);
-						}
-						else
-						{
-							TargetFile.ChunkType = EContainerChunkType::BulkData;
-							TargetFile.ChunkId = CreateChunkId(Package->GlobalPackageId, 0, EIoChunkType::BulkData, *TargetFile.TargetPath);
-						}
-						if (Package->FileName.IsEmpty())
-						{
-							Package->FileName = FPaths::ChangeExtension(SourceFile.NormalizedPath, TEXT(".uasset"));
-						}
-					}
-					else
-					{
-						check(CookedFileStatData->FileType == FCookedFileStatData::PackageData);
-						TargetFile.ChunkType = EContainerChunkType::PackageData;
-						ContainerTarget->Packages.Add(Package);
-
-						Package->FileName = SourceFile.NormalizedPath; // .uasset path
-						Package->UAssetSize = OriginalCookedFileStatData->FileSize;
-						Package->UExpSize = CookedFileStatData->FileSize;
-						TargetFile.ChunkId = CreateChunkId(Package->GlobalPackageId, 0, EIoChunkType::ExportBundleData, *TargetFile.TargetPath);
-					}
-
-					if (TargetFile.bForceUncompressed && !SourceFile.bNeedsEncryption)
-					{
-						// Only keep the regions for the file if neither compression nor encryption are enabled, otherwise the regions will be meaningless.
-						TargetFile.FileRegions = CookedFileStatData->FileRegions;
-					}
+					check(TargetFile.Package);
+					ContainerTarget->Packages.Add(TargetFile.Package);
 				}
+
+				ContainerTarget->TargetFiles.Emplace(MoveTemp(TargetFile));
 			}
 		}
 	}
@@ -1428,8 +1782,8 @@ void LogContainerPackageInfo(const TArray<FContainerTargetSpec*>& ContainerTarge
 
 
 	uint64 TotalHeaderSize = 0;
-	uint64 TotalSummarySize = 0;
-	uint64 TotalUGraphSize = 0;
+	uint64 TotalGraphSize = 0;
+	uint64 TotalExportBundleEntriesSize = 0;
 	uint64 TotalImportMapSize = 0;
 	uint64 TotalExportMapSize = 0;
 	uint64 TotalNameMapSize = 0;
@@ -1441,8 +1795,8 @@ void LogContainerPackageInfo(const TArray<FContainerTargetSpec*>& ContainerTarge
 	UE_LOG(LogIoStore, Display, TEXT("%-30s %13s %13s %13s %13s %13s %13s"),
 		TEXT("Container"),
 		TEXT("Header"),
-		TEXT("Summary"),
 		TEXT("Graph"),
+		TEXT("ExportBundleEntries"),
 		TEXT("ImportMap"),
 		TEXT("ExportMap"),
 		TEXT("NameMap"));
@@ -1450,34 +1804,34 @@ void LogContainerPackageInfo(const TArray<FContainerTargetSpec*>& ContainerTarge
 	for (const FContainerTargetSpec* ContainerTarget : ContainerTargets)
 	{
 		uint64 HeaderSize = 0;
-		uint64 SummarySize = ContainerTarget->Packages.Num() * sizeof(FPackageSummary);
-		uint64 UGraphSize = 0;
+		uint64 GraphSize = 0;
+		uint64 ExportBundleEntriesSize = 0;
 		uint64 ImportMapSize = 0;
 		uint64 ExportMapSize = 0;
 		uint64 NameMapSize = 0;
 
 		for (const FLegacyCookedPackage* Package : ContainerTarget->Packages)
 		{
-			UGraphSize += Package->OptimizedPackage->GetExportBundleArcsSize();
+			HeaderSize += Package->OptimizedPackage->GetHeaderSize();
+			GraphSize += Package->OptimizedPackage->GetGraphDataSize();
+			ExportBundleEntriesSize += Package->OptimizedPackage->GetExportBundleEntriesSize();
 			ImportMapSize += Package->OptimizedPackage->GetImportMapSize();
 			ExportMapSize += Package->OptimizedPackage->GetExportMapSize();
 			NameMapSize += Package->OptimizedPackage->GetNameMapSize();
 		}
 
-		HeaderSize = SummarySize + UGraphSize + ImportMapSize + ExportMapSize + NameMapSize;
-
 		UE_LOG(LogIoStore, Display, TEXT("%-30s %13.0lf %13.0lf %13.0lf %13.0lf %13.0lf %13.0lf"),
 			*ContainerTarget->Name.ToString(),
 			(double)HeaderSize / 1024.0,
-			(double)SummarySize / 1024.0,
-			(double)UGraphSize / 1024.0,
+			(double)GraphSize / 1024.0,
+			(double)ExportBundleEntriesSize / 1024.0,
 			(double)ImportMapSize / 1024.0,
 			(double)ExportMapSize / 1024.0,
 			(double)NameMapSize / 1024.0);
 
 		TotalHeaderSize += HeaderSize;
-		TotalSummarySize += SummarySize;
-		TotalUGraphSize += UGraphSize;
+		TotalGraphSize += GraphSize;
+		TotalExportBundleEntriesSize += ExportBundleEntriesSize;
 		TotalImportMapSize += ImportMapSize;
 		TotalExportMapSize += ExportMapSize;
 		TotalNameMapSize += NameMapSize;
@@ -1486,8 +1840,8 @@ void LogContainerPackageInfo(const TArray<FContainerTargetSpec*>& ContainerTarge
 	UE_LOG(LogIoStore, Display, TEXT("%-30s %13.0lf %13.0lf %13.0lf %13.0lf %13.0lf %13.0lf"),
 		TEXT("TOTAL"),
 		(double)TotalHeaderSize / 1024.0,
-		(double)TotalSummarySize / 1024.0,
-		(double)TotalUGraphSize / 1024.0,
+		(double)TotalGraphSize / 1024.0,
+		(double)TotalExportBundleEntriesSize / 1024.0,
 		(double)TotalImportMapSize / 1024.0,
 		(double)TotalExportMapSize / 1024.0,
 		(double)TotalNameMapSize / 1024.0);
@@ -1499,8 +1853,9 @@ void LogContainerPackageInfo(const TArray<FContainerTargetSpec*>& ContainerTarge
 class FIoStoreWriteRequestManager
 {
 public:
-	FIoStoreWriteRequestManager(FPackageStoreOptimizer& InPackageStoreOptimizer)
+	FIoStoreWriteRequestManager(FPackageStoreOptimizer& InPackageStoreOptimizer, FCookedPackageStore* InPackageStore)
 		: PackageStoreOptimizer(InPackageStoreOptimizer)
+		, PackageStore(InPackageStore)
 		, MemoryAvailableEvent(FPlatformProcess::GetSynchEventFromPool(false))
 	{
 		InitiatorThread = Async(EAsyncExecution::Thread, [this]() { InitiatorThreadFunc(); });
@@ -1518,91 +1873,149 @@ public:
 
 	IIoStoreWriteRequest* Read(const FContainerTargetFile& InTargetFile)
 	{
-		return new FWriteContainerTargetFileRequest(*this, InTargetFile);
+		if (PackageStore)
+		{
+			return new FCookedPackageStoreWriteRequest(*this, InTargetFile);
+		}
+		else
+		{
+			return new FLooseFileWriteRequest(*this, InTargetFile);
+		}
 	}
 
 private:
+	struct FQueueEntry;
+	
 	class FWriteContainerTargetFileRequest
 		: public IIoStoreWriteRequest
 	{
+		friend class FIoStoreWriteRequestManager;
+
 	public:
-		FWriteContainerTargetFileRequest(FIoStoreWriteRequestManager& InManager, const FContainerTargetFile& InTargetFile)
-			: Manager(InManager)
-			, TargetFile(InTargetFile)
-			, FileRegions(TargetFile.FileRegions)
+		virtual ~FWriteContainerTargetFileRequest()
 		{
 		}
 
-		~FWriteContainerTargetFileRequest()
+		virtual void PrepareSourceBufferAsync(FGraphEventRef InCompletionEvent) override
 		{
+			CompletionEvent = InCompletionEvent;
+			Manager.ScheduleLoad(this);
 		}
 
-		const FContainerTargetFile& GetTargetFile() const
+		virtual uint64 GetOrderHint() override
 		{
-			return TargetFile;
+			return TargetFile.IdealOrder;
 		}
 
-		void PrepareSourceBufferAsync(FGraphEventRef InCompletionEvent) override
+		virtual TArrayView<const FFileRegion> GetRegions() override
 		{
-			CompletionEvent = InCompletionEvent;
-			Manager.Schedule(this);
+			return FileRegions;
 		}
 
-		const FIoBuffer* GetSourceBuffer() override
+		virtual const FIoBuffer* GetSourceBuffer() override
 		{
-			return SourceBuffer;
+			return &SourceBuffer;
 		}
-
-		void FreeSourceBuffer() override
+		
+		virtual void FreeSourceBuffer() override
 		{
-			delete SourceBuffer;
-			SourceBuffer = nullptr;
-			Manager.OnBufferMemoryFreed(TargetFile.SourceSize);
+			SourceBuffer = FIoBuffer();
+			Manager.OnBufferMemoryFreed(SourceBufferSize);
 		}
 
-		uint64 GetOrderHint() override
+		uint64 GetSourceBufferSize() const
 		{
-			return TargetFile.IdealOrder;
+			return SourceBufferSize;
 		}
 
-		TArrayView<const FFileRegion> GetRegions() override
-		{
-			return FileRegions;
-		}
+		virtual void LoadSourceBufferAsync() = 0;
 
-		FIoBuffer& PrepareSourceBuffer()
-		{
-			check(!SourceBuffer);
-			SourceBuffer = new FIoBuffer(TargetFile.SourceSize);
-			return *SourceBuffer;
-		}
+	protected:
+		FWriteContainerTargetFileRequest(FIoStoreWriteRequestManager& InManager,const FContainerTargetFile& InTargetFile)
+			: Manager(InManager)
+			, TargetFile(InTargetFile)
+			, FileRegions(TargetFile.FileRegions)
+			, SourceBufferSize(TargetFile.SourceSize) { }
 
-		void AsyncReadCallback()
+		void OnSourceBufferLoaded()
 		{
-			//TRACE_CPUPROFILER_EVENT_SCOPE(ReadCallback);
-			if (TargetFile.ChunkType == EContainerChunkType::PackageData)
-			{
-				*SourceBuffer = Manager.PackageStoreOptimizer.CreatePackageBuffer(TargetFile.Package->OptimizedPackage, *SourceBuffer, bHasUpdatedExportBundleRegions ? nullptr : &FileRegions);
-				bHasUpdatedExportBundleRegions = true;
-			}
-			TArray<FBaseGraphTask*> NewTasks;
-			CompletionEvent->DispatchSubsequents(NewTasks);
+			Manager.ScheduleRetire(this);
+			CompletionEvent->DispatchSubsequents();
 		}
 
-	private:
 		FIoStoreWriteRequestManager& Manager;
 		const FContainerTargetFile& TargetFile;
 		TArray<FFileRegion> FileRegions;
+		uint64 SourceBufferSize;
 		FGraphEventRef CompletionEvent;
-		FIoBuffer* SourceBuffer = nullptr;
+		FIoBuffer SourceBuffer;
 		bool bHasUpdatedExportBundleRegions = false;
+		FQueueEntry* QueueEntry = nullptr;
+	};
+
+	// Used when staging from cooked files
+	class FLooseFileWriteRequest
+		: public FWriteContainerTargetFileRequest
+	{
+	public:
+		FLooseFileWriteRequest(FIoStoreWriteRequestManager& InManager,const FContainerTargetFile& InTargetFile)
+			: FWriteContainerTargetFileRequest(InManager, InTargetFile) { }
+
+		virtual void LoadSourceBufferAsync() override
+		{
+			SourceBuffer = FIoBuffer(GetSourceBufferSize());
+
+			QueueEntry->FileHandle.Reset(
+				FPlatformFileManager::Get().GetPlatformFile().OpenAsyncRead(*TargetFile.NormalizedSourcePath));
+			
+			FAsyncFileCallBack Callback = [this](bool, IAsyncReadRequest* ReadRequest)
+			{
+				if (TargetFile.ChunkType == EContainerChunkType::PackageData)
+				{
+					SourceBuffer = Manager.PackageStoreOptimizer.CreatePackageBuffer(TargetFile.Package->OptimizedPackage, SourceBuffer, bHasUpdatedExportBundleRegions ? nullptr : &FileRegions);
+					bHasUpdatedExportBundleRegions = true;
+				}
+				OnSourceBufferLoaded();
+			};
+
+			QueueEntry->ReadRequest.Reset(
+				QueueEntry->FileHandle->ReadRequest(0, SourceBuffer.DataSize(), AIOP_Normal, &Callback, SourceBuffer.Data()));
+		}
+	};
+
+	// Used when cooking directly to I/O store container file
+	class FCookedPackageStoreWriteRequest
+		: public FWriteContainerTargetFileRequest
+	{
+	public:
+		FCookedPackageStoreWriteRequest(FIoStoreWriteRequestManager& InManager,const FContainerTargetFile& InTargetFile)
+			: FWriteContainerTargetFileRequest(InManager, InTargetFile) {}
+
+		virtual void LoadSourceBufferAsync() override
+		{
+			Manager.PackageStore->ReadChunkAsync(
+				TargetFile.ChunkId,
+				[this](TIoStatusOr<FIoBuffer> Status)
+				{
+					SourceBuffer = Status.ConsumeValueOrDie();
+					if (TargetFile.ChunkType == EContainerChunkType::PackageData)
+					{
+						check(TargetFile.Package->UAssetSize > 0);
+						const uint64 HeaderSize = TargetFile.Package->UAssetSize;
+						FIoBuffer ExportsBuffer(SourceBuffer.Data() + HeaderSize, SourceBuffer.DataSize() - HeaderSize, SourceBuffer);
+						SourceBuffer = Manager.PackageStoreOptimizer.CreatePackageBuffer(TargetFile.Package->OptimizedPackage, ExportsBuffer, bHasUpdatedExportBundleRegions ? nullptr : &FileRegions);
+						bHasUpdatedExportBundleRegions = true;
+					}
+					OnSourceBufferLoaded();
+				});
+		}
 	};
 
 	struct FQueueEntry
 	{
 		FQueueEntry* Next = nullptr;
-		IAsyncReadFileHandle* FileHandle = nullptr;
-		IAsyncReadRequest* ReadRequest = nullptr;
+		TUniquePtr<IAsyncReadFileHandle> FileHandle;
+		TUniquePtr<IAsyncReadRequest> ReadRequest;
 		FWriteContainerTargetFileRequest* WriteRequest = nullptr;
 	};
 
@@ -1679,46 +2092,43 @@ private:
 		TAtomic<bool> bIsDoneAdding{ false };
 	};
 
-	void Schedule(FWriteContainerTargetFileRequest* WriteRequest)
+	void ScheduleLoad(FWriteContainerTargetFileRequest* WriteRequest)
 	{
 		FQueueEntry* QueueEntry = new FQueueEntry();
 		QueueEntry->WriteRequest = WriteRequest;
+		WriteRequest->QueueEntry = QueueEntry;
 		InitiatorQueue.Enqueue(QueueEntry);
 	}
 
+	void ScheduleRetire(FWriteContainerTargetFileRequest* WriteRequest)
+	{
+		RetirerQueue.Enqueue(WriteRequest->QueueEntry);
+	}
+
 	void Start(FQueueEntry* QueueEntry)
 	{
-		check(!QueueEntry->ReadRequest);
-		check(!QueueEntry->FileHandle);
-		const FContainerTargetFile& TargetFile = QueueEntry->WriteRequest->GetTargetFile();
+		const uint64 SourceBufferSize = QueueEntry->WriteRequest->GetSourceBufferSize();
 
 		uint64 LocalUsedBufferMemory = UsedBufferMemory.Load();
-		while (LocalUsedBufferMemory > 0 && LocalUsedBufferMemory + TargetFile.SourceSize > BufferMemoryLimit)
+		while (LocalUsedBufferMemory > 0 && LocalUsedBufferMemory + SourceBufferSize > BufferMemoryLimit)
 		{
 			TRACE_CPUPROFILER_EVENT_SCOPE(WaitForBufferMemory);
 			MemoryAvailableEvent->Wait();
 			LocalUsedBufferMemory = UsedBufferMemory.Load();
 		}
-		UsedBufferMemory.AddExchange(TargetFile.SourceSize);
 
-		IAsyncReadFileHandle* FileHandle = FPlatformFileManager::Get().GetPlatformFile().OpenAsyncRead(*TargetFile.NormalizedSourcePath);
-		FAsyncFileCallBack Callback = [this, FileHandle, QueueEntry](bool, IAsyncReadRequest* ReadRequest)
-		{
-			QueueEntry->WriteRequest->AsyncReadCallback();
-			QueueEntry->FileHandle = FileHandle;
-			QueueEntry->ReadRequest = ReadRequest;
-			RetirerQueue.Enqueue(QueueEntry);
-		};
-		FileHandle->ReadRequest(0, TargetFile.SourceSize, AIOP_Normal, &Callback, QueueEntry->WriteRequest->PrepareSourceBuffer().Data());
+		UsedBufferMemory.AddExchange(SourceBufferSize);
+		QueueEntry->WriteRequest->LoadSourceBufferAsync();
 	}
 
 	void Retire(FQueueEntry* QueueEntry)
 	{
-		check(QueueEntry->ReadRequest);
-		check(QueueEntry->FileHandle);
-		QueueEntry->ReadRequest->WaitCompletion();
-		delete QueueEntry->ReadRequest;
-		delete QueueEntry->FileHandle;
+		if (QueueEntry->ReadRequest.IsValid())
+		{
+			QueueEntry->ReadRequest->WaitCompletion();
+			QueueEntry->ReadRequest.Reset();
+			QueueEntry->FileHandle.Reset();
+		}
 		delete QueueEntry;
 	}
 
@@ -1766,6 +2176,7 @@ private:
 	}
 
 	FPackageStoreOptimizer& PackageStoreOptimizer;
+	FCookedPackageStore* PackageStore;
 	TFuture<void> InitiatorThread;
 	TFuture<void> RetirerThread;
 	FQueue InitiatorQueue;
@@ -1790,7 +2201,7 @@ int32 CreateTarget(const FIoStoreArguments& Arguments, const FIoStoreWriterSetti
 
 	FPackageStoreOptimizer PackageStoreOptimizer;
 	PackageStoreOptimizer.Initialize(Arguments.TargetPlatform);
-	FIoStoreWriteRequestManager WriteRequestManager(PackageStoreOptimizer);
+	FIoStoreWriteRequestManager WriteRequestManager(PackageStoreOptimizer, Arguments.PackageStore.Get());
 
 	TArray<FContainerTargetSpec*> ContainerTargets;
 	UE_LOG(LogIoStore, Display, TEXT("Creating container targets..."));
@@ -1862,7 +2273,14 @@ int32 CreateTarget(const FIoStoreArguments& Arguments, const FIoStoreWriterSetti
 		}
 	}
 
-	ParsePackageAssets(Packages, PackageStoreOptimizer);
+	if (Arguments.PackageStore.IsValid())
+	{
+		ParsePackageAssetsFromPackageStore(*Arguments.PackageStore, Packages, PackageStoreOptimizer);
+	}
+	else
+	{
+		ParsePackageAssetsFromFiles(Packages, PackageStoreOptimizer);
+	}
 
 	UE_LOG(LogIoStore, Display, TEXT("Processing shader libraries..."));
 	ProcessShaderLibraries(Arguments, ContainerTargets);
@@ -1890,7 +2308,6 @@ int32 CreateTarget(const FIoStoreArguments& Arguments, const FIoStoreWriterSetti
 		}
 	}
 
-
 	for (FContainerTargetSpec* ContainerTarget : ContainerTargets)
 	{
 		if (ContainerTarget->IoStoreWriter)
@@ -1917,13 +2334,31 @@ int32 CreateTarget(const FIoStoreArguments& Arguments, const FIoStoreWriterSetti
 		}
 	}
 
+	TMap<FPackageId, FPackageStorePackage*> OptimizedPackagesMap;
+	for (FLegacyCookedPackage* Package : Packages)
+	{
+		check(Package->OptimizedPackage);
+		OptimizedPackagesMap.Add(Package->OptimizedPackage->GetId(), Package->OptimizedPackage);
+	}
+
+	if (Arguments.PackageStore.IsValid())
+	{
+		UE_LOG(LogIoStore, Display, TEXT("Processing imports.."));
+		PackageStoreOptimizer.FindImports(OptimizedPackagesMap);
+	}
+
+	UE_LOG(LogIoStore, Display, TEXT("Processing redirects..."));
+	PackageStoreOptimizer.ProcessRedirects(OptimizedPackagesMap);
+
 	UE_LOG(LogIoStore, Display, TEXT("Optimizing packages..."));
+	PackageStoreOptimizer.OptimizeExportBundles(OptimizedPackagesMap);
+
+	UE_LOG(LogIoStore, Display, TEXT("Finalizing packages..."));
 	for (FLegacyCookedPackage* Package : Packages)
 	{
 		check(Package->OptimizedPackage);
-		PackageStoreOptimizer.BeginResolvePackage(Package->OptimizedPackage);
+		PackageStoreOptimizer.FinalizePackage(Package->OptimizedPackage);
 	}
-	PackageStoreOptimizer.Flush(true);
 
 	UE_LOG(LogIoStore, Display, TEXT("Creating disk layout..."));
 	CreateDiskLayout(ContainerTargets, Packages, Arguments.OrderMaps, PackageIdMap, Arguments.bClusterByOrderFilePriority);
@@ -1932,7 +2367,7 @@ int32 CreateTarget(const FIoStoreArguments& Arguments, const FIoStoreWriterSetti
 	{
 		if (ContainerTarget->IoStoreWriter)
 		{
-			TArray<FPackageStoreContainerHeaderEntry> ContainerHeaderEntries;
+			TArray<FPackageStoreEntryResource> PackageStoreEntries;
 			for (FContainerTargetFile& TargetFile : ContainerTarget->TargetFiles)
 			{
 				if (TargetFile.ChunkType == EContainerChunkType::PackageData)
@@ -1950,11 +2385,11 @@ int32 CreateTarget(const FIoStoreArguments& Arguments, const FIoStoreWriterSetti
 					{
 						ContainerTarget->IoStoreWriter->Append(TargetFile.ChunkId, WriteRequestManager.Read(TargetFile), WriteOptions);
 					}
-					ContainerHeaderEntries.Add(PackageStoreOptimizer.CreateContainerHeaderEntry(TargetFile.Package->OptimizedPackage));
+					PackageStoreEntries.Add(PackageStoreOptimizer.CreatePackageStoreEntry(TargetFile.Package->OptimizedPackage));
 				}
 			}
 
-			ContainerTarget->Header = PackageStoreOptimizer.CreateContainerHeader(ContainerTarget->ContainerId, ContainerHeaderEntries);
+			ContainerTarget->Header = PackageStoreOptimizer.CreateContainerHeader(ContainerTarget->ContainerId, PackageStoreEntries);
 			FLargeMemoryWriter HeaderAr(0, true);
 			HeaderAr << ContainerTarget->Header;
 			int64 DataSize = HeaderAr.TotalSize();
@@ -1995,7 +2430,11 @@ int32 CreateTarget(const FIoStoreArguments& Arguments, const FIoStoreWriterSetti
 	uint64 InitialLoadSize = 0;
 	if (GlobalIoStoreWriter)
 	{
-		InitialLoadSize = PackageStoreOptimizer.WriteScriptObjects(GlobalIoStoreWriter);
+		FIoBuffer ScriptObjectsBuffer = PackageStoreOptimizer.CreateScriptObjectsBuffer();
+		InitialLoadSize = ScriptObjectsBuffer.DataSize();
+		FIoWriteOptions WriteOptions;
+		WriteOptions.DebugName = TEXT("ScriptObjects");
+		GlobalIoStoreWriter->Append(CreateIoChunkId(0, 0, EIoChunkType::ScriptObjects), ScriptObjectsBuffer, WriteOptions);
 	}
 
 	UE_LOG(LogIoStore, Display, TEXT("Serializing container(s)..."));
@@ -2004,7 +2443,7 @@ int32 CreateTarget(const FIoStoreArguments& Arguments, const FIoStoreWriterSetti
 	IoStoreWriterResults.Reserve(IoStoreWriters.Num());
 	for (FIoStoreWriter* IoStoreWriter : IoStoreWriters)
 	{
-		TFuture<FIoStoreWriterResult> FlushTask = Async(EAsyncExecution::ThreadPool, [IoStoreWriter]()
+		TFuture<FIoStoreWriterResult> FlushTask = Async(EAsyncExecution::Thread, [IoStoreWriter]()
 		{
 			return IoStoreWriter->Flush().ConsumeValueOrDie();
 		});
@@ -2071,9 +2510,10 @@ int32 CreateTarget(const FIoStoreArguments& Arguments, const FIoStoreWriterSetti
 	UE_LOG(LogIoStore, Display, TEXT("Input:  %8.2f MB Shared shaders"), (double)SharedShaderSize / 1024.0 / 1024.0);
 	UE_LOG(LogIoStore, Display, TEXT("Input:  %8.2f MB Unique shaders"), (double)UniqueShaderSize / 1024.0 / 1024.0);
 	UE_LOG(LogIoStore, Display, TEXT(""));
-	UE_LOG(LogIoStore, Display, TEXT("Output: %8llu Export bundles"), PackageStoreOptimizer.GetTotalExportBundleCount());
 	UE_LOG(LogIoStore, Display, TEXT("Output: %8llu Export bundle entries"), PackageStoreOptimizer.GetTotalExportBundleEntryCount());
-	UE_LOG(LogIoStore, Display, TEXT("Output: %8llu Export bundle arcs"), PackageStoreOptimizer.GetTotalExportBundleArcCount());
+	UE_LOG(LogIoStore, Display, TEXT("Output: %8llu Export bundles"), PackageStoreOptimizer.GetTotalExportBundleCount());
+	UE_LOG(LogIoStore, Display, TEXT("Output: %8llu Internal export bundle arcs"), PackageStoreOptimizer.GetTotalInternalBundleArcsCount());
+	UE_LOG(LogIoStore, Display, TEXT("Output: %8llu External export bundle arcs"), PackageStoreOptimizer.GetTotalExternalBundleArcsCount());
 	UE_LOG(LogIoStore, Display, TEXT("Output: %8llu Name map entries"), NameMapCount);
 	UE_LOG(LogIoStore, Display, TEXT("Output: %8llu Imported package entries"), ImportedPackagesCount);
 	UE_LOG(LogIoStore, Display, TEXT("Output: %8llu Packages without imports"), NoImportedPackagesCount);
@@ -2435,42 +2875,21 @@ int32 Describe(
 		return -1;
 	}
 
-	UE_LOG(LogIoStore, Display, TEXT("Loading global name map..."));
-
-	TIoStatusOr<FIoBuffer> GlobalNamesIoBuffer = GlobalReader->Read(CreateIoChunkId(0, 0, EIoChunkType::LoaderGlobalNames), FIoReadOptions());
-	if (!GlobalNamesIoBuffer.IsOk())
-	{
-		UE_LOG(LogIoStore, Warning, TEXT("Failed reading names chunk from global container '%s'"), *GlobalContainerPath);
-		return -1;
-	}
-
-	TIoStatusOr<FIoBuffer> GlobalNameHashesIoBuffer = GlobalReader->Read(CreateIoChunkId(0, 0, EIoChunkType::LoaderGlobalNameHashes), FIoReadOptions());
-	if (!GlobalNameHashesIoBuffer.IsOk())
-	{
-		UE_LOG(LogIoStore, Warning, TEXT("Failed reading name hashes chunk from global container '%s'"), *GlobalContainerPath);
-		return -1;
-	}
-
-	TArray<FNameEntryId> GlobalNameMap;
-	LoadNameBatch(
-		GlobalNameMap,
-		TArrayView<const uint8>(GlobalNamesIoBuffer.ValueOrDie().Data(), GlobalNamesIoBuffer.ValueOrDie().DataSize()),
-		TArrayView<const uint8>(GlobalNameHashesIoBuffer.ValueOrDie().Data(), GlobalNameHashesIoBuffer.ValueOrDie().DataSize()));
-
 	UE_LOG(LogIoStore, Display, TEXT("Loading script imports..."));
 
-	TIoStatusOr<FIoBuffer> InitialLoadIoBuffer = GlobalReader->Read(CreateIoChunkId(0, 0, EIoChunkType::LoaderInitialLoadMeta), FIoReadOptions());
-	if (!InitialLoadIoBuffer.IsOk())
+	TIoStatusOr<FIoBuffer> ScriptObjectsBuffer = GlobalReader->Read(CreateIoChunkId(0, 0, EIoChunkType::ScriptObjects), FIoReadOptions());
+	if (!ScriptObjectsBuffer.IsOk())
 	{
 		UE_LOG(LogIoStore, Warning, TEXT("Failed reading initial load meta chunk from global container '%s'"), *GlobalContainerPath);
 		return -1;
 	}
 
 	TMap<FPackageObjectIndex, FScriptObjectDesc> ScriptObjectByGlobalIdMap;
-	FLargeMemoryReader InitialLoadArchive(InitialLoadIoBuffer.ValueOrDie().Data(), InitialLoadIoBuffer.ValueOrDie().DataSize());
+	FLargeMemoryReader ScriptObjectsArchive(ScriptObjectsBuffer.ValueOrDie().Data(), ScriptObjectsBuffer.ValueOrDie().DataSize());
+	TArray<FNameEntryId> GlobalNameMap = LoadNameBatch(ScriptObjectsArchive);
 	int32 NumScriptObjects = 0;
-	InitialLoadArchive << NumScriptObjects;
-	const FScriptObjectEntry* ScriptObjectEntries = reinterpret_cast<const FScriptObjectEntry*>(InitialLoadIoBuffer.ValueOrDie().Data() + InitialLoadArchive.Tell());
+	ScriptObjectsArchive << NumScriptObjects;
+	const FScriptObjectEntry* ScriptObjectEntries = reinterpret_cast<const FScriptObjectEntry*>(ScriptObjectsBuffer.ValueOrDie().Data() + ScriptObjectsArchive.Tell());
 	for (int32 ScriptObjectIndex = 0; ScriptObjectIndex < NumScriptObjects; ++ScriptObjectIndex)
 	{
 		const FScriptObjectEntry& ScriptObjectEntry = ScriptObjectEntries[ScriptObjectIndex];
@@ -2579,11 +2998,11 @@ int32 Describe(
 			Job.RawCulturePackageMap = ContainerHeader.CulturePackageMap;
 			Job.RawPackageRedirects = ContainerHeader.PackageRedirects;
 
-			TArrayView<FPackageStoreEntry> StoreEntries(reinterpret_cast<FPackageStoreEntry*>(ContainerHeader.StoreEntries.GetData()), ContainerHeader.PackageCount);
+			TArrayView<FFilePackageStoreEntry> StoreEntries(reinterpret_cast<FFilePackageStoreEntry*>(ContainerHeader.StoreEntries.GetData()), ContainerHeader.PackageCount);
 
 			int32 PackageIndex = 0;
 			Job.Packages.Reserve(StoreEntries.Num());
-			for (FPackageStoreEntry& ContainerEntry : StoreEntries)
+			for (FFilePackageStoreEntry& ContainerEntry : StoreEntries)
 			{
 				const FPackageId& PackageId = ContainerHeader.PackageIds[PackageIndex++];
 				FPackageDesc* PackageDesc = new FPackageDesc();
@@ -2679,25 +3098,20 @@ int32 Describe(
 		check(IoBuffer.IsOk());
 		const uint8* PackageSummaryData = IoBuffer.ValueOrDie().Data();
 		const FPackageSummary* PackageSummary = reinterpret_cast<const FPackageSummary*>(PackageSummaryData);
-		const uint64 PackageSummarySize = PackageSummary->GraphDataOffset + PackageSummary->GraphDataSize;
-		if (PackageSummarySize > IoBuffer.ValueOrDie().DataSize())
+		if (PackageSummary->HeaderSize > IoBuffer.ValueOrDie().DataSize())
 		{
-			ReadOptions.SetRange(0, PackageSummarySize);
+			ReadOptions.SetRange(0, PackageSummary->HeaderSize);
 			IoBuffer = Reader->Read(Job.ChunkId, ReadOptions);
 			PackageSummaryData = IoBuffer.ValueOrDie().Data();
 			PackageSummary = reinterpret_cast<const FPackageSummary*>(PackageSummaryData);
 		}
 
 		TArray<FNameEntryId> PackageNameMap;
-		if (PackageSummary->NameMapNamesSize)
 		{
 			TRACE_CPUPROFILER_EVENT_SCOPE(LoadNameBatch);
-			const uint8* NameMapNamesData = PackageSummaryData + PackageSummary->NameMapNamesOffset;
-			const uint8* NameMapHashesData = PackageSummaryData + PackageSummary->NameMapHashesOffset;
-			LoadNameBatch(
-				PackageNameMap,
-				TArrayView<const uint8>(NameMapNamesData, PackageSummary->NameMapNamesSize),
-				TArrayView<const uint8>(NameMapHashesData, PackageSummary->NameMapHashesSize));
+			TArrayView<const uint8> NameMapView(PackageSummaryData + sizeof(FPackageSummary), PackageSummary->HeaderSize - sizeof(FPackageSummary));
+			FMemoryReaderView NameMapReader(NameMapView);
+			PackageNameMap = LoadNameBatch(NameMapReader);
 		}
 
 		Job.PackageDesc->PackageName = FName::CreateFromDisplayId(PackageNameMap[PackageSummary->Name.GetIndex()], PackageSummary->Name.GetNumber());
@@ -2727,9 +3141,9 @@ int32 Describe(
 			ExportDesc.SerialSize = ExportMapEntry.CookedSerialSize;
 		}
 
-		const FExportBundleHeader* ExportBundleHeaders = reinterpret_cast<const FExportBundleHeader*>(PackageSummaryData + PackageSummary->ExportBundlesOffset);
-		const FExportBundleEntry* ExportBundleEntries = reinterpret_cast<const FExportBundleEntry*>(ExportBundleHeaders + Job.PackageDesc->ExportBundleCount);
-		uint64 CurrentExportOffset = PackageSummarySize;
+		const FExportBundleHeader* ExportBundleHeaders = reinterpret_cast<const FExportBundleHeader*>(PackageSummaryData + PackageSummary->GraphDataOffset);
+		const FExportBundleEntry* ExportBundleEntries = reinterpret_cast<const FExportBundleEntry*>(PackageSummaryData + PackageSummary->ExportBundleEntriesOffset);
+		uint64 CurrentExportOffset = PackageSummary->HeaderSize;
 		for (int32 ExportBundleIndex = 0; ExportBundleIndex < Job.PackageDesc->ExportBundleCount; ++ExportBundleIndex)
 		{
 			TArray<FExportBundleEntryDesc>& ExportBundleDesc = Job.PackageDesc->ExportBundles.AddDefaulted_GetRef();
@@ -3404,6 +3818,282 @@ static int32 Diff(
 	return 0;
 }
 
+int32 Staged2Zen(const FString& BuildPath, const FKeyChain& KeyChain, const ITargetPlatform* TargetPlatform)
+{
+	FString PlatformName = TargetPlatform->PlatformName();
+	FString CookedOutputPath = FPaths::Combine(FPaths::ProjectDir(), TEXT("Saved"), TEXT("Cooked"), PlatformName);
+	if (IFileManager::Get().DirectoryExists(*CookedOutputPath))
+	{
+		UE_LOG(LogIoStore, Error, TEXT("'%s' already exists"), *CookedOutputPath);
+		return -1;
+	}
+
+	TArray<FString> ContainerFiles;
+	IFileManager::Get().FindFilesRecursive(ContainerFiles, *BuildPath, TEXT("*.utoc"), true, false);
+	if (ContainerFiles.IsEmpty())
+	{
+		UE_LOG(LogIoStore, Error, TEXT("No container files found"));
+		return -1;
+	}
+
+	TArray<FString> PakFiles;
+	IFileManager::Get().FindFilesRecursive(PakFiles, *BuildPath, TEXT("*.pak"), true, false);
+	if (PakFiles.IsEmpty())
+	{
+		UE_LOG(LogIoStore, Error, TEXT("No pak files found"));
+		return -1;
+	}
+
+	UE_LOG(LogIoStore, Display, TEXT("Extracting files from paks..."));
+	FPakPlatformFile PakPlatformFile;
+	for (const auto& KV : KeyChain.EncryptionKeys)
+	{
+		FCoreDelegates::GetRegisterEncryptionKeyMulticastDelegate().Broadcast(KV.Key, KV.Value.Key);
+	}
+	PakPlatformFile.Initialize(&FPlatformFileManager::Get().GetPlatformFile(), TEXT(""));
+	FString CookedEngineContentPath = FPaths::Combine(CookedOutputPath, TEXT("Engine"), TEXT("Content"));
+	IFileManager::Get().MakeDirectory(*CookedEngineContentPath, true);
+	FString CookedProjectContentPath = FPaths::Combine(CookedOutputPath, FApp::GetProjectName(), TEXT("Content"));
+	IFileManager::Get().MakeDirectory(*CookedProjectContentPath, true);
+	FString EngineContentPakPath = TEXT("../../../Engine/Content/");
+	FString ProjectContentPakPath = FPaths::Combine(TEXT("../../.."), FApp::GetProjectName(), TEXT("Content"));
+	for (const FString& PakFilePath : PakFiles)
+	{
+		PakPlatformFile.Mount(*PakFilePath, 0);
+		TArray<FString> FilesInPak;
+		PakPlatformFile.GetPrunedFilenamesInPakFile(PakFilePath, FilesInPak);
+		for (const FString& FileInPak : FilesInPak)
+		{
+			FString FileName = FPaths::GetCleanFilename(FileInPak);
+			if (FileName == TEXT("AssetRegistry.bin"))
+			{
+				FString TargetPath = FPaths::Combine(CookedOutputPath, FApp::GetProjectName(), FileName);
+				PakPlatformFile.CopyFile(*TargetPath, *FileInPak);
+			}
+			else if (FileName.EndsWith(TEXT(".ushaderbytecode")))
+			{
+				FString TargetPath = FPaths::Combine(CookedProjectContentPath, FileName);
+				PakPlatformFile.CopyFile(*TargetPath, *FileInPak);
+			}
+			else if (FileName.StartsWith("GlobalShaderCache"))
+			{
+				FString TargetPath = FPaths::Combine(CookedOutputPath, TEXT("Engine"), FileName);
+				PakPlatformFile.CopyFile(*TargetPath, *FileInPak);
+			}
+			else if (FileName.EndsWith(TEXT(".ufont")))
+			{
+				FString TargetPath;
+				if (FileInPak.StartsWith(EngineContentPakPath))
+				{
+					TargetPath = FPaths::Combine(CookedEngineContentPath, *FileInPak + EngineContentPakPath.Len());
+				}
+				else if (FileInPak.StartsWith(ProjectContentPakPath))
+				{
+					TargetPath = FPaths::Combine(CookedProjectContentPath, *FileInPak + ProjectContentPakPath.Len());
+				}
+				else
+				{
+					UE_DEBUG_BREAK();
+					continue;
+				}
+				IFileManager::Get().MakeDirectory(*FPaths::GetPath(TargetPath), true);
+				PakPlatformFile.CopyFile(*TargetPath, *FileInPak);
+			}
+		}
+	}
+
+	struct FBulkDataInfo
+	{
+		FString FileName;
+		IPackageStoreWriter::FBulkDataInfo::EType BulkDataType;
+		TTuple<FIoStoreReader*, FIoChunkId> Chunk;
+	};
+
+	struct FPackageInfo
+	{
+		FName PackageName;
+		FString FileName;
+		TTuple<FIoStoreReader*, FIoChunkId> Chunk;
+		TArray<FBulkDataInfo> BulkData;
+		FPackageStoreEntryResource PackageStoreEntry;
+	};
+
+	struct FCollectedData
+	{
+		TSet<FIoChunkId> SeenChunks;
+		TMap<FName, FPackageInfo> Packages;
+		TMap<FPackageId, FName> PackageIdToName;
+		TArray<TTuple<FIoStoreReader*, FIoChunkId>> ContainerHeaderChunks;
+	} CollectedData;
+
+	UE_LOG(LogIoStore, Display, TEXT("Collecting chunks..."));
+	TArray<TUniquePtr<FIoStoreReader>> IoStoreReaders;
+	IoStoreReaders.Reserve(ContainerFiles.Num());
+	for (const FString& ContainerFilePath : ContainerFiles)
+	{
+		TUniquePtr<FIoStoreReader> Reader = CreateIoStoreReader(*ContainerFilePath, KeyChain);
+		if (!Reader.IsValid())
+		{
+			UE_LOG(LogIoStore, Warning, TEXT("Failed to read container '%s'"), *ContainerFilePath);
+			continue;
+		}
+
+		TMap<FIoChunkId, FString> ChunkFileNamesMap;
+		if (EnumHasAnyFlags(Reader->GetContainerFlags(), EIoContainerFlags::Indexed))
+		{
+			const FIoDirectoryIndexReader& IndexReader = Reader->GetDirectoryIndexReader();
+			IterateDirectoryIndex(
+				FIoDirectoryIndexHandle::RootDirectory(),
+				TEXT(""),
+				IndexReader,
+				[&ChunkFileNamesMap, &Reader](FString Filename, uint32 TocEntryIndex) -> bool
+				{
+					TIoStatusOr<FIoStoreTocChunkInfo> ChunkInfo = Reader->GetChunkInfo(TocEntryIndex);
+					if (ChunkInfo.IsOk())
+					{
+						ChunkFileNamesMap.Add(ChunkInfo.ValueOrDie().Id, Filename);
+					}
+					return true;
+				});
+		}
+		Reader->EnumerateChunks([&ChunkFileNamesMap, &Reader, &CollectedData](const FIoStoreTocChunkInfo& ChunkInfo)
+			{
+				if (CollectedData.SeenChunks.Contains(ChunkInfo.Id))
+				{
+					return true;
+				}
+				CollectedData.SeenChunks.Add(ChunkInfo.Id);
+				EIoChunkType ChunkType = static_cast<EIoChunkType>(ChunkInfo.Id.GetData()[11]);
+				if (ChunkType == EIoChunkType::ExportBundleData ||
+					ChunkType == EIoChunkType::BulkData ||
+					ChunkType == EIoChunkType::OptionalBulkData ||
+					ChunkType == EIoChunkType::MemoryMappedBulkData)
+				{
+					FString PackageNameStr;
+					FString* FindFileName = ChunkFileNamesMap.Find(ChunkInfo.Id);
+					UE_CLOG(!FindFileName, LogIoStore, Fatal, TEXT("Missing file name for package chunk"));
+					if (FPackageName::TryConvertFilenameToLongPackageName(*FindFileName, PackageNameStr, nullptr))
+					{
+						FName PackageName(PackageNameStr);
+						CollectedData.PackageIdToName.Add(FPackageId::FromName(PackageName), PackageName);
+						FPackageInfo& PackageInfo = CollectedData.Packages.FindOrAdd(PackageName);
+						if (ChunkType == EIoChunkType::ExportBundleData)
+						{
+							PackageInfo.FileName = *FindFileName;
+							PackageInfo.PackageName = PackageName;
+							PackageInfo.Chunk = MakeTuple(Reader.Get(), ChunkInfo.Id);
+						}
+						else
+						{
+							FBulkDataInfo& BulkDataInfo = PackageInfo.BulkData.AddDefaulted_GetRef();
+							BulkDataInfo.FileName = *FindFileName;
+							BulkDataInfo.Chunk = MakeTuple(Reader.Get(), ChunkInfo.Id);
+							if (ChunkType == EIoChunkType::OptionalBulkData)
+							{
+								BulkDataInfo.BulkDataType = IPackageStoreWriter::FBulkDataInfo::Optional;
+							}
+							else if (ChunkType == EIoChunkType::MemoryMappedBulkData)
+							{
+								BulkDataInfo.BulkDataType = IPackageStoreWriter::FBulkDataInfo::Mmap;
+							}
+							else
+							{
+								BulkDataInfo.BulkDataType = IPackageStoreWriter::FBulkDataInfo::Standard;
+							}
+						}
+					}
+					else
+					{
+						UE_LOG(LogIoStore, Warning, TEXT("Failed to convert file name '%s' to package name"), **FindFileName);
+					}
+				}
+				else if (ChunkType == EIoChunkType::ContainerHeader)
+				{
+					CollectedData.ContainerHeaderChunks.Emplace(Reader.Get(), ChunkInfo.Id);
+				}
+				return true;
+			});
+
+		IoStoreReaders.Emplace(MoveTemp(Reader));
+	}
+
+	UE_LOG(LogIoStore, Display, TEXT("Reading container headers..."));
+	for (const auto& ContainerHeaderChunk : CollectedData.ContainerHeaderChunks)
+	{
+		FIoBuffer ContainerHeaderBuffer = ContainerHeaderChunk.Key->Read(ContainerHeaderChunk.Value, FIoReadOptions()).ValueOrDie();
+		FMemoryReaderView Ar(MakeArrayView(ContainerHeaderBuffer.Data(), ContainerHeaderBuffer.DataSize()));
+		FContainerHeader ContainerHeader;
+		Ar << ContainerHeader;
+		
+		const FFilePackageStoreEntry* StoreEntry = reinterpret_cast<const FFilePackageStoreEntry*>(ContainerHeader.StoreEntries.GetData());
+		for (const FPackageId& PackageId : ContainerHeader.PackageIds)
+		{
+			const FName* FindPackageName = CollectedData.PackageIdToName.Find(PackageId);
+			if (FindPackageName)
+			{
+				FPackageInfo* FindPackageInfo = CollectedData.Packages.Find(*FindPackageName);
+				check(FindPackageInfo);
+				
+				FPackageStoreEntryResource& PackageStoreEntryResource = FindPackageInfo->PackageStoreEntry;
+				PackageStoreEntryResource.PackageName = *FindPackageName;
+				PackageStoreEntryResource.ExportInfo.ExportBundleCount = StoreEntry->ExportBundleCount;
+				PackageStoreEntryResource.ExportInfo.ExportBundlesSize = StoreEntry->ExportBundlesSize;
+				PackageStoreEntryResource.ExportInfo.ExportCount = StoreEntry->ExportCount;
+				PackageStoreEntryResource.ExportInfo.LoadOrder = StoreEntry->LoadOrder;
+				PackageStoreEntryResource.ImportedPackageIds.SetNum(StoreEntry->ImportedPackages.Num());
+				FMemory::Memcpy(PackageStoreEntryResource.ImportedPackageIds.GetData(), StoreEntry->ImportedPackages.Data(), sizeof(FPackageId) * StoreEntry->ImportedPackages.Num());
+			}
+			++StoreEntry;
+		}
+	}
+
+	FString MetaDataOutputPath = FPaths::Combine(CookedOutputPath, FApp::GetProjectName(), TEXT("Metadata"));
+	TUniquePtr<FZenStoreWriter> ZenStoreWriter = MakeUnique<FZenStoreWriter>(CookedOutputPath, MetaDataOutputPath, TargetPlatform, true);
+	
+	IPackageStoreWriter::FCookInfo CookInfo;
+	ZenStoreWriter->BeginCook(CookInfo);
+	int32 LocalPackageIndex = 0;
+	TArray<FPackageInfo> PackagesArray;
+	CollectedData.Packages.GenerateValueArray(PackagesArray);
+	TAtomic<int32> UploadCount { 0 };
+	ParallelFor(PackagesArray.Num(), [&UploadCount, &PackagesArray, &ZenStoreWriter](int32 Index)
+	{
+		const FPackageInfo& PackageInfo = PackagesArray[Index];
+
+		IPackageStoreWriter::FPackageBaseInfo PackageStorePackageBaseInfo;
+		PackageStorePackageBaseInfo.PackageName = PackageInfo.PackageName;
+
+		ZenStoreWriter->BeginPackage(PackageStorePackageBaseInfo);
+
+		IPackageStoreWriter::FPackageInfo PackageStorePackageInfo;
+		PackageStorePackageInfo.PackageName = PackageInfo.PackageName;
+		PackageStorePackageInfo.LooseFilePath = PackageInfo.FileName;
+		PackageStorePackageInfo.ChunkId = PackageInfo.Chunk.Value;
+
+		FIoBuffer PackageDataBuffer = PackageInfo.Chunk.Key->Read(PackageInfo.Chunk.Value, FIoReadOptions()).ValueOrDie();
+		ZenStoreWriter->WriteIoStorePackageData(PackageStorePackageInfo, PackageDataBuffer, PackageInfo.PackageStoreEntry, TArray<FFileRegion>());
+
+		for (const FBulkDataInfo& BulkDataInfo : PackageInfo.BulkData)
+		{
+			IPackageStoreWriter::FBulkDataInfo PackageStoreBulkDataInfo;
+			PackageStoreBulkDataInfo.PackageName = PackageInfo.PackageName;
+			PackageStoreBulkDataInfo.LooseFilePath = BulkDataInfo.FileName;
+			PackageStoreBulkDataInfo.ChunkId = BulkDataInfo.Chunk.Value;
+			PackageStoreBulkDataInfo.BulkdataType = BulkDataInfo.BulkDataType;
+			FIoBuffer BulkDataBuffer = BulkDataInfo.Chunk.Key->Read(BulkDataInfo.Chunk.Value, FIoReadOptions()).ValueOrDie();
+			ZenStoreWriter->WriteBulkdata(PackageStoreBulkDataInfo, BulkDataBuffer, TArray<FFileRegion>());
+		}
+		ZenStoreWriter->CommitPackage(PackageStorePackageBaseInfo);
+
+		int32 LocalUploadCount = UploadCount.IncrementExchange() + 1;
+		UE_CLOG(LocalUploadCount % 1000 == 0, LogIoStore, Display, TEXT("Uploading package %d/%d"), LocalUploadCount, PackagesArray.Num());
+	}, EParallelForFlags::ForceSingleThread); // Single threaded for now to limit memory usage
+
+	UE_LOG(LogIoStore, Display, TEXT("Waiting for uploads to finish..."));
+	ZenStoreWriter->EndCook();
+	return 0;
+}
+
 static bool ParsePakResponseFile(const TCHAR* FilePath, TArray<FContainerSourceFile>& OutFiles)
 {
 	TArray<FString> ResponseFileContents;
@@ -3979,19 +4669,26 @@ int32 CreateIoStoreContainerFiles(const TCHAR* CmdLine)
 		}
 
 		// Enable file region metadata if required by the target platform.
-		bool bFileRegions = Arguments.TargetPlatform->SupportsFeature(ETargetPlatformFeatures::CookFileRegionMetadata);
-		GeneralIoWriterSettings.bEnableFileRegions = bFileRegions;
+		Arguments.bFileRegions = Arguments.TargetPlatform->SupportsFeature(ETargetPlatformFeatures::CookFileRegionMetadata);
+		GeneralIoWriterSettings.bEnableFileRegions = Arguments.bFileRegions;
 
-		UE_LOG(LogIoStore, Display, TEXT("Searching for cooked assets in folder '%s'"), *Arguments.CookedDir);
-		FCookedFileVisitor CookedFileVistor(Arguments.CookedFileStatMap, nullptr, bFileRegions);
-		IFileManager::Get().IterateDirectoryStatRecursively(*Arguments.CookedDir, CookedFileVistor);
-		UE_LOG(LogIoStore, Display, TEXT("Found '%d' files"), Arguments.CookedFileStatMap.Num());
-
-		int32 ReturnValue = CreateTarget(Arguments, GeneralIoWriterSettings);
-		if (ReturnValue != 0)
+		TUniquePtr<FCookedPackageStore> PackageStore = MakeUnique<FCookedPackageStore>();
+		FIoStatus Status = PackageStore->Load(Arguments.CookedDir);
+		if (Status.IsOk())
 		{
-			return ReturnValue;
+			UE_LOG(LogIoStore, Display, TEXT("Using package store"));
+			Arguments.PackageStore = MoveTemp(PackageStore);
 		}
+		else
+		{
+			UE_LOG(LogIoStore, Display, TEXT("Searching for cooked assets in folder '%s'"), *Arguments.CookedDir);
+			FCookedFileVisitor CookedFileVistor(Arguments.CookedFileStatMap, nullptr, Arguments.bFileRegions);
+			IFileManager::Get().IterateDirectoryStatRecursively(*Arguments.CookedDir, CookedFileVistor);
+			UE_LOG(LogIoStore, Display, TEXT("Found '%d' files"), Arguments.CookedFileStatMap.Num());
+		}
+
+		const int32 ReturnValue = CreateTarget(Arguments, GeneralIoWriterSettings);
+		return ReturnValue;
 	}
 	else if (FParse::Param(FCommandLine::Get(), TEXT("CreateContentPatch")))
 	{
@@ -4077,6 +4774,15 @@ int32 CreateIoStoreContainerFiles(const TCHAR* CmdLine)
 			
 			return Diff(SourcePath, SourceKeyChain, TargetPath, TargetKeyChain, OutPath);
 		}
+		else if (FParse::Param(FCommandLine::Get(), TEXT("Staged2Zen")))
+		{
+			FString BuildPath;
+			if (!FParse::Value(FCommandLine::Get(), TEXT("BuildPath="), BuildPath) || !Arguments.TargetPlatform)
+			{
+				UE_LOG(LogIoStore, Error, TEXT("Incorrect arguments. Expected: -Staged2Zen -BuildPath=<Path> -TargetPlatform=<Platform>"));
+			}
+			return Staged2Zen(BuildPath, Arguments.KeyChain, Arguments.TargetPlatform);
+		}
 		else
 		{
 			UE_LOG(LogIoStore, Display, TEXT("Usage:"));
