Integrating //UE5/Dev-Cooker @ 16678003 to //UE5/Main (Zousar.Shaker-YEG-0943-Quaternary)

diff --git a/Engine/Source/Developer/IoStoreUtilities/Private/PackageStoreOptimizer.cpp b/Engine/Source/Developer/IoStoreUtilities/Private/PackageStoreOptimizer.cpp
--- a/Engine/Source/Developer/IoStoreUtilities/Private/PackageStoreOptimizer.cpp
+++ b/Engine/Source/Developer/IoStoreUtilities/Private/PackageStoreOptimizer.cpp
@@ -4,10 +4,11 @@
 #include "ProfilingDebugging/CountersTrace.h"
 #include "Serialization/BufferWriter.h"
 #include "Serialization/LargeMemoryWriter.h"
+#include "Serialization/MemoryReader.h"
 #include "UObject/NameBatchSerialization.h"
 #include "Containers/Map.h"
 #include "UObject/UObjectHash.h"
-#include "Serialization/PackageStore.h"
+#include "Serialization/FilePackageStore.h"
 #include "Interfaces/ITargetPlatform.h"
 #include "UObject/Object.h"
 #include "Serialization/MemoryReader.h"
@@ -110,41 +111,71 @@ FPackageStorePackage* FPackageStoreOptimizer::CreatePackageFromCookedHeader(cons
 
 	Package->SourceName = *RemapLocalizationPathIfNeeded(Name.ToString(), &Package->Region);
 
-	LoadCookedHeader(Package, CookedHeaderBuffer);
-	ProcessImports(Package);
-	ProcessExports(Package);
+	FCookedHeaderData CookedHeaderData = LoadCookedHeader(CookedHeaderBuffer);
+	Package->PackageFlags = CookedHeaderData.Summary.GetPackageFlags();
+	Package->CookedHeaderSize = CookedHeaderData.Summary.TotalHeaderSize;
+	for (const FName& SummaryName : CookedHeaderData.SummaryNames)
+	{
+		Package->NameMapBuilder.AddName(SummaryName);
+	}
+	ProcessImports(CookedHeaderData, Package);
+	ProcessExports(CookedHeaderData, Package);
+	ProcessPreloadDependencies(CookedHeaderData, Package);
+	
+	CreateExportBundles(Package);
 
 	return Package;
 }
 
-void FPackageStoreOptimizer::LoadCookedHeader(FPackageStorePackage* Package, const FIoBuffer& CookedHeaderBuffer) const
+FPackageStorePackage* FPackageStoreOptimizer::CreatePackageFromPackageStoreHeader(const FName& Name, const FIoBuffer& Buffer, const FPackageStoreEntryResource& PackageStoreEntry) const
 {
+	FPackageStorePackage* Package = new FPackageStorePackage();
+	Package->Name = Name;
+	Package->Id = FPackageId::FromName(Name);
+	Package->LoadOrder = PackageStoreEntry.ExportInfo.LoadOrder;
+	
+	Package->SourceName = *RemapLocalizationPathIfNeeded(Name.ToString(), &Package->Region);
+
+	FPackageStoreHeaderData PackageStoreHeaderData = LoadPackageStoreHeader(Buffer, PackageStoreEntry);
+	Package->PackageFlags = PackageStoreHeaderData.Summary.PackageFlags;
+	Package->CookedHeaderSize = PackageStoreHeaderData.Summary.CookedHeaderSize;
+	for (const FNameEntryId& DisplayIndex : PackageStoreHeaderData.NameMap)
+	{
+		Package->NameMapBuilder.AddName(FName::CreateFromDisplayId(DisplayIndex, 0));
+	}
+	ProcessImports(PackageStoreHeaderData, Package);
+	ProcessExports(PackageStoreHeaderData, Package);
+	ProcessPreloadDependencies(PackageStoreHeaderData, Package);
+	CreateExportBundles(Package);
+
+	return Package;
+}
+
+PRAGMA_DISABLE_DEPRECATION_WARNINGS
+FPackageStoreOptimizer::FCookedHeaderData FPackageStoreOptimizer::LoadCookedHeader(const FIoBuffer& CookedHeaderBuffer) const
+{
+	FCookedHeaderData CookedHeaderData;
 	TArrayView<const uint8> MemView(CookedHeaderBuffer.Data(), CookedHeaderBuffer.DataSize());
 	FMemoryReaderView Ar(MemView);
 
+	FPackageFileSummary& Summary = CookedHeaderData.Summary;
 	{
 		TGuardValue<int32> GuardAllowUnversionedContentInEditor(GAllowUnversionedContentInEditor, 1);
-		Ar << Package->Summary;
+		Ar << Summary;
 	}
 
-	Package->PackageFlags = Package->Summary.GetPackageFlags();
-	Package->CookedHeaderSize = Package->Summary.TotalHeaderSize;
-
-	Ar.SetFilterEditorOnly((Package->PackageFlags & EPackageFlags::PKG_FilterEditorOnly) != 0);
+	Ar.SetFilterEditorOnly((CookedHeaderData.Summary.GetPackageFlags() & EPackageFlags::PKG_FilterEditorOnly) != 0);
 
-	if (Package->Summary.NameCount > 0)
+	if (Summary.NameCount > 0)
 	{
-		Ar.Seek(Package->Summary.NameOffset);
+		Ar.Seek(Summary.NameOffset);
 
 		FNameEntrySerialized NameEntry(ENAME_LinkerConstructor);
 
-		//Package.SummaryNames.Reserve(Summary.NameCount);
-		for (int32 I = 0; I < Package->Summary.NameCount; ++I)
+		for (int32 I = 0; I < Summary.NameCount; ++I)
 		{
 			Ar << NameEntry;
-			FName Name(NameEntry);
-			//Package.SummaryNames.Add(Name);
-			Package->NameMapBuilder.AddName(Name);
+			CookedHeaderData.SummaryNames.Add(NameEntry);
 		}
 	}
 
@@ -154,7 +185,7 @@ void FPackageStoreOptimizer::LoadCookedHeader(FPackageStorePackage* Package, con
 	public:
 		using FArchiveProxy::FArchiveProxy;
 
-		FNameReaderProxyArchive(FArchive& InAr, const TArray<FNameEntryId>& InNameMap)
+		FNameReaderProxyArchive(FArchive& InAr, const TArray<FName>& InNameMap)
 			: FArchiveProxy(InAr)
 			, NameMap(InNameMap)
 		{
@@ -174,48 +205,127 @@ void FPackageStoreOptimizer::LoadCookedHeader(FPackageStorePackage* Package, con
 				UE_LOG(LogPackageStoreOptimizer, Fatal, TEXT("Bad name index %i/%i"), NameIndex, NameMap.Num());
 			}
 
-			const FNameEntryId MappedName = NameMap[NameIndex];
-			Name = FName::CreateFromDisplayId(MappedName, Number);
+			const FName& MappedName = NameMap[NameIndex];
+			Name = FName::CreateFromDisplayId(MappedName.GetDisplayIndex(), Number);
 
 			return *this;
 		}
 
 	private:
-		const TArray<FNameEntryId>& NameMap;
+		const TArray<FName>& NameMap;
 	};
-	FNameReaderProxyArchive ProxyAr(Ar, Package->NameMapBuilder.GetNameMap());
+	FNameReaderProxyArchive ProxyAr(Ar, CookedHeaderData.SummaryNames);
 
-	if (Package->Summary.ImportCount > 0)
+	if (Summary.ImportCount > 0)
 	{
-		Package->ObjectImports.Reserve(Package->Summary.ImportCount);
-		ProxyAr.Seek(Package->Summary.ImportOffset);
-		for (int32 I = 0; I < Package->Summary.ImportCount; ++I)
+		CookedHeaderData.ObjectImports.Reserve(Summary.ImportCount);
+		ProxyAr.Seek(Summary.ImportOffset);
+		for (int32 I = 0; I < Summary.ImportCount; ++I)
 		{
-			ProxyAr << Package->ObjectImports.AddDefaulted_GetRef();
+			ProxyAr << CookedHeaderData.ObjectImports.AddDefaulted_GetRef();
 		}
 	}
 
-	if (Package->Summary.PreloadDependencyCount > 0)
+	if (Summary.PreloadDependencyCount > 0)
 	{
-		Package->PreloadDependencies.Reserve(Package->Summary.PreloadDependencyCount);
-		ProxyAr.Seek(Package->Summary.PreloadDependencyOffset);
-		for (int32 I = 0; I < Package->Summary.PreloadDependencyCount; ++I)
+		CookedHeaderData.PreloadDependencies.Reserve(Summary.PreloadDependencyCount);
+		ProxyAr.Seek(Summary.PreloadDependencyOffset);
+		for (int32 I = 0; I < Summary.PreloadDependencyCount; ++I)
 		{
-			ProxyAr << Package->PreloadDependencies.AddDefaulted_GetRef();
+			ProxyAr << CookedHeaderData.PreloadDependencies.AddDefaulted_GetRef();
 		}
 	}
 
-	if (Package->Summary.ExportCount > 0)
+	if (Summary.ExportCount > 0)
 	{
-		Package->ObjectExports.Reserve(Package->Summary.ExportCount);
-		ProxyAr.Seek(Package->Summary.ExportOffset);
-		for (int32 I = 0; I < Package->Summary.ExportCount; ++I)
+		CookedHeaderData.ObjectExports.Reserve(Summary.ExportCount);
+		ProxyAr.Seek(Summary.ExportOffset);
+		for (int32 I = 0; I < Summary.ExportCount; ++I)
 		{
-			FObjectExport& ObjectExport = Package->ObjectExports.AddDefaulted_GetRef();
+			FObjectExport& ObjectExport = CookedHeaderData.ObjectExports.AddDefaulted_GetRef();
 			ProxyAr << ObjectExport;
-			Package->ExportsSerialSize += ObjectExport.SerialSize;
 		}
 	}
+
+	return CookedHeaderData;
+}
+PRAGMA_ENABLE_DEPRECATION_WARNINGS
+
+FPackageStoreOptimizer::FPackageStoreHeaderData FPackageStoreOptimizer::LoadPackageStoreHeader(const FIoBuffer& PackageStoreHeaderBuffer, const FPackageStoreEntryResource& PackageStoreEntry) const
+{
+	FPackageStoreHeaderData PackageStoreHeaderData;
+
+	const uint8* HeaderData = PackageStoreHeaderBuffer.Data();
+
+	FPackageSummary& Summary = PackageStoreHeaderData.Summary;
+	Summary = *reinterpret_cast<const FPackageSummary*>(HeaderData);
+	check(PackageStoreHeaderBuffer.DataSize() == Summary.HeaderSize);
+
+	TArray<FNameEntryId>& NameMap = PackageStoreHeaderData.NameMap;
+	TArrayView<const uint8> NameMapView(HeaderData + sizeof(FPackageSummary), Summary.HeaderSize - sizeof(FPackageSummary));
+	FMemoryReaderView NameMapReader(NameMapView);
+	NameMap = LoadNameBatch(NameMapReader);
+
+	for (FPackageId PackageId : PackageStoreEntry.ImportedPackageIds)
+	{
+		PackageStoreHeaderData.ImportedPackageIds.Add(PackageId);
+	}
+
+	PackageStoreHeaderData.Imports =
+		MakeArrayView<const FPackageObjectIndex>(
+			reinterpret_cast<const FPackageObjectIndex*>(HeaderData + Summary.ImportMapOffset),
+			(Summary.ExportMapOffset - Summary.ImportMapOffset) / sizeof(FPackageObjectIndex));
+
+	PackageStoreHeaderData.Exports =
+		MakeArrayView<const FExportMapEntry>(
+			reinterpret_cast<const FExportMapEntry*>(HeaderData + Summary.ExportMapOffset),
+			(Summary.ExportBundleEntriesOffset - Summary.ExportMapOffset) / sizeof(FExportMapEntry));
+
+	PackageStoreHeaderData.ExportBundleHeaders =
+		MakeArrayView<const FExportBundleHeader>(
+			reinterpret_cast<const FExportBundleHeader*>(HeaderData + Summary.GraphDataOffset),
+			PackageStoreEntry.ExportInfo.ExportBundleCount);
+
+	PackageStoreHeaderData.ExportBundleEntries = 
+		MakeArrayView<const FExportBundleEntry>(
+			reinterpret_cast<const FExportBundleEntry*>(HeaderData + Summary.ExportBundleEntriesOffset),
+			PackageStoreEntry.ExportInfo.ExportCount * FExportBundleEntry::ExportCommandType_Count);
+
+	const uint64 ExportBundleHeadersSize = sizeof(FExportBundleHeader) * PackageStoreEntry.ExportInfo.ExportBundleCount;
+	const uint64 ArcsDataOffset = Summary.GraphDataOffset + ExportBundleHeadersSize;
+	const uint64 ArcsDataSize = Summary.HeaderSize - ArcsDataOffset;
+
+	FMemoryReaderView ArcsAr(MakeArrayView<const uint8>(HeaderData + ArcsDataOffset, ArcsDataSize));
+
+	int32 InternalArcsCount = 0;
+	ArcsAr << InternalArcsCount;
+
+	for (int32 Idx = 0; Idx < InternalArcsCount; ++Idx)
+	{
+		FPackageStorePackage::FInternalArc& InternalArc = PackageStoreHeaderData.InternalArcs.AddDefaulted_GetRef();
+		ArcsAr << InternalArc.FromExportBundleIndex;
+		ArcsAr << InternalArc.ToExportBundleIndex;
+	}
+
+	for (FPackageId ImportedPackageId : PackageStoreHeaderData.ImportedPackageIds)
+	{
+		int32 ExternalArcsCount = 0;
+		ArcsAr << ExternalArcsCount;
+
+		for (int32 Idx = 0; Idx < ExternalArcsCount; ++Idx)
+		{
+			FPackageStorePackage::FExternalArc ExternalArc;
+			ArcsAr << ExternalArc.FromImportIndex;
+			uint8 FromCommandType = 0;
+			ArcsAr << FromCommandType;
+			ExternalArc.FromCommandType = static_cast<FExportBundleEntry::EExportCommandType>(FromCommandType);
+			ArcsAr << ExternalArc.ToExportBundleIndex;
+
+			PackageStoreHeaderData.ExternalArcs.Add(ExternalArc);
+		}
+	}
+
+	return PackageStoreHeaderData;
 }
 
 void FPackageStoreOptimizer::ResolveImport(FPackageStorePackage::FImport* Imports, const FObjectImport* ObjectImports, int32 LocalImportIndex) const
@@ -251,14 +361,15 @@ void FPackageStoreOptimizer::ResolveImport(FPackageStorePackage::FImport* Import
 	}
 }
 
-void FPackageStoreOptimizer::ProcessImports(FPackageStorePackage* Package) const
+void FPackageStoreOptimizer::ProcessImports(const FCookedHeaderData& CookedHeaderData, FPackageStorePackage* Package) const
 {
-	int32 ImportCount = Package->ObjectImports.Num();
+	int32 ImportCount = CookedHeaderData.ObjectImports.Num();
 	Package->Imports.SetNum(ImportCount);
 
+	TSet<FPackageId> ImportedPackageIds;
 	for (int32 ImportIndex = 0; ImportIndex < ImportCount; ++ImportIndex)
 	{
-		ResolveImport(Package->Imports.GetData(), Package->ObjectImports.GetData(), ImportIndex);
+		ResolveImport(Package->Imports.GetData(), CookedHeaderData.ObjectImports.GetData(), ImportIndex);
 		FPackageStorePackage::FImport& Import = Package->Imports[ImportIndex];
 		if (Import.bIsScriptImport)
 		{
@@ -270,9 +381,25 @@ void FPackageStoreOptimizer::ProcessImports(FPackageStorePackage* Package) const
 		}
 		else
 		{
-			Package->ImportedPackageIds.Add(Import.PackageId);
+			ImportedPackageIds.Add(Import.PackageId);
 		}
 	}
+	Package->ImportedPackageIds = ImportedPackageIds.Array();
+	Algo::Sort(Package->ImportedPackageIds);
+}
+
+void FPackageStoreOptimizer::ProcessImports(const FPackageStoreHeaderData& PackageStoreHeaderData, FPackageStorePackage* Package) const
+{
+	Package->ImportedPackageIds = PackageStoreHeaderData.ImportedPackageIds;
+
+	for (FPackageObjectIndex ImportObjectIndex : PackageStoreHeaderData.Imports)
+	{
+		FPackageStorePackage::FImport& Import = Package->Imports.AddDefaulted_GetRef();
+		Import.GlobalImportIndex = ImportObjectIndex;
+		Import.bIsScriptImport = ImportObjectIndex.IsScriptImport();
+		Import.bIsPackageImport = ImportObjectIndex.IsPackageImport();
+		// Import.FullName and Import.PackageId will be resolve in FindImports after all package headers have been loaded
+	}
 }
 
 void FPackageStoreOptimizer::ResolveExport(
@@ -310,9 +437,38 @@ void FPackageStoreOptimizer::ResolveExport(
 	}
 }
 
-void FPackageStoreOptimizer::ProcessExports(FPackageStorePackage* Package) const
+void FPackageStoreOptimizer::ResolveExport(FPackageStorePackage::FExport* Exports, const int32 LocalExportIndex, const FName& PackageName) const
 {
-	int32 ExportCount = Package->ObjectExports.Num();
+	FPackageStorePackage::FExport* Export = Exports + LocalExportIndex;
+	if (Export->FullName.Len() == 0)
+	{
+		Export->FullName.Reserve(256);
+		if (Export->OuterIndex.IsNull())
+		{
+			PackageName.AppendString(Export->FullName);
+			Export->FullName.AppendChar(TEXT('/'));
+			Export->ObjectName.AppendString(Export->FullName);
+			Export->FullName.ToLowerInline();
+			check(Export->FullName.Len() > 0);
+		}
+		else
+		{
+			check(Export->OuterIndex.IsExport());
+			int32 OuterExportIndex = Export->OuterIndex.ToExport();
+			ResolveExport(Exports, OuterExportIndex, PackageName);
+			FString& OuterName = Exports[OuterExportIndex].FullName;
+			check(OuterName.Len() > 0);
+			Export->FullName.Append(OuterName);
+			Export->FullName.AppendChar(TEXT('/'));
+			Export->ObjectName.AppendString(Export->FullName);
+			Export->FullName.ToLowerInline();
+		}
+	}
+}
+
+void FPackageStoreOptimizer::ProcessExports(const FCookedHeaderData& CookedHeaderData, FPackageStorePackage* Package) const
+{
+	int32 ExportCount = CookedHeaderData.ObjectExports.Num();
 	Package->Exports.SetNum(ExportCount);
 	Package->ExportGraphNodes.Reserve(ExportCount * 2);
 
@@ -332,16 +488,19 @@ void FPackageStoreOptimizer::ProcessExports(FPackageStorePackage* Package) const
 
 	for (int32 ExportIndex = 0; ExportIndex < ExportCount; ++ExportIndex)
 	{
-		const FObjectExport& ObjectExport = Package->ObjectExports[ExportIndex];
+		const FObjectExport& ObjectExport = CookedHeaderData.ObjectExports[ExportIndex];
+		Package->ExportsSerialSize += ObjectExport.SerialSize;
+
 		FPackageStorePackage::FExport& Export = Package->Exports[ExportIndex];
 		Export.ObjectName = ObjectExport.ObjectName;
 		Export.ObjectFlags = ObjectExport.ObjectFlags;
 		Export.CookedSerialOffset = ObjectExport.SerialOffset;
-		Export.CookedSerialSize = ObjectExport.SerialSize;
+		Export.SerialOffset = ObjectExport.SerialOffset - CookedHeaderData.Summary.TotalHeaderSize;
+		Export.SerialSize = ObjectExport.SerialSize;
 		Export.bNotForClient = ObjectExport.bNotForClient;
 		Export.bNotForServer = ObjectExport.bNotForServer;
 		Export.bIsPublic = (Export.ObjectFlags & RF_Public) > 0;
-		ResolveExport(Package->Exports.GetData(), Package->ObjectExports.GetData(), ExportIndex, Package->Name);
+		ResolveExport(Package->Exports.GetData(), CookedHeaderData.ObjectExports.GetData(), ExportIndex, Package->Name);
 		if (Export.bIsPublic)
 		{
 			check(Export.FullName.Len() > 0);
@@ -353,21 +512,84 @@ void FPackageStoreOptimizer::ProcessExports(FPackageStorePackage* Package) const
 		Export.SuperIndex = PackageObjectIdFromPackageIndex(Package->Imports, ObjectExport.SuperIndex);
 		Export.TemplateIndex = PackageObjectIdFromPackageIndex(Package->Imports, ObjectExport.TemplateIndex);
 
-		Export.CreateNode = &Package->ExportGraphNodes.AddDefaulted_GetRef();
-		Export.CreateNode->Package = Package;
-		Export.CreateNode->BundleEntry.CommandType = FExportBundleEntry::ExportCommandType_Create;
-		Export.CreateNode->BundleEntry.LocalExportIndex = ExportIndex;
+		for (uint8 CommandType = 0; CommandType < FExportBundleEntry::ExportCommandType_Count; ++CommandType)
+		{
+			FPackageStorePackage::FExportGraphNode& Node = Package->ExportGraphNodes.AddDefaulted_GetRef();
+			Node.BundleEntry.CommandType = FExportBundleEntry::EExportCommandType(CommandType);
+			Node.BundleEntry.LocalExportIndex = ExportIndex;
+			Node.bIsPublic = Export.bIsPublic;
+			Export.Nodes[CommandType] = &Node;
+		}
+	}
+}
+
+void FPackageStoreOptimizer::ProcessExports(const FPackageStoreHeaderData& PackageStoreHeaderData, FPackageStorePackage* Package) const
+{
+	const int32 ExportCount = PackageStoreHeaderData.Exports.Num();
+	Package->Exports.SetNum(ExportCount);
+	Package->ExportGraphNodes.Reserve(ExportCount * 2);
+
+	const TArray<FNameEntryId>& NameMap = PackageStoreHeaderData.NameMap;
+
+	for (int32 ExportIndex = 0; ExportIndex < ExportCount; ++ExportIndex)
+	{
+		const FExportMapEntry& ExportEntry =  PackageStoreHeaderData.Exports[ExportIndex];
+		Package->ExportsSerialSize += ExportEntry.CookedSerialSize;
+
+		FPackageStorePackage::FExport& Export = Package->Exports[ExportIndex];
+		Export.ObjectName = FName::CreateFromDisplayId(NameMap[ExportEntry.ObjectName.GetIndex()], ExportEntry.ObjectName.GetNumber());
+		Export.OuterIndex = ExportEntry.OuterIndex;
+		Export.ClassIndex = ExportEntry.ClassIndex;
+		Export.SuperIndex = ExportEntry.SuperIndex;
+		Export.TemplateIndex = ExportEntry.TemplateIndex;
+		Export.GlobalImportIndex = ExportEntry.GlobalImportIndex;
+		Export.ObjectFlags = ExportEntry.ObjectFlags;
+		Export.CookedSerialOffset = ExportEntry.CookedSerialOffset;
+		Export.SerialSize = ExportEntry.CookedSerialSize;
+		Export.bNotForClient = ExportEntry.FilterFlags == EExportFilterFlags::NotForClient;
+		Export.bNotForServer = ExportEntry.FilterFlags == EExportFilterFlags::NotForServer;
+		Export.bIsPublic = (ExportEntry.ObjectFlags & RF_Public) > 0;
+
+		for (uint8 CommandType = 0; CommandType < FExportBundleEntry::ExportCommandType_Count; ++CommandType)
+		{
+			FPackageStorePackage::FExportGraphNode& Node = Package->ExportGraphNodes.AddDefaulted_GetRef();
+			Node.BundleEntry.CommandType = FExportBundleEntry::EExportCommandType(CommandType);
+			Node.BundleEntry.LocalExportIndex = ExportIndex;
+			Node.bIsPublic = Export.bIsPublic;
+			Export.Nodes[CommandType] = &Node;
+		}
+
+		Package->NameMapBuilder.MarkNameAsReferenced(Export.ObjectName);
+	}
 
-		Export.SerializeNode = &Package->ExportGraphNodes.AddDefaulted_GetRef();
-		Export.SerializeNode->Package = Package;
-		Export.SerializeNode->BundleEntry.CommandType = FExportBundleEntry::ExportCommandType_Serialize;
-		Export.SerializeNode->BundleEntry.LocalExportIndex = ExportIndex;
+	for (int32 ExportIndex = 0; ExportIndex < ExportCount; ++ExportIndex)
+	{
+		ResolveExport(Package->Exports.GetData(), ExportIndex, Package->Name);
 	}
+
+	uint64 ExportSerialOffset = 0;
+	for (const FExportBundleHeader& ExportBundleHeader : PackageStoreHeaderData.ExportBundleHeaders)
+	{
+		int32 ExportEntryIndex = ExportBundleHeader.FirstEntryIndex;
+		for (uint32 Idx = 0; Idx < ExportBundleHeader.EntryCount; ++Idx)
+		{
+			const FExportBundleEntry& BundleEntry = PackageStoreHeaderData.ExportBundleEntries[ExportEntryIndex++];
+			if (BundleEntry.CommandType == FExportBundleEntry::ExportCommandType_Serialize)
+			{
+				FPackageStorePackage::FExport& Export = Package->Exports[BundleEntry.LocalExportIndex];
+				Export.SerialOffset = ExportSerialOffset;
+				ExportSerialOffset += Export.SerialSize;
+			}
+		}
+	}
+	check(ExportSerialOffset == Package->ExportsSerialSize);
 }
 
-void FPackageStoreOptimizer::SortPackagesInLoadOrder(TArray<FPackageStorePackage*>& Packages) const
+TArray<FPackageStorePackage*> FPackageStoreOptimizer::SortPackagesInLoadOrder(const TMap<FPackageId, FPackageStorePackage*>& PackagesMap) const
 {
 	TRACE_CPUPROFILER_EVENT_SCOPE(SortPackagesInLoadOrder);
+	TArray<FPackageStorePackage*> Packages;
+	PackagesMap.GenerateValueArray(Packages);
 	Algo::Sort(Packages, [](const FPackageStorePackage* A, const FPackageStorePackage* B)
 	{
 		return A->Id < B->Id;
@@ -376,14 +598,18 @@ void FPackageStoreOptimizer::SortPackagesInLoadOrder(TArray<FPackageStorePackage
 	TMap<FPackageStorePackage*, TArray<FPackageStorePackage*>> SortedEdges;
 	for (FPackageStorePackage* Package : Packages)
 	{
-		for (FPackageStorePackage* ImportedPackage : Package->ImportedWaitingPackages)
+		for (FPackageId ImportedPackageId : Package->ImportedPackageIds)
 		{
-			TArray<FPackageStorePackage*>& SourceArray = SortedEdges.FindOrAdd(ImportedPackage);
-			SourceArray.Add(Package);
+			FPackageStorePackage* FindImportedPackage = PackagesMap.FindRef(ImportedPackageId);
+			if (FindImportedPackage)
+			{
+				TArray<FPackageStorePackage*>& SourceArray = SortedEdges.FindOrAdd(FindImportedPackage);
+				SourceArray.Add(Package);
+			}
 		}
 		for (FPackageId ImportedRedirectedPackageId : Package->ImportedRedirectedPackageIds)
 		{
-			FPackageStorePackage* FindImportedPackage = WaitingPackagesMap.FindRef(ImportedRedirectedPackageId);
+			FPackageStorePackage* FindImportedPackage = PackagesMap.FindRef(ImportedRedirectedPackageId);
 			if (FindImportedPackage)
 			{
 				TArray<FPackageStorePackage*>& SourceArray = SortedEdges.FindOrAdd(FindImportedPackage);
@@ -435,174 +661,361 @@ void FPackageStoreOptimizer::SortPackagesInLoadOrder(TArray<FPackageStorePackage
 	check(Result.Num() == Packages.Num());
 	Algo::Reverse(Result);
 	Swap(Result, Packages);
+	return Packages;
 }
 
-FPackageStoreOptimizer::FExportGraphEdges FPackageStoreOptimizer::ProcessPreloadDependencies(const TArray<FPackageStorePackage*>& Packages) const
+TArray<FPackageStorePackage::FExportBundleGraphNode*> FPackageStoreOptimizer::SortExportBundleGraphNodesInLoadOrder(const TArray<FPackageStorePackage*>& Packages, FExportBundleGraphEdges& Edges) const
 {
-	TRACE_CPUPROFILER_EVENT_SCOPE(ProcessPreloadDependencies);
-	FExportGraphEdges Edges;
-	auto AddExternalDependency = [this, &Edges, &Packages](FPackageStorePackage* Package, int32 FromImportIndex, FExportBundleEntry::EExportCommandType FromExportBundleCommandType, FPackageStorePackage::FExportGraphNode* ToNode)
+	TRACE_CPUPROFILER_EVENT_SCOPE(SortExportBundleGraphNodesInLoadOrder);
+	int32 NodeCount = 0;
+	for (FPackageStorePackage* Package : Packages)
 	{
-		const FPackageStorePackage::FImport& FromImport = Package->Imports[FromImportIndex];
-		if (FromImport.GlobalImportIndex.IsScriptImport())
-		{
-			return;
-		}
+		NodeCount += Package->ExportGraphNodes.Num();
+	}
+	for (auto& KV : Edges)
+	{
+		FPackageStorePackage::FExportBundleGraphNode* ToNode = KV.Value;
+		++ToNode->IncomingEdgeCount;
+	}
 
-		FPackageStorePackage::FExternalDependency& ExternalDependency = ToNode->ExternalDependencies.AddDefaulted_GetRef();
-		ExternalDependency.PackageId = FromImport.PackageId;
-		ExternalDependency.GlobalImportIndex = FromImport.GlobalImportIndex;
-		ExternalDependency.DebugImportFullName = FName(FromImport.FullName);
-		ExternalDependency.ExportBundleCommandType = FromExportBundleCommandType;
+	auto NodeSorter = [](const FPackageStorePackage::FExportBundleGraphNode& A, const FPackageStorePackage::FExportBundleGraphNode& B)
+	{
+		return A.Index < B.Index;
+	};
 
-		const TArray<FResolvedPublicExport>* FindResolvedPublicExports = ResolvedPublicExportsMap.Find(FromImport.PackageId);
-		if (FindResolvedPublicExports)
+	for (FPackageStorePackage* Package : Packages)
+	{
+		for (FPackageStorePackage::FExportBundleGraphNode& Node : Package->ExportBundleGraphNodes)
 		{
-			bool bFoundExport = false;
-			for (const FResolvedPublicExport& PublicExport : *FindResolvedPublicExports)
+			if (Node.IncomingEdgeCount == 0)
 			{
-				if (PublicExport.GlobalImportIndex == FromImport.GlobalImportIndex)
+				Package->NodesWithNoIncomingEdges.HeapPush(&Node, NodeSorter);
+			}
+		}
+	}
+	TArray<FPackageStorePackage::FExportBundleGraphNode*> LoadOrder;
+	LoadOrder.Reserve(NodeCount);
+	while (LoadOrder.Num() < NodeCount)
+	{
+		bool bMadeProgress = false;
+		for (FPackageStorePackage* Package : Packages)
+		{
+			while (Package->NodesWithNoIncomingEdges.Num())
+			{
+				FPackageStorePackage::FExportBundleGraphNode* RemovedNode;
+				Package->NodesWithNoIncomingEdges.HeapPop(RemovedNode, NodeSorter, false);
+				LoadOrder.Add(RemovedNode);
+				bMadeProgress = true;
+				for (auto EdgeIt = Edges.CreateKeyIterator(RemovedNode); EdgeIt; ++EdgeIt)
 				{
-					if (FromExportBundleCommandType == FExportBundleEntry::ExportCommandType_Create)
-					{
-						check(PublicExport.CreateExportBundleIndex >= 0);
-						ExternalDependency.ResolvedExportBundleIndex = PublicExport.CreateExportBundleIndex;
-					}
-					else
+					FPackageStorePackage::FExportBundleGraphNode* ToNode = EdgeIt.Value();
+					check(ToNode->IncomingEdgeCount > 0);
+					if (--ToNode->IncomingEdgeCount == 0)
 					{
-						check(PublicExport.SerializeExportBundleIndex >= 0);
-						ExternalDependency.ResolvedExportBundleIndex = PublicExport.SerializeExportBundleIndex;
+						ToNode->Package->NodesWithNoIncomingEdges.HeapPush(ToNode, NodeSorter);
 					}
-					bFoundExport = true;
-					break;
+					EdgeIt.RemoveCurrent();
 				}
 			}
-			if (!bFoundExport)
+		}
+		check(bMadeProgress);
+	}
+	return LoadOrder;
+}
+
+void FPackageStoreOptimizer::OptimizeExportBundles(const TMap<FPackageId, FPackageStorePackage*>& PackagesMap)
+{
+	TArray<FPackageStorePackage*> Packages = SortPackagesInLoadOrder(PackagesMap);
+	for (FPackageStorePackage* Package : Packages)
+	{
+		Package->ExportBundleGraphNodes.Reserve(Package->GraphData.ExportBundles.Num());
+		for (FPackageStorePackage::FExportBundle& ExportBundle : Package->GraphData.ExportBundles)
+		{
+			FPackageStorePackage::FExportBundleGraphNode& Node = Package->ExportBundleGraphNodes.AddDefaulted_GetRef();
+			Node.Package = Package;
+			Node.Index = Package->ExportBundleGraphNodes.Num() - 1;
+			Node.ExportGraphNodes.Reserve(ExportBundle.Entries.Num());
+			for (FExportBundleEntry& ExportBundleEntry : ExportBundle.Entries)
 			{
-				ExternalDependency.bIsConfirmedMissing = true;
+				FPackageStorePackage::FExport& Export = Package->Exports[ExportBundleEntry.LocalExportIndex];
+				Node.ExportGraphNodes.Add(Export.Nodes[ExportBundleEntry.CommandType]);
 			}
 		}
-		else
+		Package->GraphData.ExportBundles.Empty();
+	}
+
+	FExportBundleGraphEdges Edges;
+	TSet<FPackageStorePackage::FExportBundleGraphNode*> Dependencies;
+	for (FPackageStorePackage* Package : Packages)
+	{
+		for (FPackageStorePackage::FExportBundleGraphNode& ExportBundleGraphNode : Package->ExportBundleGraphNodes)
 		{
-			FPackageStorePackage* FindWaitingFromPackage = WaitingPackagesMap.FindRef(FromImport.PackageId);
-			if (FindWaitingFromPackage)
+			for (FPackageStorePackage::FExportGraphNode* ExportGraphNode : ExportBundleGraphNode.ExportGraphNodes)
 			{
-				bool bFoundExport = false;
-				for (int32 ExportIndex = 0; ExportIndex < FindWaitingFromPackage->Exports.Num(); ++ExportIndex)
+				check(ExportGraphNode->ExportBundleIndex >= 0);
+				for (FPackageStorePackage::FExportGraphNode* InternalDependency : ExportGraphNode->InternalDependencies)
+				{
+					check(InternalDependency->ExportBundleIndex >= 0);
+					FPackageStorePackage::FExportBundleGraphNode* FromNode = &Package->ExportBundleGraphNodes[InternalDependency->ExportBundleIndex];
+					Dependencies.Add(FromNode);
+				}
+				for (FPackageStorePackage::FExternalDependency& ExternalDependency : ExportGraphNode->ExternalDependencies)
 				{
-					FPackageStorePackage::FExport& FromExport = FindWaitingFromPackage->Exports[ExportIndex];
-					if (FromExport.GlobalImportIndex == FromImport.GlobalImportIndex)
+					const FPackageStorePackage::FImport& FromImport = Package->Imports[ExternalDependency.ImportIndex];
+					FPackageStorePackage* FindFromPackage = PackagesMap.FindRef(FromImport.PackageId);
+					if (FindFromPackage)
 					{
-						FPackageStorePackage::FExportGraphNode* FromNode = FromExportBundleCommandType == FExportBundleEntry::ExportCommandType_Create ? FromExport.CreateNode : FromExport.SerializeNode;
-						Edges.Add(FromNode, ToNode);
-						ExternalDependency.ExportGraphNode = FromNode;
-
-						for (const FPackageId& FromRedirectedPackageId : FindWaitingFromPackage->RedirectedToPackageIds)
+						bool bFoundExport = false;
+						for (int32 ExportIndex = 0; ExportIndex < FindFromPackage->Exports.Num(); ++ExportIndex)
 						{
-							FPackageStorePackage* FindRedirectedFromPackage = WaitingPackagesMap.FindRef(FromRedirectedPackageId);
-							if (FindRedirectedFromPackage)
+							FPackageStorePackage::FExport& FromExport = FindFromPackage->Exports[ExportIndex];
+							if (FromExport.GlobalImportIndex == FromImport.GlobalImportIndex)
 							{
-								FPackageStorePackage::FExport& FromRedirectedExport = FindRedirectedFromPackage->Exports[ExportIndex];
-								FPackageStorePackage::FExportGraphNode* FromRedirectedNode = FromExportBundleCommandType == FExportBundleEntry::ExportCommandType_Create ? FromRedirectedExport.CreateNode : FromRedirectedExport.SerializeNode;
-								Edges.Add(FromRedirectedNode, ToNode);
-
-								FPackageStorePackage::FExternalDependency& RedirectedExternalDependency = ToNode->ExternalDependencies.AddDefaulted_GetRef();
-								RedirectedExternalDependency = ExternalDependency;
-								RedirectedExternalDependency.PackageId = FromRedirectedPackageId;
-								RedirectedExternalDependency.ExportGraphNode = FromRedirectedNode;
+								FPackageStorePackage::FExportGraphNode* FromExportGraphNode = FromExport.Nodes[ExternalDependency.ExportBundleCommandType];
+								check(FromExportGraphNode->ExportBundleIndex >= 0);
+								FPackageStorePackage::FExportBundleGraphNode* FromNode = &FindFromPackage->ExportBundleGraphNodes[FromExportGraphNode->ExportBundleIndex];
+								Dependencies.Add(FromNode);
+
+								for (const FPackageId& FromRedirectedPackageId : FindFromPackage->RedirectedToPackageIds)
+								{
+									FPackageStorePackage* FindRedirectedFromPackage = PackagesMap.FindRef(FromRedirectedPackageId);
+									if (FindRedirectedFromPackage)
+									{
+										FPackageStorePackage::FExport& FromRedirectedExport = FindRedirectedFromPackage->Exports[ExportIndex];
+										FPackageStorePackage::FExportGraphNode* FromRedirectedExportGraphNode = FromRedirectedExport.Nodes[ExternalDependency.ExportBundleCommandType];
+										check(FromRedirectedExportGraphNode->ExportBundleIndex >= 0);
+										FPackageStorePackage::FExportBundleGraphNode* FromRedirectedNode = &FindFromPackage->ExportBundleGraphNodes[FromRedirectedExportGraphNode->ExportBundleIndex];
+										Dependencies.Add(FromRedirectedNode);
+									}
+								}
+
+								bFoundExport = true;
+								break;
 							}
 						}
-
-						bFoundExport = true;
-						break;
+						if (!bFoundExport)
+						{
+							ExternalDependency.bIsConfirmedMissing = true;
+						}
 					}
 				}
-				if (!bFoundExport)
+			}
+			for (FPackageStorePackage::FExportBundleGraphNode* FromNode : Dependencies)
+			{
+				Edges.Add(FromNode, &ExportBundleGraphNode);
+			}
+			Dependencies.Reset();
+		}
+	}
+
+	TArray<FPackageStorePackage::FExportBundleGraphNode*> LoadOrder = SortExportBundleGraphNodesInLoadOrder(Packages, Edges);
+
+	FPackageStorePackage* PreviousPackage = nullptr;
+	for (FPackageStorePackage::FExportBundleGraphNode* Node : LoadOrder)
+	{
+		check(Node->Package);
+		FPackageStorePackage* Package = Node->Package;
+		if (PreviousPackage != Package || !Package->CurrentBundle)
+		{
+			if (Package->GraphData.ExportBundles.IsEmpty())
+			{
+				Package->LoadOrder = NextLoadOrder++;
+			}
+			Package->CurrentBundle = &Package->GraphData.ExportBundles.AddDefaulted_GetRef();
+		}
+		for (FPackageStorePackage::FExportGraphNode* ExportGraphNode : Node->ExportGraphNodes)
+		{
+			Package->CurrentBundle->Entries.Add(ExportGraphNode->BundleEntry);
+			ExportGraphNode->ExportBundleIndex = Package->GraphData.ExportBundles.Num() - 1;
+		}
+		PreviousPackage = Package;
+	}
+}
+
+void FPackageStoreOptimizer::FindImports(const TMap<FPackageId, FPackageStorePackage*>& PackagesMap) const
+{
+	using FPackageExportIndexPair = TPair<FPackageStorePackage*, int32>;
+	using FGlobalImportIndexToExport = TMap<FPackageObjectIndex, FPackageExportIndexPair>;
+
+	FGlobalImportIndexToExport GlobalImportIndexToExport;
+
+	for (const auto& PackageIdPackage : PackagesMap)
+	{
+		const FPackageStorePackage& Package = *PackageIdPackage.Value;
+		for (int32 ExportIndex = 0; ExportIndex < Package.Exports.Num(); ++ExportIndex)
+		{
+			const FPackageStorePackage::FExport& Export = Package.Exports[ExportIndex];
+			if (Export.bIsPublic)
+			{
+				GlobalImportIndexToExport.Add(Export.GlobalImportIndex, FPackageExportIndexPair(PackageIdPackage.Value, ExportIndex));
+			}
+		}
+	}
+
+	for (const auto& PackageIdPackage : PackagesMap)
+	{
+		FPackageStorePackage& Package = *PackageIdPackage.Value;
+		for (FPackageStorePackage::FImport& Import : Package.Imports)
+		{
+			if (Import.GlobalImportIndex.IsPackageImport())
+			{
+				if (const FPackageExportIndexPair* Entry = GlobalImportIndexToExport.Find(Import.GlobalImportIndex))
 				{
-					ExternalDependency.bIsConfirmedMissing = true;
+					const FPackageStorePackage& ImportedPackage = *Entry->Key;
+					const int32 ExportIndex = Entry->Value;
+					Import.PackageId = ImportedPackage.GetId();
+					Import.FullName = ImportedPackage.Exports[ExportIndex].FullName;
+				}
+				else
+				{
+					UE_LOG(LogPackageStoreOptimizer, Warning,
+						TEXT("Failed to find import '0x%llX' in package '%s'"),
+						Import.GlobalImportIndex.Value(), *Package.Name.ToString());
 				}
-
 			}
 		}
+	}
+}
+
+void FPackageStoreOptimizer::ProcessPreloadDependencies(const FCookedHeaderData& CookedHeaderData, FPackageStorePackage* Package) const
+{
+	TRACE_CPUPROFILER_EVENT_SCOPE(ProcessPreloadDependencies);
+
+	auto AddInternalDependency = [](FPackageStorePackage* Package, int32 FromExportIndex, FExportBundleEntry::EExportCommandType FromExportBundleCommandType, FPackageStorePackage::FExportGraphNode* ToNode)
+	{
+		FPackageStorePackage::FExport& FromExport = Package->Exports[FromExportIndex];
+		FPackageStorePackage::FExportGraphNode* FromNode = FromExport.Nodes[FromExportBundleCommandType];
+		ToNode->InternalDependencies.Add(FromNode);
 	};
 
-	for (FPackageStorePackage* Package : Packages)
+	auto AddExternalDependency = [](FPackageStorePackage* Package, int32 FromImportIndex, FExportBundleEntry::EExportCommandType FromExportBundleCommandType, FPackageStorePackage::FExportGraphNode* ToNode)
 	{
-		for (int32 ExportIndex = 0; ExportIndex < Package->Exports.Num(); ++ExportIndex)
+		const FPackageStorePackage::FImport& FromImport = Package->Imports[FromImportIndex];
+		if (FromImport.GlobalImportIndex.IsScriptImport())
 		{
-			FPackageStorePackage::FExport& Export = Package->Exports[ExportIndex];
-			const FObjectExport& ObjectExport = Package->ObjectExports[ExportIndex];
+			return;
+		}
+
+		FPackageStorePackage::FExternalDependency& ExternalDependency = ToNode->ExternalDependencies.AddDefaulted_GetRef();
+		ExternalDependency.ImportIndex = FromImportIndex;
+		ExternalDependency.ExportBundleCommandType = FromExportBundleCommandType;
+	};
+
+	for (int32 ExportIndex = 0; ExportIndex < Package->Exports.Num(); ++ExportIndex)
+	{
+		FPackageStorePackage::FExport& Export = Package->Exports[ExportIndex];
+		const FObjectExport& ObjectExport = CookedHeaderData.ObjectExports[ExportIndex];
 
-			Edges.Add(Export.CreateNode, Export.SerializeNode);
+		AddInternalDependency(Package, ExportIndex, FExportBundleEntry::ExportCommandType_Create, Export.Nodes[FExportBundleEntry::ExportCommandType_Serialize]);
 
-			if (ObjectExport.FirstExportDependency >= 0)
+		if (ObjectExport.FirstExportDependency >= 0)
+		{
+			int32 RunningIndex = ObjectExport.FirstExportDependency;
+			for (int32 Index = ObjectExport.SerializationBeforeSerializationDependencies; Index > 0; Index--)
 			{
-				int32 RunningIndex = ObjectExport.FirstExportDependency;
-				for (int32 Index = ObjectExport.SerializationBeforeSerializationDependencies; Index > 0; Index--)
+				FPackageIndex Dep = CookedHeaderData.PreloadDependencies[RunningIndex++];
+				if (Dep.IsExport())
 				{
-					FPackageIndex Dep = Package->PreloadDependencies[RunningIndex++];
-					if (Dep.IsExport())
-					{
-						Edges.Add(Package->Exports[Dep.ToExport()].SerializeNode, Export.SerializeNode);
-					}
-					else
-					{
-						AddExternalDependency(Package, Dep.ToImport(), FExportBundleEntry::ExportCommandType_Serialize, Export.SerializeNode);
-					}
+					AddInternalDependency(Package, Dep.ToExport(), FExportBundleEntry::ExportCommandType_Serialize, Export.Nodes[FExportBundleEntry::ExportCommandType_Serialize]);
 				}
+				else
+				{
+					AddExternalDependency(Package, Dep.ToImport(), FExportBundleEntry::ExportCommandType_Serialize, Export.Nodes[FExportBundleEntry::ExportCommandType_Serialize]);
+				}
+			}
 
-				for (int32 Index = ObjectExport.CreateBeforeSerializationDependencies; Index > 0; Index--)
+			for (int32 Index = ObjectExport.CreateBeforeSerializationDependencies; Index > 0; Index--)
+			{
+				FPackageIndex Dep = CookedHeaderData.PreloadDependencies[RunningIndex++];
+				if (Dep.IsExport())
 				{
-					FPackageIndex Dep = Package->PreloadDependencies[RunningIndex++];
-					if (Dep.IsExport())
-					{
-						Edges.Add(Package->Exports[Dep.ToExport()].CreateNode, Export.SerializeNode);
-					}
-					else
-					{
-						AddExternalDependency(Package, Dep.ToImport(), FExportBundleEntry::ExportCommandType_Create, Export.SerializeNode);
-					}
+					AddInternalDependency(Package, Dep.ToExport(), FExportBundleEntry::ExportCommandType_Create, Export.Nodes[FExportBundleEntry::ExportCommandType_Serialize]);
+				}
+				else
+				{
+					AddExternalDependency(Package, Dep.ToImport(), FExportBundleEntry::ExportCommandType_Create, Export.Nodes[FExportBundleEntry::ExportCommandType_Serialize]);
 				}
+			}
 
-				for (int32 Index = ObjectExport.SerializationBeforeCreateDependencies; Index > 0; Index--)
+			for (int32 Index = ObjectExport.SerializationBeforeCreateDependencies; Index > 0; Index--)
+			{
+				FPackageIndex Dep = CookedHeaderData.PreloadDependencies[RunningIndex++];
+				if (Dep.IsExport())
 				{
-					FPackageIndex Dep = Package->PreloadDependencies[RunningIndex++];
-					if (Dep.IsExport())
-					{
-						Edges.Add(Package->Exports[Dep.ToExport()].SerializeNode, Export.CreateNode);
-					}
-					else
-					{
-						AddExternalDependency(Package, Dep.ToImport(), FExportBundleEntry::ExportCommandType_Serialize, Export.CreateNode);
-					}
+					AddInternalDependency(Package, Dep.ToExport(), FExportBundleEntry::ExportCommandType_Serialize, Export.Nodes[FExportBundleEntry::ExportCommandType_Create]);
+				}
+				else
+				{
+					AddExternalDependency(Package, Dep.ToImport(), FExportBundleEntry::ExportCommandType_Serialize, Export.Nodes[FExportBundleEntry::ExportCommandType_Create]);
 				}
+			}
 
-				for (int32 Index = ObjectExport.CreateBeforeCreateDependencies; Index > 0; Index--)
+			for (int32 Index = ObjectExport.CreateBeforeCreateDependencies; Index > 0; Index--)
+			{
+				FPackageIndex Dep = CookedHeaderData.PreloadDependencies[RunningIndex++];
+				if (Dep.IsExport())
 				{
-					FPackageIndex Dep = Package->PreloadDependencies[RunningIndex++];
-					if (Dep.IsExport())
-					{
-						Edges.Add(Package->Exports[Dep.ToExport()].CreateNode, Export.CreateNode);
-					}
-					else
-					{
-						AddExternalDependency(Package, Dep.ToImport(), FExportBundleEntry::ExportCommandType_Create, Export.CreateNode);
-					}
+					AddInternalDependency(Package, Dep.ToExport(), FExportBundleEntry::ExportCommandType_Create, Export.Nodes[FExportBundleEntry::ExportCommandType_Create]);
+				}
+				else
+				{
+					AddExternalDependency(Package, Dep.ToImport(), FExportBundleEntry::ExportCommandType_Create, Export.Nodes[FExportBundleEntry::ExportCommandType_Create]);
 				}
 			}
 		}
 	}
-	return Edges;
 }
 
-TArray<FPackageStorePackage::FExportGraphNode*> FPackageStoreOptimizer::SortExportGraphNodesInLoadOrder(const TArray<FPackageStorePackage*>& Packages, FExportGraphEdges& Edges) const
+void FPackageStoreOptimizer::ProcessPreloadDependencies(const FPackageStoreHeaderData& PackageStoreHeaderData, FPackageStorePackage* Package) const
 {
-	TRACE_CPUPROFILER_EVENT_SCOPE(SortExportGraphNodesInLoadOrder);
-	int32 NodeCount = 0;
-	for (FPackageStorePackage* Package : Packages)
+	TRACE_CPUPROFILER_EVENT_SCOPE(ProcessPreloadDependencies);
+
+	for (const FPackageStorePackage::FInternalArc& InternalArc : PackageStoreHeaderData.InternalArcs)
 	{
-		NodeCount += Package->ExportGraphNodes.Num();
+		const FExportBundleHeader& FromExportBundle = PackageStoreHeaderData.ExportBundleHeaders[InternalArc.FromExportBundleIndex];
+		const FExportBundleHeader& ToExportBundle = PackageStoreHeaderData.ExportBundleHeaders[InternalArc.ToExportBundleIndex];
+
+		uint32 FromBundleEntryIndex = FromExportBundle.FirstEntryIndex;
+		const uint32 LastFromBundleEntryIndex = FromBundleEntryIndex + FromExportBundle.EntryCount;
+		while (FromBundleEntryIndex < LastFromBundleEntryIndex)
+		{
+			const FExportBundleEntry& FromBundleEntry = PackageStoreHeaderData.ExportBundleEntries[FromBundleEntryIndex++];
+			FPackageStorePackage::FExport& FromExport = Package->Exports[FromBundleEntry.LocalExportIndex];
+			FPackageStorePackage::FExportGraphNode* FromNode = FromExport.Nodes[FromBundleEntry.CommandType];
+
+			uint32 ToBundleEntryIndex = ToExportBundle.FirstEntryIndex;
+			const uint32 LastToBundleEntryIndex = ToBundleEntryIndex + ToExportBundle.EntryCount;
+			while (ToBundleEntryIndex < LastToBundleEntryIndex)
+			{
+				const FExportBundleEntry& ToBundleEntry = PackageStoreHeaderData.ExportBundleEntries[ToBundleEntryIndex++];
+				FPackageStorePackage::FExport& ToExport = Package->Exports[ToBundleEntry.LocalExportIndex];
+				FPackageStorePackage::FExportGraphNode* ToNode = ToExport.Nodes[ToBundleEntry.CommandType];
+				ToNode->InternalDependencies.Add(FromNode);
+			}
+		}
+	}
+
+	for (const FPackageStorePackage::FExternalArc& ExternalArc : PackageStoreHeaderData.ExternalArcs)
+	{
+		const FExportBundleHeader& ToExportBundle = PackageStoreHeaderData.ExportBundleHeaders[ExternalArc.ToExportBundleIndex];
+		uint32 ToBundleEntryIndex = ToExportBundle.FirstEntryIndex;
+		const uint32 LastToBundleEntryIndex = ToBundleEntryIndex + ToExportBundle.EntryCount;
+		while (ToBundleEntryIndex < LastToBundleEntryIndex)
+		{
+			const FExportBundleEntry& ToBundleEntry = PackageStoreHeaderData.ExportBundleEntries[ToBundleEntryIndex++];
+			FPackageStorePackage::FExport& ToExport = Package->Exports[ToBundleEntry.LocalExportIndex];
+			FPackageStorePackage::FExportGraphNode* ToNode = ToExport.Nodes[ToBundleEntry.CommandType];
+			FPackageStorePackage::FExternalDependency& ExternalDependency = ToNode->ExternalDependencies.AddDefaulted_GetRef();
+			ExternalDependency.ImportIndex = ExternalArc.FromImportIndex;
+			ExternalDependency.ExportBundleCommandType = ExternalArc.FromCommandType;
+		}
 	}
+}
+
+TArray<FPackageStorePackage::FExportGraphNode*> FPackageStoreOptimizer::SortExportGraphNodesInLoadOrder(FPackageStorePackage* Package, FExportGraphEdges& Edges) const
+{
+	TRACE_CPUPROFILER_EVENT_SCOPE(SortExportGraphNodesInLoadOrder);
+	int32 NodeCount = Package->ExportGraphNodes.Num();
 	for (auto& KV : Edges)
 	{
 		FPackageStorePackage::FExportGraphNode* ToNode = KV.Value;
@@ -611,94 +1024,73 @@ TArray<FPackageStorePackage::FExportGraphNode*> FPackageStoreOptimizer::SortExpo
 
 	auto NodeSorter = [](const FPackageStorePackage::FExportGraphNode& A, const FPackageStorePackage::FExportGraphNode& B)
 	{
-		if (A.BundleEntry.LocalExportIndex == B.BundleEntry.LocalExportIndex)
+		if (A.bIsPublic != B.bIsPublic)
+		{
+			return A.bIsPublic;
+		}
+		if (A.BundleEntry.CommandType != B.BundleEntry.CommandType)
 		{
 			return A.BundleEntry.CommandType < B.BundleEntry.CommandType;
 		}
 		return A.BundleEntry.LocalExportIndex < B.BundleEntry.LocalExportIndex;
 	};
 
-	for (FPackageStorePackage* Package : Packages)
+	TArray<FPackageStorePackage::FExportGraphNode*> NodesWithNoIncomingEdges;
+	NodesWithNoIncomingEdges.Reserve(NodeCount);
+	for (FPackageStorePackage::FExportGraphNode& Node : Package->ExportGraphNodes)
 	{
-		for (FPackageStorePackage::FExportGraphNode& Node : Package->ExportGraphNodes)
+		if (Node.IncomingEdgeCount == 0)
 		{
-			if (Node.IncomingEdgeCount == 0)
-			{
-				Package->NodesWithNoIncomingEdges.HeapPush(&Node, NodeSorter);
-			}
+			NodesWithNoIncomingEdges.HeapPush(&Node, NodeSorter);
 		}
 	}
+
 	TArray<FPackageStorePackage::FExportGraphNode*> LoadOrder;
 	LoadOrder.Reserve(NodeCount);
-	while (LoadOrder.Num() < NodeCount)
+	while (NodesWithNoIncomingEdges.Num())
 	{
-		bool bMadeProgress = false;
-		for (FPackageStorePackage* Package : Packages)
+		FPackageStorePackage::FExportGraphNode* RemovedNode;
+		NodesWithNoIncomingEdges.HeapPop(RemovedNode, NodeSorter, false);
+		LoadOrder.Add(RemovedNode);
+		for (auto EdgeIt = Edges.CreateKeyIterator(RemovedNode); EdgeIt; ++EdgeIt)
 		{
-			while (Package->NodesWithNoIncomingEdges.Num())
+			FPackageStorePackage::FExportGraphNode* ToNode = EdgeIt.Value();
+			check(ToNode->IncomingEdgeCount > 0);
+			if (--ToNode->IncomingEdgeCount == 0)
 			{
-				FPackageStorePackage::FExportGraphNode* RemovedNode;
-				Package->NodesWithNoIncomingEdges.HeapPop(RemovedNode, NodeSorter, false);
-				LoadOrder.Add(RemovedNode);
-				bMadeProgress = true;
-				for (auto EdgeIt = Edges.CreateKeyIterator(RemovedNode); EdgeIt; ++EdgeIt)
-				{
-					FPackageStorePackage::FExportGraphNode* ToNode = EdgeIt.Value();
-					check(ToNode->IncomingEdgeCount > 0);
-					if (--ToNode->IncomingEdgeCount == 0)
-					{
-						ToNode->Package->NodesWithNoIncomingEdges.HeapPush(ToNode, NodeSorter);
-					}
-					EdgeIt.RemoveCurrent();
-				}
+				NodesWithNoIncomingEdges.HeapPush(ToNode, NodeSorter);
 			}
+			EdgeIt.RemoveCurrent();
 		}
-		check(bMadeProgress);
 	}
+	check(LoadOrder.Num() == NodeCount);
 	return LoadOrder;
 }
 
-void FPackageStoreOptimizer::CreateExportBundles(TArray<FPackageStorePackage*>& Packages)
+void FPackageStoreOptimizer::CreateExportBundles(FPackageStorePackage* Package) const
 {
 	TRACE_CPUPROFILER_EVENT_SCOPE(CreateExportBundles);
-	SortPackagesInLoadOrder(Packages);
-	FExportGraphEdges Edges = ProcessPreloadDependencies(Packages);
-	TArray<FPackageStorePackage::FExportGraphNode*> LoadOrder = SortExportGraphNodesInLoadOrder(Packages, Edges);
-
-	FPackageStorePackage* LastPackage = nullptr;
-	for (FPackageStorePackage::FExportGraphNode* Node : LoadOrder)
+	FExportGraphEdges Edges;
+	for (FPackageStorePackage::FExportGraphNode& ExportGraphNode : Package->ExportGraphNodes)
 	{
-		check(Node && Node->Package);
-		if (!Node || !Node->Package)
+		for (FPackageStorePackage::FExportGraphNode* InternalDependency : ExportGraphNode.InternalDependencies)
 		{
-			// Needed to please static analysis
-			continue;
+			Edges.Add(InternalDependency, &ExportGraphNode);
 		}
-		FPackageStorePackage::FExportBundle* Bundle;
-		if (Node->Package != LastPackage)
-		{
-			Node->ExportBundleIndex = Node->Package->ExportBundles.Num();
-			Bundle = &Node->Package->ExportBundles.AddDefaulted_GetRef();
-			Bundle->LoadOrder = NextLoadOrder++;
-			++TotalExportBundleCount;
-			LastPackage = Node->Package;
-		}
-		else
+	}
+	TArray<FPackageStorePackage::FExportGraphNode*> LoadOrder = SortExportGraphNodesInLoadOrder(Package, Edges);
+	FPackageStorePackage::FExportBundle* CurrentBundle = nullptr;
+	for (FPackageStorePackage::FExportGraphNode* Node : LoadOrder)
+	{
+		if (!CurrentBundle)
 		{
-			Node->ExportBundleIndex = Node->Package->ExportBundles.Num() - 1;
-			Bundle = &Node->Package->ExportBundles[Node->ExportBundleIndex];
+			Package->CurrentBundle = &Package->GraphData.ExportBundles.AddDefaulted_GetRef();
 		}
-		Bundle->Entries.Add(Node->BundleEntry);
-		++TotalExportBundleEntryCount;
-
-		for (FPackageStorePackage::FExternalDependency& ExternalDependency : Node->ExternalDependencies)
+		Package->CurrentBundle->Entries.Add(Node->BundleEntry);
+		Node->ExportBundleIndex = Package->GraphData.ExportBundles.Num() - 1;
+		if (Node->bIsPublic)
 		{
-			if (ExternalDependency.ExportGraphNode)
-			{
-				check(ExternalDependency.ExportGraphNode->ExportBundleIndex >= 0);
-				ExternalDependency.ResolvedExportBundleIndex = ExternalDependency.ExportGraphNode->ExportBundleIndex;
-				ExternalDependency.ExportGraphNode = nullptr;
-			}
+			CurrentBundle = nullptr;
 		}
 	}
 }
@@ -908,21 +1300,22 @@ void FPackageStoreOptimizer::RedirectPackageUnverified(FPackageStorePackage& Tar
 	}
 }
 
-void FPackageStoreOptimizer::ProcessRedirects(const TArray<FPackageStorePackage*> Packages) const
+void FPackageStoreOptimizer::ProcessRedirects(const TMap<FPackageId, FPackageStorePackage*>& PackagesMap) const
 {
 	TRACE_CPUPROFILER_EVENT_SCOPE(ProcessRedirects);
 	TMap<FPackageObjectIndex, FPackageObjectIndex> RedirectedToSourceImportIndexMap;
 	TMap<FPackageId, FPackageId> SourceToRedirectedPackageMap;
 
-	for (FPackageStorePackage* Package : Packages)
+	for (const auto& KV : PackagesMap)
 	{
+		FPackageStorePackage* Package = KV.Value;
 		if (Package->SourceName.IsNone() || Package->Name == Package->SourceName)
 		{
 			continue;
 		}
 		
 		FPackageId SourcePackageId = FPackageId::FromName(Package->SourceName);
-		FPackageStorePackage* SourcePackage = WaitingPackagesMap.FindRef(SourcePackageId);
+		FPackageStorePackage* SourcePackage = PackagesMap.FindRef(SourcePackageId);
 		
 		if (SourcePackage)
 		{
@@ -980,11 +1373,16 @@ void FPackageStoreOptimizer::ProcessRedirects(const TArray<FPackageStorePackage*
 		}
 	}
 
-	for (FPackageStorePackage* Package : Packages)
+	for (const auto& KV : PackagesMap)
 	{
-		for (FPackageStorePackage* ImportedPackage : Package->ImportedWaitingPackages)
+		FPackageStorePackage* Package = KV.Value;
+		for (FPackageId ImportedPackageId : Package->ImportedPackageIds)
 		{
-			Package->ImportedRedirectedPackageIds.Append(ImportedPackage->RedirectedToPackageIds);
+			FPackageStorePackage* FindImportedPackage = PackagesMap.FindRef(ImportedPackageId);
+			if (FindImportedPackage)
+			{
+				Package->ImportedRedirectedPackageIds.Append(FindImportedPackage->RedirectedToPackageIds);
+			}
 		}
 
 		for (FPackageStorePackage::FImport& Import : Package->Imports)
@@ -1007,19 +1405,69 @@ void FPackageStoreOptimizer::ProcessRedirects(const TArray<FPackageStorePackage*
 	}
 }
 
-void FPackageStoreOptimizer::ResolvePackages(TArray<FPackageStorePackage*>& Packages, bool bAllowMissingImports)
+void FPackageStoreOptimizer::SerializeGraphData(const TArray<FPackageId>& ImportedPackageIds, FPackageStorePackage::FGraphData& GraphData, FBufferWriter& GraphArchive) const
 {
-	ProcessRedirects(Packages);
-	CreateExportBundles(Packages);
-	for (FPackageStorePackage* Package : Packages)
+	uint32 ExportBundleEntryIndex = 0;
+	for (const FPackageStorePackage::FExportBundle& ExportBundle : GraphData.ExportBundles)
+	{
+		const uint32 EntryCount = ExportBundle.Entries.Num();
+		FExportBundleHeader ExportBundleHeader{ ExportBundle.SerialOffset, ExportBundleEntryIndex, EntryCount };
+		GraphArchive << ExportBundleHeader;
+		ExportBundleEntryIndex += EntryCount;
+	}
+	Algo::Sort(GraphData.InternalArcs, [](const FPackageStorePackage::FInternalArc& A, const FPackageStorePackage::FInternalArc& B)
+	{
+		if (A.ToExportBundleIndex == B.ToExportBundleIndex)
+		{
+			return A.FromExportBundleIndex < B.FromExportBundleIndex;
+		}
+		return A.ToExportBundleIndex < B.ToExportBundleIndex;
+	});
+	int32 InternalArcsCount = GraphData.InternalArcs.Num();
+	GraphArchive << InternalArcsCount;
+	for (FPackageStorePackage::FInternalArc& InternalArc : GraphData.InternalArcs)
 	{
-		FinalizePackage(Package, bAllowMissingImports);
+		GraphArchive << InternalArc.FromExportBundleIndex;
+		GraphArchive << InternalArc.ToExportBundleIndex;
+	}
+
+	for (FPackageId ImportedPackageId : ImportedPackageIds)
+	{
+		TArray<FPackageStorePackage::FExternalArc>* FindArcsFromImportedPackage = GraphData.ExternalArcs.Find(ImportedPackageId);
+		if (!FindArcsFromImportedPackage)
+		{
+			int32 ExternalArcCount = 0;
+			GraphArchive << ExternalArcCount;
+		}
+		else
+		{
+			Algo::Sort(*FindArcsFromImportedPackage, [](const FPackageStorePackage::FExternalArc& A, const FPackageStorePackage::FExternalArc& B)
+			{
+				if (A.ToExportBundleIndex == B.ToExportBundleIndex)
+				{
+					if (A.FromImportIndex == B.FromImportIndex)
+					{
+						return A.FromCommandType < B.FromCommandType;
+					}
+					return A.FromImportIndex < B.FromImportIndex;
+				}
+				return A.ToExportBundleIndex < B.ToExportBundleIndex;
+			});
+			int32 ExternalArcCount = FindArcsFromImportedPackage->Num();
+			GraphArchive << ExternalArcCount;
+			for (FPackageStorePackage::FExternalArc& Arc : *FindArcsFromImportedPackage)
+			{
+				GraphArchive << Arc.FromImportIndex;
+				uint8 FromCommandType = uint8(Arc.FromCommandType);
+				GraphArchive << FromCommandType;
+				GraphArchive << Arc.ToExportBundleIndex;
+			}
+		}
 	}
 }
 
 void FPackageStoreOptimizer::FinalizePackageHeader(FPackageStorePackage* Package) const
 {
-	// Temporary Archive for serializing ImportMap
 	FBufferWriter ImportMapArchive(nullptr, 0, EBufferWriterFlags::AllowResize | EBufferWriterFlags::TakeOwnership);
 	for (const FPackageStorePackage::FImport& Import : Package->Imports)
 	{
@@ -1028,49 +1476,12 @@ void FPackageStoreOptimizer::FinalizePackageHeader(FPackageStorePackage* Package
 	}
 	Package->ImportMapSize = ImportMapArchive.Tell();
 
-	// Temporary Archive for serializing EDL graph data
-	FBufferWriter GraphArchive(nullptr, 0, EBufferWriterFlags::AllowResize | EBufferWriterFlags::TakeOwnership);
-	int32 ReferencedPackagesCount = Package->ExternalArcs.Num();
-	GraphArchive << ReferencedPackagesCount;
-	TArray<TTuple<FPackageId, TArray<FPackageStorePackage::FArc>>> SortedExternalArcs;
-	SortedExternalArcs.Reserve(Package->ExternalArcs.Num());
-	for (auto& KV : Package->ExternalArcs)
-	{
-		FPackageId ImportedPackageId = KV.Key;
-		TArray<FPackageStorePackage::FArc> SortedArcs = KV.Value;
-		Algo::Sort(SortedArcs, [](const FPackageStorePackage::FArc& A, const FPackageStorePackage::FArc& B)
-		{
-			if (A.FromNodeIndex == B.FromNodeIndex)
-			{
-				return A.ToNodeIndex < B.ToNodeIndex;
-			}
-			return A.FromNodeIndex < B.FromNodeIndex;
-		});
-		SortedExternalArcs.Emplace(ImportedPackageId, MoveTemp(SortedArcs));
-	}
-	Algo::Sort(SortedExternalArcs, [](const TTuple<FPackageId, TArray<FPackageStorePackage::FArc>>& A, const TTuple<FPackageId, TArray<FPackageStorePackage::FArc>>& B)
-	{
-		return A.Key < B.Key;
-	});
-	for (auto& KV : SortedExternalArcs)
-	{
-		FPackageId ImportedPackageId = KV.Key;
-		TArray<FPackageStorePackage::FArc>& Arcs = KV.Value;
-		int32 ExternalArcCount = Arcs.Num();
-
-		GraphArchive << ImportedPackageId;
-		GraphArchive << ExternalArcCount;
-		GraphArchive.Serialize(Arcs.GetData(), ExternalArcCount * sizeof(FPackageStorePackage::FArc));
-	}
-	Package->ExportBundleArcsSize = GraphArchive.Tell();
-
-	// Temporary Archive for serializing export map data
 	FBufferWriter ExportMapArchive(nullptr, 0, EBufferWriterFlags::AllowResize | EBufferWriterFlags::TakeOwnership);
 	for (const FPackageStorePackage::FExport& Export : Package->Exports)
 	{
 		FExportMapEntry ExportMapEntry;
 		ExportMapEntry.CookedSerialOffset = Export.CookedSerialOffset;
-		ExportMapEntry.CookedSerialSize = Export.CookedSerialSize;
+		ExportMapEntry.CookedSerialSize = Export.SerialSize;
 		ExportMapEntry.ObjectName = Package->NameMapBuilder.MapName(Export.ObjectName);
 		ExportMapEntry.OuterIndex = Export.OuterIndex;
 		ExportMapEntry.ClassIndex = Export.ClassIndex;
@@ -1092,280 +1503,164 @@ void FPackageStoreOptimizer::FinalizePackageHeader(FPackageStorePackage* Package
 	}
 	Package->ExportMapSize = ExportMapArchive.Tell();
 
-	// Temporary archive for serializing export bundle data
-	FBufferWriter ExportBundlesArchive(nullptr, 0, EBufferWriterFlags::AllowResize | EBufferWriterFlags::TakeOwnership);
-	uint32 ExportBundleEntryIndex = 0;
-	for (const FPackageStorePackage::FExportBundle& ExportBundle : Package->ExportBundles)
-	{
-		const uint32 EntryCount = ExportBundle.Entries.Num();
-		FExportBundleHeader ExportBundleHeader{ ExportBundleEntryIndex, EntryCount };
-		ExportBundlesArchive << ExportBundleHeader;
-
-		ExportBundleEntryIndex += EntryCount;
-	}
-	for (const FPackageStorePackage::FExportBundle& ExportBundle : Package->ExportBundles)
+	FBufferWriter ExportBundleEntriesArchive(nullptr, 0, EBufferWriterFlags::AllowResize | EBufferWriterFlags::TakeOwnership);
+	for (const FPackageStorePackage::FExportBundle& ExportBundle : Package->GraphData.ExportBundles)
 	{
 		for (FExportBundleEntry BundleEntry : ExportBundle.Entries)
 		{
-			ExportBundlesArchive << BundleEntry;
+			ExportBundleEntriesArchive << BundleEntry;
 		}
 	}
-	Package->ExportBundleHeadersSize = ExportBundlesArchive.Tell();
+	Package->ExportBundleEntriesSize = ExportBundleEntriesArchive.Tell();
+
+	FBufferWriter GraphArchive(nullptr, 0, EBufferWriterFlags::AllowResize | EBufferWriterFlags::TakeOwnership);
+	SerializeGraphData(Package->ImportedPackageIds, Package->GraphData, GraphArchive);
+	Package->GraphDataSize = GraphArchive.Tell();
 
 	Package->NameMapBuilder.MarkNameAsReferenced(Package->Name);
 	FMappedName MappedPackageName = Package->NameMapBuilder.MapName(Package->Name);
 	Package->NameMapBuilder.MarkNameAsReferenced(Package->SourceName);
 	FMappedName MappedPackageSourceName = Package->NameMapBuilder.MapName(Package->SourceName);
 
-	TArray<uint8> NamesBuffer;
-	TArray<uint8> NameHashesBuffer;
-	SaveNameBatch(Package->NameMapBuilder.GetNameMap(), NamesBuffer, NameHashesBuffer);
-	Package->NameMapSize = Align(NamesBuffer.Num(), 8) + NameHashesBuffer.Num();
+	FBufferWriter NameMapArchive(nullptr, 0, EBufferWriterFlags::AllowResize | EBufferWriterFlags::TakeOwnership);
+	SaveNameBatch(Package->NameMapBuilder.GetNameMap(), NameMapArchive);
+	Package->NameMapSize = NameMapArchive.Tell();
 
-	uint64 HeaderSerialSize =
+	Package->HeaderSize =
 		sizeof(FPackageSummary)
 		+ Package->NameMapSize
 		+ Package->ImportMapSize
 		+ Package->ExportMapSize
-		+ Package->ExportBundleHeadersSize
-		+ Package->ExportBundleArcsSize;
+		+ Package->ExportBundleEntriesSize
+		+ Package->GraphDataSize;
 
-	Package->OptimizedHeaderBuffer = FIoBuffer(HeaderSerialSize);
-	uint8* HeaderData = Package->OptimizedHeaderBuffer.Data();
-	FMemory::Memzero(HeaderData, HeaderSerialSize);
+	Package->HeaderBuffer = FIoBuffer(Package->HeaderSize);
+	uint8* HeaderData = Package->HeaderBuffer.Data();
+	FMemory::Memzero(HeaderData, Package->HeaderSize);
 	FPackageSummary* PackageSummary = reinterpret_cast<FPackageSummary*>(HeaderData);
-
+	PackageSummary->HeaderSize = Package->HeaderSize;
 	PackageSummary->Name = MappedPackageName;
 	PackageSummary->SourceName = MappedPackageSourceName;
 	PackageSummary->PackageFlags = Package->PackageFlags;
 	PackageSummary->CookedHeaderSize = Package->CookedHeaderSize;
-	FBufferWriter HeaderArchive(HeaderData, HeaderSerialSize);
+	FBufferWriter HeaderArchive(HeaderData, Package->HeaderSize);
 	HeaderArchive.Seek(sizeof(FPackageSummary));
 
-	// NameMap data
-	{
-		PackageSummary->NameMapNamesOffset = HeaderArchive.Tell();
-		check(PackageSummary->NameMapNamesOffset % 8 == 0);
-		PackageSummary->NameMapNamesSize = NamesBuffer.Num();
-		HeaderArchive.Serialize(NamesBuffer.GetData(), NamesBuffer.Num());
-		PackageSummary->NameMapHashesOffset = Align(HeaderArchive.Tell(), 8);
-		int32 PaddingByteCount = PackageSummary->NameMapHashesOffset - HeaderArchive.Tell();
-		if (PaddingByteCount)
-		{
-			check(PaddingByteCount < 8);
-			uint8 PaddingBytes[8]{ 0 };
-			HeaderArchive.Serialize(PaddingBytes, PaddingByteCount);
-		}
-		PackageSummary->NameMapHashesSize = NameHashesBuffer.Num();
-		HeaderArchive.Serialize(NameHashesBuffer.GetData(), NameHashesBuffer.Num());
-	}
-
-	// ImportMap data
-	{
-		check(ImportMapArchive.Tell() == Package->ImportMapSize);
-		PackageSummary->ImportMapOffset = HeaderArchive.Tell();
-		HeaderArchive.Serialize(ImportMapArchive.GetWriterData(), ImportMapArchive.Tell());
-	}
-
-	// ExportMap data
-	{
-		check(ExportMapArchive.Tell() == Package->ExportMapSize);
-		PackageSummary->ExportMapOffset = HeaderArchive.Tell();
-		HeaderArchive.Serialize(ExportMapArchive.GetWriterData(), ExportMapArchive.Tell());
-	}
-
-	// ExportBundle data
-	{
-		check(ExportBundlesArchive.Tell() == Package->ExportBundleHeadersSize);
-		PackageSummary->ExportBundlesOffset = HeaderArchive.Tell();
-		HeaderArchive.Serialize(ExportBundlesArchive.GetWriterData(), ExportBundlesArchive.Tell());
-	}
+	HeaderArchive.Serialize(NameMapArchive.GetWriterData(), NameMapArchive.Tell());
+	PackageSummary->ImportMapOffset = HeaderArchive.Tell();
+	HeaderArchive.Serialize(ImportMapArchive.GetWriterData(), ImportMapArchive.Tell());
+	PackageSummary->ExportMapOffset = HeaderArchive.Tell();
+	HeaderArchive.Serialize(ExportMapArchive.GetWriterData(), ExportMapArchive.Tell());
+	PackageSummary->ExportBundleEntriesOffset = HeaderArchive.Tell();
+	HeaderArchive.Serialize(ExportBundleEntriesArchive.GetWriterData(), ExportBundleEntriesArchive.Tell());
+	PackageSummary->GraphDataOffset = HeaderArchive.Tell();
+	HeaderArchive.Serialize(GraphArchive.GetWriterData(), GraphArchive.Tell());
+	check(HeaderArchive.Tell() == PackageSummary->HeaderSize)
+}
 
-	// Graph data
+void FPackageStoreOptimizer::FinalizePackage(FPackageStorePackage* Package)
+{
+	++TotalPackageCount;
+	TotalExportBundleCount += Package->GraphData.ExportBundles.Num();
+	uint64 PackageExportBundleEntryCount = 0;
+	for (const FPackageStorePackage::FExportBundle& ExportBundle : Package->GraphData.ExportBundles)
 	{
-		check(GraphArchive.Tell() == Package->ExportBundleArcsSize);
-		PackageSummary->GraphDataOffset = HeaderArchive.Tell();
-		PackageSummary->GraphDataSize = Package->ExportBundleArcsSize;
-		HeaderArchive.Serialize(GraphArchive.GetWriterData(), GraphArchive.Tell());
+		PackageExportBundleEntryCount += ExportBundle.Entries.Num();
 	}
-}
+	check(PackageExportBundleEntryCount == FExportBundleEntry::ExportCommandType_Count * Package->Exports.Num());
+	TotalExportBundleEntryCount += PackageExportBundleEntryCount;
 
-void FPackageStoreOptimizer::FinalizePackage(FPackageStorePackage* Package, bool bAllowMissingImports)
-{
 	for (FPackageStorePackage::FExportGraphNode& Node : Package->ExportGraphNodes)
 	{
 		check(Node.ExportBundleIndex >= 0);
-		for (FPackageStorePackage::FExternalDependency& ExternalDependency : Node.ExternalDependencies)
+		for (FPackageStorePackage::FExportGraphNode* InternalDependency : Node.InternalDependencies)
 		{
-			if (ExternalDependency.bIsConfirmedMissing)
-			{
-				UE_LOG(LogPackageStoreOptimizer, Warning, TEXT("Package '%s' missing import '%s'"), *Package->Name.ToString(), *ExternalDependency.DebugImportFullName.ToString());
-				continue;
-			}
-			TArray<FPackageStorePackage::FArc>& ArcsFromImportedPackage = Package->ExternalArcs.FindOrAdd(ExternalDependency.PackageId);
-			FPackageStorePackage::FArc Arc;
-			Arc.FromNodeIndex = ExternalDependency.ResolvedExportBundleIndex;
-			Arc.ToNodeIndex = Node.ExportBundleIndex;
-			if (!ArcsFromImportedPackage.Contains(Arc))
+			FPackageStorePackage::FInternalArc Arc;
+			check(InternalDependency->ExportBundleIndex >= 0);
+			Arc.FromExportBundleIndex = InternalDependency->ExportBundleIndex;
+			Arc.ToExportBundleIndex = Node.ExportBundleIndex;
+			if (Arc.FromExportBundleIndex != Arc.ToExportBundleIndex && !Package->GraphData.InternalArcs.Contains(Arc))
 			{
-				ArcsFromImportedPackage.Add(Arc);
-				++TotalExportBundleArcCount;
+				Package->GraphData.InternalArcs.Add(Arc);
+				++TotalInternalBundleArcsCount;
 			}
 		}
-	}
-
-	FinalizePackageHeader(Package);
-
-	for (FPackageStorePackage* ImportedPackage : Package->ImportedWaitingPackages)
-	{
-		ImportedPackage->ImportedByWaitingPackages.Remove(Package);
-	}
-	for (FPackageStorePackage* ImportedByPackage : Package->ImportedByWaitingPackages)
-	{
-		ImportedByPackage->ImportedWaitingPackages.Remove(Package);
-	}
-	WaitingPackagesMap.Remove(Package->Id);
-
-	TArray<FResolvedPublicExport>& ResolvedPublicExports = ResolvedPublicExportsMap.FindOrAdd(Package->Id);
-	check(ResolvedPublicExports.IsEmpty());
-	for (const FPackageStorePackage::FExport& Export : Package->Exports)
-	{
-		if (Export.bIsPublic)
-		{
-			FResolvedPublicExport& PublicExport = ResolvedPublicExports.AddDefaulted_GetRef();
-			PublicExport.GlobalImportIndex = Export.GlobalImportIndex;
-			check(!PublicExport.GlobalImportIndex.IsNull());
-			PublicExport.CreateExportBundleIndex = Export.CreateNode->ExportBundleIndex;
-			check(PublicExport.CreateExportBundleIndex >= 0);
-			PublicExport.SerializeExportBundleIndex = Export.SerializeNode->ExportBundleIndex;
-			check(PublicExport.SerializeExportBundleIndex >= 0);
-		}
-	}
-
-	ResolvedPackages.Add(Package);
-}
-
-bool FPackageStoreOptimizer::IsWaitingForImportsRecursive(FPackageStorePackage* Package) const
-{
-	if (Package->UnresolvedImportedPackagesCount)
-	{
-		return true;
-	}
-
-	if (Package->bTemporaryMark)
-	{
-		return false;
-	}
-	Package->bTemporaryMark = true;
-
-	for (FPackageStorePackage* ImportedPackage : Package->ImportedWaitingPackages)
-	{
-		if (IsWaitingForImportsRecursive(ImportedPackage))
-		{
-			Package->bTemporaryMark = false;
-			return true;
-		}
-	}
-	Package->bTemporaryMark = false;
-	return false;
-}
 
-void FPackageStoreOptimizer::BeginResolvePackage(FPackageStorePackage* Package)
-{
-	++TotalPackageCount;
-	for (const FPackageId& ImportedPackageId : Package->ImportedPackageIds)
-	{
-		check(ImportedPackageId != Package->Id);
-		if (!ResolvedPublicExportsMap.Contains(ImportedPackageId))
+		for (FPackageStorePackage::FExternalDependency& ExternalDependency : Node.ExternalDependencies)
 		{
-			FPackageStorePackage* FindWaitingPackage = WaitingPackagesMap.FindRef(ImportedPackageId);
-			if (FindWaitingPackage)
+			check(ExternalDependency.ImportIndex >= 0);
+			const FPackageStorePackage::FImport& Import = Package->Imports[ExternalDependency.ImportIndex];
+			if (ExternalDependency.bIsConfirmedMissing)
 			{
-				Package->ImportedWaitingPackages.Add(FindWaitingPackage);
-				FindWaitingPackage->ImportedByWaitingPackages.Add(Package);
+				UE_LOG(LogPackageStoreOptimizer, Warning, TEXT("Package '%s' missing import '%s'"), *Package->Name.ToString(), *Import.FullName);
+				continue;
 			}
-			else
+
+			TArray<FPackageStorePackage::FExternalArc>& ArcsFromImportedPackage = Package->GraphData.ExternalArcs.FindOrAdd(Import.PackageId);
+			FPackageStorePackage::FExternalArc Arc;
+			Arc.FromImportIndex = ExternalDependency.ImportIndex;
+			Arc.FromCommandType = ExternalDependency.ExportBundleCommandType;
+			Arc.ToExportBundleIndex = Node.ExportBundleIndex;
+			if (!ArcsFromImportedPackage.Contains(Arc))
 			{
-				++Package->UnresolvedImportedPackagesCount;
-				ImportedPackageToWaitingPackagesMap.Add(ImportedPackageId, Package);
+				ArcsFromImportedPackage.Add(Arc);
+				++TotalExternalBundleArcsCount;
 			}
 		}
-		else
-		{
-			check(!WaitingPackagesMap.Contains(ImportedPackageId));
-		}
-	}
-
-	WaitingPackagesMap.Add(Package->Id, Package);
-	
-	for (auto WaitingPackageIt = ImportedPackageToWaitingPackagesMap.CreateKeyIterator(Package->Id); WaitingPackageIt; ++WaitingPackageIt)
-	{
-		FPackageStorePackage* WaitingPackage = WaitingPackageIt.Value();
-		WaitingPackage->ImportedWaitingPackages.Add(Package);
-		Package->ImportedByWaitingPackages.Add(WaitingPackage);
-		check(WaitingPackage->UnresolvedImportedPackagesCount > 0);
-		WaitingPackageIt.RemoveCurrent();
-	}
-
-	if (!bIncrementalResolveEnabled)
-	{
-		return;
 	}
-	TArray<FPackageStorePackage*> PackagesReadyToResolve;
+	uint64 SerialOffset = 0;
+	for (FPackageStorePackage::FExportBundle& ExportBundle : Package->GraphData.ExportBundles)
 	{
-		TRACE_CPUPROFILER_EVENT_SCOPE(UpdateWaitList);
-		for (auto& KV : WaitingPackagesMap)
+		ExportBundle.SerialOffset = SerialOffset;
+		for (const FExportBundleEntry& BundleEntry : ExportBundle.Entries)
 		{
-			FPackageStorePackage* WaitingPackage = KV.Value;
-			if (!IsWaitingForImportsRecursive(WaitingPackage))
+			if (BundleEntry.CommandType == FExportBundleEntry::ExportCommandType_Serialize)
 			{
-				PackagesReadyToResolve.Add(WaitingPackage);
+				const FPackageStorePackage::FExport& Export = Package->Exports[BundleEntry.LocalExportIndex];
+				SerialOffset += Export.SerialSize;
 			}
 		}
 	}
 
-	if (PackagesReadyToResolve.Num())
-	{
-		ResolvePackages(PackagesReadyToResolve, false);
-	}
+	FinalizePackageHeader(Package);
 }
 
 FIoBuffer FPackageStoreOptimizer::CreatePackageBuffer(const FPackageStorePackage* Package, const FIoBuffer& CookedExportsDataBuffer, TArray<FFileRegion>* InOutFileRegions) const
 {
-	check(Package->OptimizedHeaderBuffer.DataSize() > 0);
-	const uint64 BundleBufferSize = Package->OptimizedHeaderBuffer.DataSize() + Package->ExportsSerialSize;
+	check(Package->HeaderBuffer.DataSize() > 0);
+	check(Package->HeaderBuffer.DataSize() == Package->HeaderSize);
+	const uint64 BundleBufferSize = Package->HeaderSize + Package->ExportsSerialSize;
 	FIoBuffer BundleBuffer(BundleBufferSize);
-	FMemory::Memcpy(BundleBuffer.Data(), Package->OptimizedHeaderBuffer.Data(), Package->OptimizedHeaderBuffer.DataSize());
-	uint64 BundleBufferOffset = Package->OptimizedHeaderBuffer.DataSize();
+	FMemory::Memcpy(BundleBuffer.Data(), Package->HeaderBuffer.Data(), Package->HeaderSize);
+	uint64 BundleBufferOffset = Package->HeaderSize;
 
 	TArray<FFileRegion> OutputRegions;
 
-	for (const FPackageStorePackage::FExportBundle& ExportBundle : Package->ExportBundles)
+	for (const FPackageStorePackage::FExportBundle& ExportBundle : Package->GraphData.ExportBundles)
 	{
 		for (const FExportBundleEntry& BundleEntry : ExportBundle.Entries)
 		{
 			if (BundleEntry.CommandType == FExportBundleEntry::ExportCommandType_Serialize)
 			{
 				const FPackageStorePackage::FExport& Export = Package->Exports[BundleEntry.LocalExportIndex];
-				uint64 AdjustedSerialOffset = Export.CookedSerialOffset - Package->CookedHeaderSize;
-				check(AdjustedSerialOffset + Export.CookedSerialSize <= CookedExportsDataBuffer.DataSize());
-				FMemory::Memcpy(BundleBuffer.Data() + BundleBufferOffset, CookedExportsDataBuffer.Data() + AdjustedSerialOffset, Export.CookedSerialSize);
+				check(Export.SerialOffset + Export.SerialSize <= CookedExportsDataBuffer.DataSize());
+				FMemory::Memcpy(BundleBuffer.Data() + BundleBufferOffset, CookedExportsDataBuffer.Data() + Export.SerialOffset, Export.SerialSize);
 
 				if (InOutFileRegions)
 				{
 					// Find overlapping regions and adjust them to match the new offset of the export data
 					for (const FFileRegion& Region : *InOutFileRegions)
 					{
-						if (AdjustedSerialOffset <= Region.Offset && Region.Offset + Region.Length <= AdjustedSerialOffset + Export.CookedSerialSize)
+						if (Export.SerialOffset <= Region.Offset && Region.Offset + Region.Length <= Export.SerialOffset + Export.SerialSize)
 						{
 							FFileRegion NewRegion = Region;
-							NewRegion.Offset -= AdjustedSerialOffset;
+							NewRegion.Offset -= Export.SerialOffset;
 							NewRegion.Offset += BundleBufferOffset;
 							OutputRegions.Add(NewRegion);
 						}
 					}
 				}
-				BundleBufferOffset += Export.CookedSerialSize;
+				BundleBufferOffset += Export.SerialSize;
 			}
 		}
 	}
@@ -1484,36 +1779,10 @@ void FPackageStoreOptimizer::FindScriptObjects(const ITargetPlatform* TargetPlat
 		}
 	}
 
-	FLargeMemoryWriter ScriptObjectsArchive(0, true);
-
-	int32 NumScriptObjects = ScriptObjectsMap.Num();
-	ScriptObjectsArchive << NumScriptObjects;
 	TotalScriptObjectCount = ScriptObjectsMap.Num();
 }
 
-void FPackageStoreOptimizer::Flush(bool bAllowMissingImports, TFunction<void(FPackageStorePackage*)> ResolvedPackageCallback)
-{
-	if (bAllowMissingImports)
-	{
-		TArray<FPackageStorePackage*> PackagesToResolve;
-		WaitingPackagesMap.GenerateValueArray(PackagesToResolve);
-		ResolvePackages(PackagesToResolve, true);
-		
-		check(WaitingPackagesMap.IsEmpty());
-		ImportedPackageToWaitingPackagesMap.Empty();
-	}
-
-	if (ResolvedPackageCallback)
-	{
-		for (FPackageStorePackage* ResolvedPackage : ResolvedPackages)
-		{
-			ResolvedPackageCallback(ResolvedPackage);
-		}
-	}
-	ResolvedPackages.Empty();
-}
-
-uint64 FPackageStoreOptimizer::WriteScriptObjects(FIoStoreWriter* IoStoreWriter) const
+FIoBuffer FPackageStoreOptimizer::CreateScriptObjectsBuffer() const
 {
 	TArray<FScriptObjectData> ScriptObjectsAsArray;
 	ScriptObjectsMap.GenerateValueArray(ScriptObjectsAsArray);
@@ -1522,62 +1791,55 @@ uint64 FPackageStoreOptimizer::WriteScriptObjects(FIoStoreWriter* IoStoreWriter)
 		return A.FullName < B.FullName;
 	});
 	
-	FLargeMemoryWriter ScriptObjectsArchive(0, true);
-	int32 NumScriptObjects = ScriptObjectsAsArray.Num();
-	ScriptObjectsArchive << NumScriptObjects;
+	TArray<FScriptObjectEntry> ScriptObjectEntries;
+	ScriptObjectEntries.Reserve(ScriptObjectsAsArray.Num());
 	FPackageStoreNameMapBuilder NameMapBuilder;
 	NameMapBuilder.SetNameMapType(FMappedName::EType::Global);
 	for (const FScriptObjectData& ImportData : ScriptObjectsAsArray)
 	{
 		NameMapBuilder.MarkNameAsReferenced(ImportData.ObjectName);
-		FScriptObjectEntry Entry;
+		FScriptObjectEntry& Entry = ScriptObjectEntries.AddDefaulted_GetRef();
 		Entry.ObjectName = NameMapBuilder.MapName(ImportData.ObjectName).ToUnresolvedMinimalName();
 		Entry.GlobalIndex = ImportData.GlobalIndex;
 		Entry.OuterIndex = ImportData.OuterIndex;
 		Entry.CDOClassIndex = ImportData.CDOClassIndex;
+	}
 
+	FLargeMemoryWriter ScriptObjectsArchive(0, true);
+	SaveNameBatch(NameMapBuilder.GetNameMap(), ScriptObjectsArchive);
+	int32 NumScriptObjects = ScriptObjectEntries.Num();
+	ScriptObjectsArchive << NumScriptObjects;
+	for (FScriptObjectEntry& Entry : ScriptObjectEntries)
+	{
 		ScriptObjectsArchive << Entry;
 	}
 
-	FIoWriteOptions WriteOptions;
-	WriteOptions.DebugName = TEXT("LoaderInitialLoadMeta");
-	IoStoreWriter->Append(CreateIoChunkId(0, 0, EIoChunkType::LoaderInitialLoadMeta), FIoBuffer(FIoBuffer::Wrap, ScriptObjectsArchive.GetData(), ScriptObjectsArchive.TotalSize()), WriteOptions);
-
-	TArray<uint8> Names;
-	TArray<uint8> Hashes;
-	SaveNameBatch(NameMapBuilder.GetNameMap(), /* out */ Names, /* out */ Hashes);
-	WriteOptions.DebugName = TEXT("LoaderGlobalNames");
-	IoStoreWriter->Append(CreateIoChunkId(0, 0, EIoChunkType::LoaderGlobalNames),
-		FIoBuffer(FIoBuffer::Wrap, Names.GetData(), Names.Num()), WriteOptions);
-	WriteOptions.DebugName = TEXT("LoaderGlobalNameHashes");
-	IoStoreWriter->Append(CreateIoChunkId(0, 0, EIoChunkType::LoaderGlobalNameHashes),
-		FIoBuffer(FIoBuffer::Wrap, Hashes.GetData(), Hashes.Num()), WriteOptions);
-
-	return ScriptObjectsArchive.TotalSize() + Names.Num() + Hashes.Num();
+	int64 DataSize = ScriptObjectsArchive.TotalSize();
+	return FIoBuffer(FIoBuffer::AssumeOwnership, ScriptObjectsArchive.ReleaseOwnership(), DataSize);
 }
 
-FPackageStoreContainerHeaderEntry FPackageStoreOptimizer::CreateContainerHeaderEntry(const FPackageStorePackage* Package) const
+FPackageStoreEntryResource FPackageStoreOptimizer::CreatePackageStoreEntry(const FPackageStorePackage* Package) const
 {
-	FPackageStoreContainerHeaderEntry Result;
-	Result.Name = Package->Name;
-	Result.SourceName = Package->SourceName;
-	Result.Region = Package->Region;
-	Result.ExportBundlesSize = Package->OptimizedHeaderBuffer.DataSize() + Package->ExportsSerialSize;
-	Result.ExportCount = Package->Exports.Num();
-	Result.ExportBundleCount = Package->ExportBundles.Num();
-	Result.LoadOrder = Package->GetLoadOrder();
+	FPackageStoreEntryResource Result;
+	Result.Flags = Package->bIsRedirected ? EPackageStoreEntryFlags::Redirected : EPackageStoreEntryFlags::None;
+	Result.PackageName = Package->Name;
+	Result.SourcePackageName = Package->SourceName;
+	Result.Region = FName(*Package->Region);
+	Result.ExportInfo.ExportBundlesSize = Package->HeaderSize + Package->ExportsSerialSize;
+	Result.ExportInfo.ExportCount = Package->Exports.Num();
+	Result.ExportInfo.ExportBundleCount = Package->GraphData.ExportBundles.Num();
+	Result.ExportInfo.LoadOrder = Package->GetLoadOrder();
 	Result.ImportedPackageIds = Package->ImportedPackageIds;
 	Result.ShaderMapHashes = Package->ShaderMapHashes.Array();
-	Result.bIsRedirected = Package->bIsRedirected;
 	return Result;
 }
 
-FContainerHeader FPackageStoreOptimizer::CreateContainerHeader(const FIoContainerId& ContainerId, const TArray<FPackageStoreContainerHeaderEntry>& Packages) const
+FContainerHeader FPackageStoreOptimizer::CreateContainerHeader(const FIoContainerId& ContainerId, TArrayView<const FPackageStoreEntryResource> PackageStoreEntries) const
 {
 	FContainerHeader Header;
 	Header.ContainerId = ContainerId;
-	Header.PackageCount = Packages.Num();
-	int32 StoreTocSize = Header.PackageCount * sizeof(FPackageStoreEntry);
+	Header.PackageCount = PackageStoreEntries.Num();
+	int32 StoreTocSize = Header.PackageCount * sizeof(FFilePackageStoreEntry);
 	FLargeMemoryWriter StoreTocArchive(0, true);
 	FLargeMemoryWriter StoreDataArchive(0, true);
 
@@ -1592,48 +1854,41 @@ FContainerHeader FPackageStoreOptimizer::CreateContainerHeader(const FIoContaine
 		StoreTocArchive << OffsetToDataFromThis;
 	};
 
-	TArray<const FPackageStoreContainerHeaderEntry*> SortedPackages;
-	SortedPackages.Reserve(Packages.Num());
-	for (const FPackageStoreContainerHeaderEntry& Package : Packages)
+	TArray<const FPackageStoreEntryResource*> SortedPackageStoreEntries;
+	SortedPackageStoreEntries.Reserve(PackageStoreEntries.Num());
+	for (const FPackageStoreEntryResource& Entry : PackageStoreEntries)
 	{
-		SortedPackages.Add(&Package);
+		SortedPackageStoreEntries.Add(&Entry);
 	}
-	Algo::Sort(SortedPackages, [](const FPackageStoreContainerHeaderEntry* A, const FPackageStoreContainerHeaderEntry* B)
+	Algo::Sort(SortedPackageStoreEntries, [](const FPackageStoreEntryResource* A, const FPackageStoreEntryResource* B)
 	{
-		return A->GetId() < B->GetId();
+		return A->GetPackageId() < B->GetPackageId();
 	});
 
-	Header.PackageIds.Reserve(SortedPackages.Num());
-	for (const FPackageStoreContainerHeaderEntry* Package : SortedPackages)
+	Header.PackageIds.Reserve(SortedPackageStoreEntries.Num());
+	for (const FPackageStoreEntryResource* Entry : SortedPackageStoreEntries)
 	{
-		Header.PackageIds.Add(Package->GetId());
-		if (Package->bIsRedirected)
+		Header.PackageIds.Add(Entry->GetPackageId());
+		if (Entry->IsRedirected())
 		{
-			if (Package->Region.Len() > 0)
+			if (Entry->Region.IsValid())
 			{
-				Header.CulturePackageMap.FindOrAdd(Package->Region).Emplace(Package->GetSourceId(), Package->GetId());
+				Header.CulturePackageMap.FindOrAdd(Entry->Region.ToString()).Emplace(Entry->GetSourcePackageId(), Entry->GetPackageId());
 			}
 			else
 			{
-				Header.PackageRedirects.Add(MakeTuple(Package->GetSourceId(), Package->GetId()));
+				Header.PackageRedirects.Add(MakeTuple(Entry->GetSourcePackageId(), Entry->GetPackageId()));
 			}
 		}
 		
 		// StoreEntries
-		uint64 ExportBundlesSize = Package->ExportBundlesSize;
-		int32 ExportCount = Package->ExportCount;
-		int32 ExportBundleCount = Package->ExportBundleCount;
-		uint32 LoadOrder = Package->LoadOrder;
-		uint32 Pad = 0;
-
-		StoreTocArchive << ExportBundlesSize;
-		StoreTocArchive << ExportCount;
-		StoreTocArchive << ExportBundleCount;
-		StoreTocArchive << LoadOrder;
+		FPackageStoreExportInfo ExportInfo = Entry->ExportInfo;
+		int32 Pad = 0;
+		StoreTocArchive << ExportInfo; 
 		StoreTocArchive << Pad;
 
 		// ImportedPackages
-		const TArray<FPackageId>& ImportedPackageIds = Package->ImportedPackageIds;
+		const TArray<FPackageId>& ImportedPackageIds = Entry->ImportedPackageIds;
 		SerializePackageEntryCArrayHeader(ImportedPackageIds.Num());
 		for (FPackageId ImportedPackageId : ImportedPackageIds)
 		{
@@ -1642,7 +1897,7 @@ FContainerHeader FPackageStoreOptimizer::CreateContainerHeader(const FIoContaine
 		}
 		
 		// ShaderMapHashes
-		const TArray<FSHAHash>& ShaderMapHashes = Package->ShaderMapHashes;
+		const TArray<FSHAHash>& ShaderMapHashes = Entry->ShaderMapHashes;
 		SerializePackageEntryCArrayHeader(ShaderMapHashes.Num());
 		for (const FSHAHash& ShaderMapHash : ShaderMapHashes)
 		{
@@ -1657,8 +1912,6 @@ FContainerHeader FPackageStoreOptimizer::CreateContainerHeader(const FIoContaine
 	FBufferWriter PackageStoreArchive(Header.StoreEntries.GetData(), StoreByteCount);
 	PackageStoreArchive.Serialize(StoreTocArchive.GetData(), StoreTocArchive.TotalSize());
 	PackageStoreArchive.Serialize(StoreDataArchive.GetData(), StoreDataArchive.TotalSize());
-	FPackageStoreNameMapBuilder DummyNameMapBuilder;
-	SaveNameBatch(DummyNameMapBuilder.GetNameMap(), Header.Names, Header.NameHashes);
 
 	return Header;
 }
diff --git a/Engine/Source/Developer/IoStoreUtilities/Private/PackageStoreOptimizer.h b/Engine/Source/Developer/IoStoreUtilities/Private/PackageStoreOptimizer.h
--- a/Engine/Source/Developer/IoStoreUtilities/Private/PackageStoreOptimizer.h
+++ b/Engine/Source/Developer/IoStoreUtilities/Private/PackageStoreOptimizer.h
@@ -3,10 +3,13 @@
 #pragma once
 
 #include "Serialization/AsyncLoading2.h"
+#include "IO/PackageStore.h"
 #include "IO/IoDispatcher.h"
 #include "UObject/PackageFileSummary.h"
 #include "UObject/UObjectMarks.h"
 
+class FBufferWriter;
+
 class FPackageStoreNameMapBuilder
 {
 public:
@@ -78,38 +81,6 @@ private:
 	FMappedName::EType NameMapType = FMappedName::EType::Package;
 };
 
-class FPackageStoreContainerHeaderEntry
-{
-public:
-	FPackageId GetId() const
-	{
-		return FPackageId::FromName(Name);
-	}
-
-	FPackageId GetSourceId() const
-	{
-		if (SourceName.IsNone())
-		{
-			return FPackageId();
-		}
-		return FPackageId::FromName(SourceName);
-	}
-
-private:
-	FName Name;
-	FName SourceName;
-	FString Region;
-	uint64 ExportBundlesSize = 0;
-	int32 ExportCount = 0;
-	int32 ExportBundleCount = 0;
-	uint32 LoadOrder = 0;
-	TArray<FPackageId> ImportedPackageIds;
-	TArray<FSHAHash> ShaderMapHashes;
-	bool bIsRedirected = false;
-
-	friend class FPackageStoreOptimizer;
-};
-
 class FPackageStorePackage
 {
 public:
@@ -118,14 +89,14 @@ public:
 		return Id;
 	}
 
-	uint64 GetLoadOrder() const
+	uint32 GetLoadOrder() const
 	{
-		return ExportBundles.Num() > 0 ? ExportBundles[0].LoadOrder : 0;
+		return LoadOrder;
 	}
 
-	uint64 GetExportBundleCount() const
+	uint64 GetHeaderSize() const
 	{
-		return ExportBundles.Num();
+		return HeaderSize;
 	}
 
 	uint64 GetImportMapSize() const
@@ -137,15 +108,20 @@ public:
 	{
 		return ExportMapSize;
 	}
-	
-	uint64 GetExportBundleArcsSize() const
+
+	uint64 GetExportBundleCount() const
 	{
-		return ExportBundleArcsSize;
+		return GraphData.ExportBundles.Num();
 	}
 
-	uint64 GetExportBundleHeadersSize() const
+	uint64 GetExportBundleEntriesSize() const
+	{
+		return ExportBundleEntriesSize;
+	}
+	
+	uint64 GetGraphDataSize() const
 	{
-		return ExportBundleHeadersSize;
+		return GraphDataSize;
 	}
 
 	uint64 GetNameCount() const
@@ -188,22 +164,19 @@ private:
 
 	struct FExternalDependency
 	{
-		FPackageId PackageId;
-		FPackageObjectIndex GlobalImportIndex;
-		FName DebugImportFullName;
+		int32 ImportIndex = -1;
 		FExportBundleEntry::EExportCommandType ExportBundleCommandType;
-		FExportGraphNode* ExportGraphNode = nullptr;
-		int32 ResolvedExportBundleIndex = -1;
 		bool bIsConfirmedMissing = false;
 	};
 
 	struct FExportGraphNode
 	{
-		FPackageStorePackage* Package = nullptr;
 		FExportBundleEntry BundleEntry;
+		TArray<FExportGraphNode*> InternalDependencies;
 		TArray<FExternalDependency> ExternalDependencies;
 		int32 ExportBundleIndex = -1;
 		int32 IncomingEdgeCount = 0;
+		bool bIsPublic = false;
 	};
 
 	struct FExport
@@ -217,18 +190,18 @@ private:
 		FPackageObjectIndex TemplateIndex;
 		EObjectFlags ObjectFlags = RF_NoFlags;
 		uint64 CookedSerialOffset = uint64(-1);
-		uint64 CookedSerialSize = uint64(-1);
+		uint64 SerialOffset = uint64(-1);
+		uint64 SerialSize = uint64(-1);
 		bool bNotForClient = false;
 		bool bNotForServer = false;
 		bool bIsPublic = false;
-		FExportGraphNode* CreateNode = nullptr;
-		FExportGraphNode* SerializeNode = nullptr;
+		FExportGraphNode* Nodes[FExportBundleEntry::ExportCommandType_Count] = { nullptr };
 	};
 
 	struct FExportBundle
 	{
+		uint64 SerialOffset = uint64(-1);
 		TArray<FExportBundleEntry> Entries;
-		uint32 LoadOrder = MAX_int32;
 	};
 
 	struct FImport
@@ -240,54 +213,79 @@ private:
 		bool bIsPackageImport = false;
 	};
 
-	struct FArc
+	struct FInternalArc
 	{
-		uint32 FromNodeIndex;
-		uint32 ToNodeIndex;
+		int32 FromExportBundleIndex;
+		int32 ToExportBundleIndex;
 
-		bool operator==(const FArc& Other) const
+		bool operator==(const FInternalArc& Other) const
 		{
-			return ToNodeIndex == Other.ToNodeIndex && FromNodeIndex == Other.FromNodeIndex;
+			return FromExportBundleIndex == Other.FromExportBundleIndex &&
+				ToExportBundleIndex == Other.ToExportBundleIndex;
 		}
 	};
 
+	struct FExternalArc
+	{
+		int32 FromImportIndex;
+		FExportBundleEntry::EExportCommandType FromCommandType;
+		int32 ToExportBundleIndex;
+
+		bool operator==(const FExternalArc& Other) const
+		{
+			return FromImportIndex == Other.FromImportIndex &&
+				FromCommandType == Other.FromCommandType &&
+				ToExportBundleIndex == Other.ToExportBundleIndex;
+		}
+	};
+
+	struct FGraphData
+	{
+		TArray<FExportBundle> ExportBundles;
+		TArray<FInternalArc> InternalArcs;
+		TMap<FPackageId, TArray<FExternalArc>> ExternalArcs;
+	};
+
+	struct FExportBundleGraphNode
+	{
+		FPackageStorePackage* Package = nullptr;
+		TArray<FExportGraphNode*> ExportGraphNodes;
+		int32 Index = -1;
+		int32 IncomingEdgeCount = 0;
+	};
+
 	FPackageId Id;
 	FName Name;
 	FName SourceName;
 	FString Region;
 
-	FPackageFileSummary Summary;
-	TArray<FObjectImport> ObjectImports;
-	TArray<FObjectExport> ObjectExports;
-	TArray<FPackageIndex> PreloadDependencies;
-
 	FPackageStoreNameMapBuilder NameMapBuilder;
 	TArray<FImport> Imports;
 	TArray<FExport> Exports;
 	TArray<FExportGraphNode> ExportGraphNodes;
-	TMap<FPackageId, TArray<FArc>> ExternalArcs;
-	TArray<FExportBundle> ExportBundles;
+	FGraphData GraphData;
 
 	TArray<FPackageId> ImportedPackageIds;
-	TArray<FPackageStorePackage*> ImportedWaitingPackages;
-	TArray<FPackageStorePackage*> ImportedByWaitingPackages;
 	TSet<FPackageId> RedirectedToPackageIds;
 	TSet<FPackageId> ImportedRedirectedPackageIds;
 	TSet<FSHAHash> ShaderMapHashes;
 
-	FIoBuffer OptimizedHeaderBuffer;
+	FIoBuffer HeaderBuffer;
 
 	uint32 PackageFlags = 0;
 	uint32 CookedHeaderSize = 0;
+	uint64 HeaderSize = 0;
 	uint64 ExportsSerialSize = 0;
 	uint64 ImportMapSize = 0;
 	uint64 ExportMapSize = 0;
-	uint64 ExportBundleArcsSize = 0;
-	uint64 ExportBundleHeadersSize = 0;
+	uint64 ExportBundleEntriesSize = 0;
+	uint64 GraphDataSize = 0;
 	uint64 NameMapSize = 0;
+	uint32 LoadOrder = 0;
 	
-	uint32 UnresolvedImportedPackagesCount = 0;
-	TArray<FExportGraphNode*> NodesWithNoIncomingEdges;
+	TArray<FExportBundleGraphNode> ExportBundleGraphNodes;
+	FPackageStorePackage::FExportBundle* CurrentBundle = nullptr;
+	TArray<FExportBundleGraphNode*> NodesWithNoIncomingEdges;
 	bool bTemporaryMark = false;
 	bool bPermanentMark = false;
 
@@ -314,9 +312,14 @@ public:
 		return TotalExportBundleEntryCount;
 	}
 
-	uint64 GetTotalExportBundleArcCount() const
+	uint64 GetTotalInternalBundleArcsCount() const
+	{
+		return TotalInternalBundleArcsCount;
+	}
+
+	uint64 GetTotalExternalBundleArcsCount() const
 	{
-		return TotalExportBundleArcCount;
+		return TotalExternalBundleArcsCount;
 	}
 
 	uint64 GetTotalScriptObjectCount() const
@@ -326,19 +329,17 @@ public:
 
 	void Initialize(const ITargetPlatform* TargetPlatform);
 
-	void EnableIncrementalResolve()
-	{
-		bIncrementalResolveEnabled = true;
-	}
-
 	FPackageStorePackage* CreateMissingPackage(const FName& Name) const;
 	FPackageStorePackage* CreatePackageFromCookedHeader(const FName& Name, const FIoBuffer& CookedHeaderBuffer) const;
-	void BeginResolvePackage(FPackageStorePackage* Package);
-	void Flush(bool bAllowMissingImports, TFunction<void(FPackageStorePackage*)> ResolvedPackageCallback = TFunction<void(FPackageStorePackage*)>());
+	FPackageStorePackage* CreatePackageFromPackageStoreHeader(const FName& Name, const FIoBuffer& Buffer, const FPackageStoreEntryResource& PackageStoreEntry) const;
+	void FinalizePackage(FPackageStorePackage* Package);
 	FIoBuffer CreatePackageBuffer(const FPackageStorePackage* Package, const FIoBuffer& CookedExportsBuffer, TArray<FFileRegion>* InOutFileRegions) const;
-	FPackageStoreContainerHeaderEntry CreateContainerHeaderEntry(const FPackageStorePackage* Package) const;
-	FContainerHeader CreateContainerHeader(const FIoContainerId& ContainerId, const TArray<FPackageStoreContainerHeaderEntry>& Packages) const;
-	uint64 WriteScriptObjects(FIoStoreWriter* IoStoreWriter) const; // TODO: Merge name map and return FIoBuffer instead
+	FPackageStoreEntryResource CreatePackageStoreEntry(const FPackageStorePackage* Package) const;
+	FContainerHeader CreateContainerHeader(const FIoContainerId& ContainerId, TArrayView<const FPackageStoreEntryResource> PackageStoreEntries) const;
+	FIoBuffer CreateScriptObjectsBuffer() const;
+	void ProcessRedirects(const TMap<FPackageId, FPackageStorePackage*>& PackagesMap) const;
+	void OptimizeExportBundles(const TMap<FPackageId, FPackageStorePackage*>& PackagesMap);
+	void FindImports(const TMap<FPackageId, FPackageStorePackage*>& PackagesMap) const;
 
 private:
 	struct FScriptObjectData
@@ -350,44 +351,59 @@ private:
 		FPackageObjectIndex CDOClassIndex;
 	};
 
-	struct FResolvedPublicExport
+	struct FCookedHeaderData
 	{
-		FPackageObjectIndex GlobalImportIndex;
-		int32 CreateExportBundleIndex = -1;
-		int32 SerializeExportBundleIndex = -1;
+		FPackageFileSummary Summary;
+		TArray<FName> SummaryNames;
+		TArray<FObjectImport> ObjectImports;
+		TArray<FObjectExport> ObjectExports;
+		TArray<FPackageIndex> PreloadDependencies;
+	};
+
+	struct FPackageStoreHeaderData
+	{
+		FPackageSummary Summary;
+		TArray<FPackageId> ImportedPackageIds;
+		TArray<FNameEntryId> NameMap;
+		TArray<FPackageObjectIndex> Imports;
+		TArray<FExportMapEntry> Exports;
+		TArray<FExportBundleHeader> ExportBundleHeaders;
+		TArray<FExportBundleEntry> ExportBundleEntries;
+		TArray<FPackageStorePackage::FInternalArc> InternalArcs;
+		TArray<FPackageStorePackage::FExternalArc> ExternalArcs;
 	};
 
 	using FExportGraphEdges = TMultiMap<FPackageStorePackage::FExportGraphNode*, FPackageStorePackage::FExportGraphNode*>;
+	using FExportBundleGraphEdges = TMultiMap<FPackageStorePackage::FExportBundleGraphNode*, FPackageStorePackage::FExportBundleGraphNode*>;
 
-	void LoadCookedHeader(FPackageStorePackage* Package, const FIoBuffer& CookedHeaderBuffer) const;
+	FCookedHeaderData LoadCookedHeader(const FIoBuffer& CookedHeaderBuffer) const;
+	FPackageStoreHeaderData LoadPackageStoreHeader(const FIoBuffer& PackageStoreHeaderBuffer, const FPackageStoreEntryResource& PackageStoreEntry) const;
 	void ResolveImport(FPackageStorePackage::FImport* Imports, const FObjectImport* ObjectImports, int32 LocalImportIndex) const;
-	void ProcessImports(FPackageStorePackage* Package) const;
+	void ProcessImports(const FCookedHeaderData& CookedHeaderData, FPackageStorePackage* Package) const;
+	void ProcessImports(const FPackageStoreHeaderData& PackageStoreHeaderData, FPackageStorePackage* Package) const;
 	void ResolveExport(FPackageStorePackage::FExport* Exports, const FObjectExport* ObjectExports, const int32 LocalExportIndex, const FName& PackageName) const;
-	void ProcessExports(FPackageStorePackage* Package) const;
-	void SortPackagesInLoadOrder(TArray<FPackageStorePackage*>& Packages) const;
-	FExportGraphEdges ProcessPreloadDependencies(const TArray<FPackageStorePackage*>& Packages) const;
-	TArray<FPackageStorePackage::FExportGraphNode*> SortExportGraphNodesInLoadOrder(const TArray<FPackageStorePackage*>& Packages, FExportGraphEdges& Edges) const;
-	void CreateExportBundles(TArray<FPackageStorePackage*>& Packages);
+	void ResolveExport(FPackageStorePackage::FExport* Exports, const int32 LocalExportIndex, const FName& PackageName) const;
+	void ProcessExports(const FCookedHeaderData& CookedHeaderData, FPackageStorePackage* Package) const;
+	void ProcessExports(const FPackageStoreHeaderData& PackageStoreHeaderData, FPackageStorePackage* Package) const;
+	void ProcessPreloadDependencies(const FCookedHeaderData& CookedHeaderData, FPackageStorePackage* Package) const;
+	void ProcessPreloadDependencies(const FPackageStoreHeaderData& PackageStoreHeaderData, FPackageStorePackage* Package) const;
+	TArray<FPackageStorePackage*> SortPackagesInLoadOrder(const TMap<FPackageId, FPackageStorePackage*>& PackagesMap) const;
+	void SerializeGraphData(const TArray<FPackageId>& ImportedPackageIds, FPackageStorePackage::FGraphData& GraphData, FBufferWriter& GraphArchive) const;
+	TArray<FPackageStorePackage::FExportGraphNode*> SortExportGraphNodesInLoadOrder(FPackageStorePackage* Package, FExportGraphEdges& Edges) const;
+	TArray<FPackageStorePackage::FExportBundleGraphNode*> SortExportBundleGraphNodesInLoadOrder(const TArray<FPackageStorePackage*>& Packages, FExportBundleGraphEdges& Edges) const;
+	void CreateExportBundles(FPackageStorePackage* Package) const;
 	bool RedirectPackage(const FPackageStorePackage& SourcePackage, FPackageStorePackage& TargetPackage, TMap<FPackageObjectIndex, FPackageObjectIndex>& RedirectedToSourceImportIndexMap) const;
 	void RedirectPackageUnverified(FPackageStorePackage& TargetPackage, TMap<FPackageObjectIndex, FPackageObjectIndex>& RedirectedToSourceImportIndexMap) const;
-	void ProcessRedirects(const TArray<FPackageStorePackage*> Packages) const;
-	void ResolvePackages(TArray<FPackageStorePackage*>& Packages, bool bAllowMissingImports);
 	void FinalizePackageHeader(FPackageStorePackage* Package) const;
-	void FinalizePackage(FPackageStorePackage* Package, bool bAllowMissingImports);
-	bool IsWaitingForImportsRecursive(FPackageStorePackage* Package) const;
 	void FindScriptObjectsRecursive(FPackageObjectIndex OuterIndex, UObject* Object, EObjectMark ExcludedObjectMarks, const ITargetPlatform* TargetPlatform);
 	void FindScriptObjects(const ITargetPlatform* TargetPlatform);
 
 	TMap<FPackageObjectIndex, FScriptObjectData> ScriptObjectsMap;
-	TMap<FPackageId, TArray<FResolvedPublicExport>> ResolvedPublicExportsMap;
-	TMap<FPackageId, FPackageStorePackage*> WaitingPackagesMap;
-	TMultiMap<FPackageId, FPackageStorePackage*> ImportedPackageToWaitingPackagesMap;
-	TArray<FPackageStorePackage*> ResolvedPackages;
 	uint64 TotalPackageCount = 0;
 	uint64 TotalExportBundleCount = 0;
 	uint64 TotalExportBundleEntryCount = 0;
-	uint64 TotalExportBundleArcCount = 0;
+	uint64 TotalInternalBundleArcsCount = 0;
+	uint64 TotalExternalBundleArcsCount = 0;
 	uint64 TotalScriptObjectCount = 0;
 	uint32 NextLoadOrder = 0;
-	bool bIncrementalResolveEnabled = false;
-};
\ No newline at end of file
+};
