Integrating //UE5/Dev-Cooker @ 16678003 to //UE5/Main (Zousar.Shaker-YEG-0943-Quaternary)

diff --git a/Engine/Source/Developer/IoStoreUtilities/Private/ZenFileSystemManifest.cpp b/Engine/Source/Developer/IoStoreUtilities/Private/ZenFileSystemManifest.cpp
new file mode 100644
--- /dev/null
+++ b/Engine/Source/Developer/IoStoreUtilities/Private/ZenFileSystemManifest.cpp
@@ -0,0 +1,223 @@
+// Copyright Epic Games, Inc. All Rights Reserved.
+
+#include "ZenFileSystemManifest.h"
+#include "Interfaces/ITargetPlatform.h"
+#include "HAL/PlatformFileManager.h"
+#include "Misc/App.h"
+#include "Misc/FileHelper.h"
+#include "Misc/Paths.h"
+#include "Misc/DataDrivenPlatformInfoRegistry.h"
+
+const FZenFileSystemManifest::FManifestEntry FZenFileSystemManifest::InvalidEntry;
+
+FZenFileSystemManifest::FZenFileSystemManifest(const ITargetPlatform& InTargetPlatform, FString InCookDirectory)
+	: TargetPlatform(InTargetPlatform)
+	, CookDirectory(MoveTemp(InCookDirectory))
+{
+	IPlatformFile& PlatformFile = FPlatformFileManager::Get().GetPlatformFile();
+
+	ServerRoot = PlatformFile.ConvertToAbsolutePathForExternalAppForRead(*FPaths::RootDir());
+	FPaths::NormalizeDirectoryName(ServerRoot);
+}
+
+void FZenFileSystemManifest::Generate()
+{
+	TRACE_CPUPROFILER_EVENT_SCOPE(GenerateStorageServerFileSystemManifest);
+
+	struct FFileFilter
+	{
+		bool IncludeDirectory(const TCHAR* Path) const
+		{
+			if (FCString::Stristr(Path, TEXT("/Intermediate")))
+			{
+				return false;
+			}
+			if (FCString::Stristr(Path, TEXT("/Binaries")))
+			{
+				return false;
+			}
+			if (FCString::Stristr(Path, TEXT("/Source")))
+			{
+				return false;
+			}
+			return true;
+		}
+
+		bool IncludeFile(const TCHAR* Path) const
+		{
+			if (FCString::Stristr(Path, TEXT("/Content/")))
+			{
+				const TCHAR* Extension = FCString::Strrchr(Path, '.');
+				if (Extension)
+				{
+					if (!FCString::Stricmp(Extension, TEXT(".uasset")) ||
+						!FCString::Stricmp(Extension, TEXT(".umap")) ||
+						!FCString::Stricmp(Extension, TEXT(".uexp")) ||
+						!FCString::Stricmp(Extension, TEXT(".ubulk")))
+					{
+						return false;
+					}
+				}
+			}
+			return true;
+		}
+	} Filter;
+	
+	IPlatformFile& PlatformFile = FPlatformFileManager::Get().GetPlatformFile();
+
+	FString RootDir = FPaths::RootDir();
+	FString EngineDir = FPaths::EngineDir();
+	FPaths::NormalizeDirectoryName(EngineDir);
+	FString ProjectDir = FPaths::ProjectDir();
+	FPaths::NormalizeDirectoryName(ProjectDir);
+
+	auto AddFilesFromDirectory =
+		[this, &PlatformFile, &RootDir]
+		(const FString& ClientDirectory, const FString& LocalDirectory, bool bIncludeSubdirs, FFileFilter* Filter = nullptr)
+		{
+			FString ServerRelativeDirectory = LocalDirectory;
+			FPaths::MakePathRelativeTo(ServerRelativeDirectory, *RootDir);
+			ServerRelativeDirectory = TEXT("/") + ServerRelativeDirectory;
+
+			TArray<FString> DirectoriesToVisit;
+			auto VisitorFunc =
+				[this, &DirectoriesToVisit, &RootDir, &ClientDirectory, &LocalDirectory, &ServerRelativeDirectory, bIncludeSubdirs, Filter]
+				(const TCHAR* InnerFileNameOrDirectory, bool bIsDirectory)
+				{
+					if (Filter)
+					{
+						if ((bIsDirectory && !Filter->IncludeDirectory(InnerFileNameOrDirectory)) ||
+							(!bIsDirectory && !Filter->IncludeFile(InnerFileNameOrDirectory)))
+						{
+							return true;
+						}
+					}
+					if (bIsDirectory)
+					{
+						if (bIncludeSubdirs)
+						{
+							DirectoriesToVisit.Add(InnerFileNameOrDirectory);
+						}
+					}
+					else
+					{
+						FStringView RelativePath = InnerFileNameOrDirectory;
+						RelativePath.RightChopInline(LocalDirectory.Len() + 1);
+
+						AddManifestEntry(
+							FPaths::Combine(ServerRelativeDirectory, RelativePath.GetData()),
+							FPaths::Combine(ClientDirectory, RelativePath.GetData()));
+					}
+
+					return true;
+				};
+
+			DirectoriesToVisit.Push(LocalDirectory);
+			while (!DirectoriesToVisit.IsEmpty())
+			{
+				PlatformFile.IterateDirectory(*DirectoriesToVisit.Pop(), VisitorFunc);
+			}
+		};
+
+	AddFilesFromDirectory(TEXT("/{engine}"), FPaths::Combine(CookDirectory, TEXT("Engine")), true);
+	AddFilesFromDirectory(TEXT("/{project}"), FPaths::Combine(CookDirectory, FApp::GetProjectName()), true);
+	AddFilesFromDirectory(TEXT("/{project}"), ProjectDir, false);
+	AddFilesFromDirectory(TEXT("/{engine}/Config"), FPaths::Combine(EngineDir, TEXT("Config")), true);
+	AddFilesFromDirectory(TEXT("/{project}/Config"), FPaths::Combine(ProjectDir, TEXT("Config")), true);
+	AddFilesFromDirectory(TEXT("/{engine}/Plugins"), FPaths::Combine(EngineDir, TEXT("Plugins")), true, &Filter);
+	AddFilesFromDirectory(TEXT("/{project}/Plugins"), FPaths::Combine(ProjectDir, TEXT("Plugins")), true, &Filter);
+	AddFilesFromDirectory(TEXT("/{engine}/Restricted"), FPaths::Combine(EngineDir, TEXT("Restricted")), true, &Filter);
+	AddFilesFromDirectory(TEXT("/{project}/Restricted"), FPaths::Combine(ProjectDir, TEXT("Restricted")), true, &Filter);
+	AddFilesFromDirectory(TEXT("/{engine}/Content/Internationalization"), FPaths::Combine(EngineDir, TEXT("Content"), TEXT("Internationalization")), true);
+	AddFilesFromDirectory(TEXT("/{engine}/Content/Localization"), FPaths::Combine(EngineDir, TEXT("Content"), TEXT("Localization")), true);
+	AddFilesFromDirectory(TEXT("/{project}/Content/Localization"), FPaths::Combine(ProjectDir, TEXT("Content"), TEXT("Localization")), true);
+	AddFilesFromDirectory(TEXT("/{engine}/Content/Slate"), FPaths::Combine(EngineDir, TEXT("Content"), TEXT("Slate")), true);
+	AddFilesFromDirectory(TEXT("/{project}/Content/Slate"), FPaths::Combine(ProjectDir, TEXT("Content"), TEXT("Slate")), true);
+	AddFilesFromDirectory(TEXT("/{engine}/Content/Movies"), FPaths::Combine(EngineDir, TEXT("Content"), TEXT("Movies")), true);
+	AddFilesFromDirectory(TEXT("/{project}/Content/Movies"), FPaths::Combine(ProjectDir, TEXT("Content"), TEXT("Movies")), true);
+
+	const FDataDrivenPlatformInfo& Info = FDataDrivenPlatformInfoRegistry::GetPlatformInfo(TargetPlatform.IniPlatformName());
+
+	TArray<FString> PlatformDirectories;
+	PlatformDirectories.Reserve(Info.IniParentChain.Num() + 1);
+	PlatformDirectories.Add(TargetPlatform.IniPlatformName());
+	PlatformDirectories.Append(Info.IniParentChain);
+
+	for (const FString& PlatformDirectory : PlatformDirectories)
+	{
+		AddFilesFromDirectory(FPaths::Combine(TEXT("/{engine}/Platforms"), PlatformDirectory), FPaths::Combine(EngineDir, TEXT("Platforms"), PlatformDirectory), true, &Filter);
+		AddFilesFromDirectory(FPaths::Combine(TEXT("/{project}/Platforms"), PlatformDirectory), FPaths::Combine(ProjectDir, TEXT("Platforms"), PlatformDirectory), true, &Filter);
+	}
+}
+
+const FZenFileSystemManifest::FManifestEntry& FZenFileSystemManifest::AddFile(const FString& Filename)
+{
+	FString CookedEngineDirectory = FPaths::Combine(CookDirectory, TEXT("Engine"));
+
+	auto AddEntry = [this, &Filename](const FString& ClientDirectory, const FString& LocalDirectory) -> const FManifestEntry&
+	{
+		FStringView RelativePath = Filename;
+		RelativePath.RightChopInline(LocalDirectory.Len() + 1);
+
+		FString ServerRelativeDirectory = LocalDirectory;
+		FPaths::MakePathRelativeTo(ServerRelativeDirectory, *FPaths::RootDir());
+		ServerRelativeDirectory = TEXT("/") + ServerRelativeDirectory;
+
+		return AddManifestEntry(
+			FPaths::Combine(ServerRelativeDirectory, RelativePath.GetData()),
+			FPaths::Combine(ClientDirectory, RelativePath.GetData()));
+	};
+
+	if (Filename.StartsWith(CookedEngineDirectory))
+	{
+		return AddEntry(TEXT("/{engine}"), CookedEngineDirectory);
+	}
+
+	FString CookedProjectDirectory = FPaths::Combine(CookDirectory, FApp::GetProjectName());
+	if (Filename.StartsWith(CookedProjectDirectory))
+	{
+		return AddEntry(TEXT("/{project}"), CookedProjectDirectory);
+	}
+
+	return InvalidEntry;
+}
+
+const FZenFileSystemManifest::FManifestEntry& FZenFileSystemManifest::AddManifestEntry(FString ServerPath, FString ClientPath)
+{
+	ServerPath.ReplaceInline(TEXT("\\"), TEXT("/"));
+	ClientPath.ReplaceInline(TEXT("\\"), TEXT("/"));
+
+	int32& EntryIndex = ServerPathToEntry.FindOrAdd(ServerPath, INDEX_NONE);
+
+	if (EntryIndex != INDEX_NONE)
+	{
+		return Entries[EntryIndex];
+	}
+
+	EntryIndex = Entries.Num();
+
+	FManifestEntry Entry;
+	Entry.ServerPath = MoveTemp(ServerPath);
+	Entry.ClientPath = MoveTemp(ClientPath);
+	Entry.FileId = static_cast<uint32>(EntryIndex + 1);
+
+	Entries.Add(MoveTemp(Entry));
+
+	return Entries[EntryIndex];
+}
+
+void FZenFileSystemManifest::Save(const TCHAR* Filename)
+{
+	check(Filename);
+
+	TArray<FString> CsvLines;
+	CsvLines.Add(FString::Printf(TEXT(";ServerRoot=%s, Platform=%s, CookDirectory=%s"), *ServerRoot, *TargetPlatform.PlatformName(), *CookDirectory));
+	CsvLines.Add(TEXT("FileId, ServerPath, ClientPath"));
+
+	for (const FManifestEntry& Entry : Entries)
+	{
+		CsvLines.Add(FString::Printf(TEXT("%d, %s, %s"), Entry.FileId, *Entry.ServerPath, *Entry.ClientPath));
+	}
+
+	FFileHelper::SaveStringArrayToFile(CsvLines, Filename);
+}
diff --git a/Engine/Source/Developer/IoStoreUtilities/Private/ZenFileSystemManifest.h b/Engine/Source/Developer/IoStoreUtilities/Private/ZenFileSystemManifest.h
new file mode 100644
--- /dev/null
+++ b/Engine/Source/Developer/IoStoreUtilities/Private/ZenFileSystemManifest.h
@@ -0,0 +1,49 @@
+// Copyright Epic Games, Inc. All Rights Reserved.
+
+#pragma once
+
+#include "CoreMinimal.h"
+
+class ITargetPlatform;
+class FArchive;
+
+class FZenFileSystemManifest
+{
+public:
+
+	struct FManifestEntry
+	{
+		FString ServerPath;
+		FString ClientPath;
+		uint32 FileId = ~uint32(0);
+	};
+	
+	FZenFileSystemManifest(const ITargetPlatform& InTargetPlatform, FString InCookDirectory);
+	
+	void Generate();
+
+	const FManifestEntry& AddFile(const FString& Filename);
+
+	TArrayView<const FManifestEntry> ManifestEntries() const
+	{
+		return Entries;
+	}
+
+	int32 ManifestVersion() const
+	{
+		return Entries.Num();
+	}
+
+	void Save(const TCHAR* Filename);
+
+private:
+	const FManifestEntry& AddManifestEntry(FString ServerPath, FString ClientPath);
+
+	const ITargetPlatform& TargetPlatform;
+	FString CookDirectory;
+	FString ServerRoot;
+	TMap<FString, int32> ServerPathToEntry;
+	TArray<FManifestEntry> Entries;
+	
+	static const FManifestEntry InvalidEntry;
+};
diff --git a/Engine/Source/Developer/IoStoreUtilities/Private/ZenStoreHttpClient.cpp b/Engine/Source/Developer/IoStoreUtilities/Private/ZenStoreHttpClient.cpp
new file mode 100644
--- /dev/null
+++ b/Engine/Source/Developer/IoStoreUtilities/Private/ZenStoreHttpClient.cpp
@@ -0,0 +1,1337 @@
+// Copyright Epic Games, Inc. All Rights Reserved.
+
+#include "ZenStoreHttpClient.h"
+
+#include "Serialization/CompactBinaryPackage.h"
+
+#if PLATFORM_WINDOWS
+
+#if PLATFORM_WINDOWS || PLATFORM_HOLOLENS
+#include "Windows/WindowsHWrapper.h"
+#include "Windows/AllowWindowsPlatformTypes.h"
+#endif
+#include "curl/curl.h"
+#if PLATFORM_WINDOWS || PLATFORM_HOLOLENS
+#include "Windows/HideWindowsPlatformTypes.h"
+#endif
+
+#include "Async/Async.h"
+#include "HAL/PlatformFileManager.h"
+#include "IO/IoHash.h"
+#include "Misc/ScopeLock.h"
+#include "Misc/StringBuilder.h"
+#include "ProfilingDebugging/CpuProfilerTrace.h"
+#include "Serialization/BufferArchive.h"
+#include "Serialization/CompactBinaryWriter.h"
+#include "Serialization/LargeMemoryWriter.h"
+
+DEFINE_LOG_CATEGORY_STATIC(LogZenStore, Log, All);
+
+namespace UE { namespace Zen {
+
+#define	UE_ZENHTTP_BACKEND_WAIT_INTERVAL			0.01f
+#define	UE_ZENHTTP_HTTP_REQUEST_TIMEOUT_SECONDS		30L
+#define UE_ZENHTTP_HTTP_REQUEST_TIMOUT_ENABLED		1
+#define UE_ZENHTTP_HTTP_DEBUG						0
+#define UE_ZENHTTP_MAX_FAILED_LOGIN_ATTEMPTS		16
+#define UE_ZENHTTP_MAX_ATTEMPTS						4
+#define UE_ZENHTTP_MAX_BUFFER_RESERVE				104857600u
+
+/**
+ * Minimal HTTP request type wrapping CURL without the need for managers. This request
+ * is written to allow reuse of request objects, in order to allow connections to be reused.
+ *
+ * CURL has a global library initialization (curl_global_init). We rely on this happening in
+ * the Online/HTTP library which is a dependency on this module.
+ */
+class FRequest
+{
+public:
+	/**
+	 * Supported request verbs
+	 */
+	enum RequestVerb
+	{
+		Get,
+		Put,
+		Post,
+		Delete,
+		Head
+	};
+
+	/**
+	 * Convenience result type interpreted from HTTP response code.
+	 */
+	enum Result
+	{
+		Success,
+		Failed,
+		FailedTimeout
+	};
+
+	FRequest(FStringView InDomain, bool bInLogErrors)
+		: bLogErrors(bInLogErrors)
+		, Domain(InDomain)
+	{
+		Curl = curl_easy_init();
+		Reset();
+	}
+
+	~FRequest()
+	{
+		curl_easy_cleanup(Curl);
+	}
+
+	/**
+	 * Resets all options on the request except those that should always be set.
+	 */
+	void Reset()
+	{
+		Headers.Reset();
+		ResponseHeader.Reset();
+		ResponseBuffer.Reset();
+		ResponseCode = 0;
+		ReadDataView = TConstArrayView64<uint8>();
+		WriteDataBufferPtr = nullptr;
+		WriteHeaderBufferPtr = nullptr;
+		BytesSent = 0;
+		BytesReceived = 0;
+		CurlResult = CURL_LAST;
+
+		curl_easy_reset(Curl);
+
+		// Options that are always set for all connections.
+#if UE_ZENHTTP_HTTP_REQUEST_TIMOUT_ENABLED
+		curl_easy_setopt(Curl, CURLOPT_CONNECTTIMEOUT, UE_ZENHTTP_HTTP_REQUEST_TIMEOUT_SECONDS);
+#endif
+		curl_easy_setopt(Curl, CURLOPT_FOLLOWLOCATION, 1L);
+		curl_easy_setopt(Curl, CURLOPT_NOSIGNAL, 1L);
+		// Response functions
+		curl_easy_setopt(Curl, CURLOPT_HEADERDATA, this);
+		curl_easy_setopt(Curl, CURLOPT_HEADERFUNCTION, &FRequest::StaticWriteHeaderFn);
+		curl_easy_setopt(Curl, CURLOPT_WRITEDATA, this);
+		curl_easy_setopt(Curl, CURLOPT_WRITEFUNCTION, StaticWriteBodyFn);
+		// Rewind method, handle special error case where request need to rewind data stream
+		curl_easy_setopt(Curl, CURLOPT_SEEKFUNCTION, StaticSeekFn);
+		curl_easy_setopt(Curl, CURLOPT_SEEKDATA, this);
+		// Debug hooks
+#if UE_ZENHTTP_HTTP_DEBUG
+		curl_easy_setopt(Curl, CURLOPT_DEBUGDATA, this);
+		curl_easy_setopt(Curl, CURLOPT_DEBUGFUNCTION, StaticDebugCallback);
+		curl_easy_setopt(Curl, CURLOPT_VERBOSE, 1L);
+#endif
+	}
+
+	/** Returns the HTTP response code.*/
+	const int64 GetResponseCode() const
+	{
+		return ResponseCode;
+	}
+
+	/** Returns the number of bytes received this request (headers withstanding). */
+	const size_t GetBytesReceived() const
+	{
+		return BytesReceived;
+	}
+
+	/** Returns the number of bytes sent during this request (headers withstanding). */
+	const size_t  GetBytesSent() const
+	{
+		return BytesSent;
+	}
+
+	template<RequestVerb V>
+	Result PerformBlockingUpload(FStringView Uri, FCbObjectView Obj)
+	{
+		FLargeMemoryWriter Out;
+		Obj.CopyTo(Out);
+
+		TConstArrayView64<uint8> Payload { Out.GetData(), Out.TotalSize()};
+		return PerformBlockingUpload<V>(Uri, Payload);
+	}
+
+	Result PerformBlockingPost(FStringView Uri, FCbObjectView Obj)
+	{
+		FLargeMemoryWriter Out;
+		Obj.CopyTo(Out);
+
+		return PerformBlockingPost(Uri, Out.GetView());
+	}
+
+	Result PerformBlockingPost(FStringView Uri, FMemoryView Payload)
+	{
+		uint64 ContentLength = 0u;
+
+		curl_easy_setopt(Curl, CURLOPT_POST, 1L);
+		curl_easy_setopt(Curl, CURLOPT_INFILESIZE, Payload.GetSize());
+		curl_easy_setopt(Curl, CURLOPT_READDATA, this);
+		curl_easy_setopt(Curl, CURLOPT_READFUNCTION, StaticReadFn);
+
+		// TODO: proper Content-Type: header
+
+		ContentLength = Payload.GetSize();
+		ReadDataView = { reinterpret_cast<const uint8*>(Payload.GetData()), (int64)Payload.GetSize() };
+
+		return PerformBlocking(Uri, Post, ContentLength);
+	}
+
+	/**
+	 * Upload buffer using the request, using either "Put" or "Post" verbs.
+	 * @param Uri Url to use.
+	 * @param Buffer Data to upload
+	 * @return Result of the request
+	 */
+	template<RequestVerb V>
+	Result PerformBlockingUpload(FStringView Uri, TConstArrayView64<uint8> Buffer)
+	{
+		static_assert(V == Put || V == Post, "Upload should use either Put or Post verbs.");
+
+		uint64 ContentLength = 0u;
+
+		if (V == Put)
+		{
+			curl_easy_setopt(Curl, CURLOPT_UPLOAD, 1L);
+			curl_easy_setopt(Curl, CURLOPT_INFILESIZE, Buffer.Num());
+			curl_easy_setopt(Curl, CURLOPT_READDATA, this);
+			curl_easy_setopt(Curl, CURLOPT_READFUNCTION, StaticReadFn);
+			Headers.Add(FString(TEXT("Content-Type: application/octet-stream")));
+			ContentLength = Buffer.Num();
+			ReadDataView = Buffer;
+		}
+		else if (V == Post)
+		{
+			curl_easy_setopt(Curl, CURLOPT_POST, 1L);
+			curl_easy_setopt(Curl, CURLOPT_INFILESIZE, Buffer.Num());
+			curl_easy_setopt(Curl, CURLOPT_READDATA, this);
+			curl_easy_setopt(Curl, CURLOPT_READFUNCTION, StaticReadFn);
+			
+			//
+			// TODO: proper Content-Type: header
+			//Headers.Add(V == Post ? FString(TEXT("Content-Type: application/x-www-form-urlencoded")) : FString(TEXT("Content-Type: application/json")));
+
+			ContentLength = Buffer.Num();
+			ReadDataView = Buffer;
+		}
+
+		return PerformBlocking(Uri, V, ContentLength);
+	}
+
+	/**
+	 * Download an url into a buffer using the request.
+	 * @param Uri Url to use.
+	 * @param Buffer Optional buffer where data should be downloaded to. If empty downloaded data will
+	 * be stored in an internal buffer and accessed GetResponse* methods.
+	 * @return Result of the request
+	 */
+	Result PerformBlockingDownload(FStringView Uri, TArray64<uint8>* Buffer)
+	{
+		curl_easy_setopt(Curl, CURLOPT_HTTPGET, 1L);
+		WriteDataBufferPtr = Buffer;
+
+		return PerformBlocking(Uri, Get, 0u);
+	}
+
+	/**
+	 * Query an url using the request. Queries can use either "Head" or "Delete" verbs.
+	 * @param Uri Url to use.
+	 * @return Result of the request
+	 */
+	template<RequestVerb V>
+	Result PerformBlockingQuery(FStringView Uri)
+	{
+		static_assert(V == Head || V == Delete, "Queries should use either Head or Delete verbs.");
+
+		if (V == Delete)
+		{
+			curl_easy_setopt(Curl, CURLOPT_CUSTOMREQUEST, "DELETE");
+		}
+		else if (V == Head)
+		{
+			curl_easy_setopt(Curl, CURLOPT_NOBODY, 1L);
+		}
+
+		return PerformBlocking(Uri, V, 0u);
+	}
+
+	/**
+	 * Set a header to send with the request.
+	 */
+	void SetHeader(const TCHAR* Header, const TCHAR* Value)
+	{
+		check(CurlResult == CURL_LAST); // Cannot set header after request is sent
+		Headers.Add(FString::Printf(TEXT("%s: %s"), Header, Value));
+	}
+
+	/**
+	 * Attempts to find the header from the response. Returns false if header is not present.
+	 */
+	bool GetHeader(const ANSICHAR* Header, FString& OutValue) const
+	{
+		check(CurlResult != CURL_LAST);  // Cannot query headers before request is sent
+
+		const ANSICHAR* HeadersBuffer = (const ANSICHAR*)ResponseHeader.GetData();
+		size_t HeaderLen = strlen(Header);
+
+		// Find the header key in the (ANSI) response buffer. If not found we can exist immediately
+		if (const ANSICHAR* Found = strstr(HeadersBuffer, Header))
+		{
+			const ANSICHAR* Linebreak = strchr(Found, '\r');
+			const ANSICHAR* ValueStart = Found + HeaderLen + 2; //colon and space
+			const size_t ValueSize = Linebreak - ValueStart;
+			FUTF8ToTCHAR TCHARData(ValueStart, ValueSize);
+			OutValue = FString(TCHARData.Length(), TCHARData.Get());
+			return true;
+		}
+		return false;
+	}
+
+	/**
+	 * Returns the response buffer. Note that is the request is performed
+	 * with an external buffer as target buffer this string will be empty.
+	 */
+	const TArray64<uint8>& GetResponseBuffer() const
+	{
+		return ResponseBuffer;
+	}
+
+	/**
+	 * Returns the response buffer as a string. Note that is the request is performed
+	 * with an external buffer as target buffer this string will be empty.
+	 */
+	FString GetResponseAsString() const
+	{
+		return GetAnsiBufferAsString(ResponseBuffer);
+	}
+
+	/**
+	 * Returns the response header as a string.
+	 */
+	FString GetResponseHeaderAsString()
+	{
+		return GetAnsiBufferAsString(ResponseHeader);
+	}
+
+	/** Will return true if the response code is considered a success */
+	static bool IsSuccessResponse(long ResponseCode)
+	{
+		// We consider anything in the 1XX or 2XX range a success
+		return ResponseCode >= 100 && ResponseCode < 300;
+	}
+
+private:
+
+	CURL* Curl;
+	CURLcode				CurlResult;
+	long					ResponseCode;
+	size_t					BytesSent;
+	size_t					BytesReceived;
+	bool					bLogErrors;
+
+	TConstArrayView64<uint8>	ReadDataView;
+	TArray64<uint8>*			WriteDataBufferPtr;
+	TArray64<uint8>*			WriteHeaderBufferPtr;
+
+	TArray64<uint8>			ResponseHeader;
+	TArray64<uint8>			ResponseBuffer;
+	TArray<FString>			Headers;
+	FString					Domain;
+
+	/**
+	 * Performs the request, blocking until finished.
+	 * @param Uri Address on the domain to query
+	 * @param Verb HTTP verb to use
+	 * @param Buffer Optional buffer to directly receive the result of the request.
+	 * If unset the response body will be stored in the request.
+	 */
+	Result PerformBlocking(FStringView Uri, RequestVerb Verb, uint64 ContentLength)
+	{
+		static const char* CommonHeaders[] = {
+			"User-Agent: UE",
+			nullptr
+		};
+
+		TRACE_CPUPROFILER_EVENT_SCOPE(ZenStore_CurlPerform);
+
+		// Setup request options
+		TStringBuilder<128> Url;
+		Url << Domain << Uri;
+		
+		curl_easy_setopt(Curl, CURLOPT_URL, TCHAR_TO_ANSI(*Url));
+
+		// Setup response header buffer. If caller has not setup a response data buffer, use interal.
+		WriteHeaderBufferPtr = &ResponseHeader;
+		if (WriteDataBufferPtr == nullptr)
+		{
+			WriteDataBufferPtr = &ResponseBuffer;
+		}
+
+		// Content-Length should always be set
+		Headers.Add(FString::Printf(TEXT("Content-Length: %llu"), ContentLength));
+
+		// Build headers list
+		curl_slist* CurlHeaders = nullptr;
+		// Add common headers
+		for (uint8 i = 0; CommonHeaders[i] != nullptr; ++i)
+		{
+			CurlHeaders = curl_slist_append(CurlHeaders, CommonHeaders[i]);
+		}
+		// Setup added headers
+		for (const FString& Header : Headers)
+		{
+			CurlHeaders = curl_slist_append(CurlHeaders, TCHAR_TO_ANSI(*Header));
+		}
+		curl_easy_setopt(Curl, CURLOPT_HTTPHEADER, CurlHeaders);
+
+		// Shots fired!
+		CurlResult = curl_easy_perform(Curl);
+
+		// Get response code
+		bool bRedirected = false;
+		if (CURLE_OK == curl_easy_getinfo(Curl, CURLINFO_RESPONSE_CODE, &ResponseCode))
+		{
+			bRedirected = (ResponseCode >= 300 && ResponseCode < 400);
+		}
+
+		LogResult(CurlResult, Uri, Verb);
+
+		// Clean up
+		curl_slist_free_all(CurlHeaders);
+
+		return CurlResult == CURLE_OK ? Success : Failed;
+	}
+
+	void LogResult(CURLcode Result, FStringView Uri, RequestVerb Verb) const
+	{
+		if (Result == CURLE_OK)
+		{
+			bool bSuccess = false;
+			const TCHAR* VerbStr = nullptr;
+			FString AdditionalInfo;
+
+			switch (Verb)
+			{
+			case Head:
+				bSuccess = (ResponseCode == 400 || IsSuccessResponse(ResponseCode));
+				VerbStr = TEXT("querying");
+				break;
+			case Get:
+				bSuccess = (ResponseCode == 400 || IsSuccessResponse(ResponseCode));
+				VerbStr = TEXT("fetching");
+				AdditionalInfo = FString::Printf(TEXT("Received: %d bytes."), BytesReceived);
+				break;
+			case Put:
+				bSuccess = IsSuccessResponse(ResponseCode);
+				VerbStr = TEXT("updating");
+				AdditionalInfo = FString::Printf(TEXT("Sent: %d bytes."), BytesSent);
+				break;
+			case Post:
+				bSuccess = IsSuccessResponse(ResponseCode);
+				VerbStr = TEXT("posting");
+				break;
+			case Delete:
+				bSuccess = IsSuccessResponse(ResponseCode);
+				VerbStr = TEXT("deleting");
+				break;
+			}
+
+			if (bSuccess)
+			{
+				UE_LOG(
+					LogZenStore,
+					Verbose,
+					TEXT("Finished %s Zen store HTTP operation (response %d) from %s. %s"),
+					VerbStr,
+					ResponseCode,
+					*FString(Uri),
+					*AdditionalInfo
+				);
+			}
+			else if (bLogErrors)
+			{
+				// Print the response body if we got one, otherwise print header.
+				FString Response = GetAnsiBufferAsString(ResponseBuffer.Num() > 0 ? ResponseBuffer : ResponseHeader);
+				Response.ReplaceCharInline('\n', ' ');
+				Response.ReplaceCharInline('\r', ' ');
+				// Dont log access denied as error, since tokens can expire mid session
+				if (ResponseCode == 401)
+				{
+					UE_LOG(
+						LogZenStore,
+						Verbose,
+						TEXT("Failed %s Zen store HTTP operation (response %d) from %s. Response: %s"),
+						VerbStr,
+						ResponseCode,
+						*FString(Uri),
+						*Response
+					);
+				}
+				else if (ResponseCode == 404)
+				{
+					// 404 is not an error, it is an expected and valid result in most cases
+				}
+				else
+				{
+					UE_LOG(
+						LogZenStore,
+						Display,
+						TEXT("Failed %s Zen store HTTP operation (response %d) from %s. Response: %s"),
+						VerbStr,
+						ResponseCode,
+						*FString(Uri),
+						*Response
+					);
+				}
+			}
+		}
+		else if (bLogErrors)
+		{
+			UE_LOG(
+				LogZenStore,
+				Display,
+				TEXT("Error while connecting to %s: %s"),
+				*Domain,
+				ANSI_TO_TCHAR(curl_easy_strerror(Result))
+			);
+		}
+	}
+
+	FString GetAnsiBufferAsString(const TArray64<uint8>& Buffer) const
+	{
+		// Content is NOT null-terminated; we need to specify lengths here
+		FUTF8ToTCHAR TCHARData(reinterpret_cast<const ANSICHAR*>(Buffer.GetData()), Buffer.Num());
+		return FString(TCHARData.Length(), TCHARData.Get());
+	}
+
+	static size_t StaticDebugCallback(CURL* Handle, curl_infotype DebugInfoType, char* DebugInfo, size_t DebugInfoSize, void* UserData)
+	{
+		FRequest* Request = static_cast<FRequest*>(UserData);
+
+		switch (DebugInfoType)
+		{
+		case CURLINFO_TEXT:
+		{
+			// Truncate at 1023 characters. This is just an arbitrary number based on a buffer size seen in
+			// the libcurl code.
+			DebugInfoSize = FMath::Min(DebugInfoSize, (size_t)1023);
+
+			// Calculate the actual length of the string due to incorrect use of snprintf() in lib/vtls/openssl.c.
+			char* FoundNulPtr = (char*)memchr(DebugInfo, 0, DebugInfoSize);
+			int CalculatedSize = FoundNulPtr != nullptr ? FoundNulPtr - DebugInfo : DebugInfoSize;
+
+			auto ConvertedString = StringCast<TCHAR>(static_cast<const ANSICHAR*>(DebugInfo), CalculatedSize);
+			FString DebugText(ConvertedString.Length(), ConvertedString.Get());
+			DebugText.ReplaceInline(TEXT("\n"), TEXT(""), ESearchCase::CaseSensitive);
+			DebugText.ReplaceInline(TEXT("\r"), TEXT(""), ESearchCase::CaseSensitive);
+			UE_LOG(LogZenStore, VeryVerbose, TEXT("%p: '%s'"), Request, *DebugText);
+		}
+		break;
+
+		case CURLINFO_HEADER_IN:
+			UE_LOG(LogZenStore, VeryVerbose, TEXT("%p: Received header (%d bytes)"), Request, DebugInfoSize);
+			break;
+
+		case CURLINFO_DATA_IN:
+			UE_LOG(LogZenStore, VeryVerbose, TEXT("%p: Received data (%d bytes)"), Request, DebugInfoSize);
+			break;
+
+		case CURLINFO_DATA_OUT:
+			UE_LOG(LogZenStore, VeryVerbose, TEXT("%p: Sent data (%d bytes)"), Request, DebugInfoSize);
+			break;
+
+		case CURLINFO_SSL_DATA_IN:
+			UE_LOG(LogZenStore, VeryVerbose, TEXT("%p: Received SSL data (%d bytes)"), Request, DebugInfoSize);
+			break;
+
+		case CURLINFO_SSL_DATA_OUT:
+			UE_LOG(LogZenStore, VeryVerbose, TEXT("%p: Sent SSL data (%d bytes)"), Request, DebugInfoSize);
+			break;
+		}
+
+		return 0;
+	}
+
+	static size_t StaticReadFn(void* Ptr, size_t SizeInBlocks, size_t BlockSizeInBytes, void* UserData)
+	{
+		FRequest* Request = static_cast<FRequest*>(UserData);
+		TConstArrayView64<uint8>& ReadDataView = Request->ReadDataView;
+
+		const size_t Offset = Request->BytesSent;
+		const size_t ReadSize = FMath::Min((size_t)ReadDataView.Num() - Offset, SizeInBlocks * BlockSizeInBytes);
+		check((size_t)ReadDataView.Num() >= Offset + ReadSize);
+
+		FMemory::Memcpy(Ptr, ReadDataView.GetData() + Offset, ReadSize);
+		Request->BytesSent += ReadSize;
+		return ReadSize;
+
+		return 0;
+	}
+
+	static size_t StaticWriteHeaderFn(void* Ptr, size_t SizeInBlocks, size_t BlockSizeInBytes, void* UserData)
+	{
+		FRequest* Request = static_cast<FRequest*>(UserData);
+		const size_t WriteSize = SizeInBlocks * BlockSizeInBytes;
+		TArray64<uint8>* WriteHeaderBufferPtr = Request->WriteHeaderBufferPtr;
+		if (WriteHeaderBufferPtr && WriteSize > 0)
+		{
+			const size_t CurrentBufferLength = WriteHeaderBufferPtr->Num();
+			if (CurrentBufferLength > 0)
+			{
+				// Remove the previous zero termination
+				(*WriteHeaderBufferPtr)[CurrentBufferLength - 1] = ' ';
+			}
+
+			// Write the header
+			WriteHeaderBufferPtr->Append((const uint8*)Ptr, WriteSize + 1);
+			(*WriteHeaderBufferPtr)[WriteHeaderBufferPtr->Num() - 1] = 0; // Zero terminate string
+			return WriteSize;
+		}
+		return 0;
+	}
+
+	static size_t StaticWriteBodyFn(void* Ptr, size_t SizeInBlocks, size_t BlockSizeInBytes, void* UserData)
+	{
+		FRequest* Request = static_cast<FRequest*>(UserData);
+		const size_t WriteSize = SizeInBlocks * BlockSizeInBytes;
+		TArray64<uint8>* WriteDataBufferPtr = Request->WriteDataBufferPtr;
+
+		if (WriteDataBufferPtr && WriteSize > 0)
+		{
+			// If this is the first part of the body being received, try to reserve 
+			// memory if content length is defined in the header.
+			if (Request->BytesReceived == 0 && Request->WriteHeaderBufferPtr)
+			{
+				static const ANSICHAR* ContentLengthHeaderStr = "Content-Length: ";
+				const ANSICHAR* Header = (const ANSICHAR*)Request->WriteHeaderBufferPtr->GetData();
+
+				if (const ANSICHAR* ContentLengthHeader = FCStringAnsi::Strstr(Header, ContentLengthHeaderStr))
+				{
+					size_t ContentLength = (size_t)FCStringAnsi::Atoi64(ContentLengthHeader + strlen(ContentLengthHeaderStr));
+					if (ContentLength > 0u && ContentLength < UE_ZENHTTP_MAX_BUFFER_RESERVE)
+					{
+						WriteDataBufferPtr->Reserve(ContentLength);
+					}
+				}
+			}
+
+			// Write to the target buffer
+			WriteDataBufferPtr->Append((const uint8*)Ptr, WriteSize);
+			Request->BytesReceived += WriteSize;
+			return WriteSize;
+		}
+
+		return 0;
+	}
+
+	static size_t StaticSeekFn(void* UserData, curl_off_t Offset, int Origin)
+	{
+		FRequest* Request = static_cast<FRequest*>(UserData);
+		size_t NewPosition = 0;
+
+		switch (Origin)
+		{
+		case SEEK_SET: NewPosition = Offset; break;
+		case SEEK_CUR: NewPosition = Request->BytesSent + Offset; break;
+		case SEEK_END: NewPosition = Request->ReadDataView.Num() + Offset; break;
+		}
+
+		// Make sure we don't seek outside of the buffer
+		if (NewPosition < 0 || NewPosition >= (size_t)Request->ReadDataView.Num())
+		{
+			return CURL_SEEKFUNC_FAIL;
+		}
+
+		// Update the used offset
+		Request->BytesSent = NewPosition;
+		return CURL_SEEKFUNC_OK;
+	}
+
+};
+
+//----------------------------------------------------------------------------------------------------------
+// Request pool
+//----------------------------------------------------------------------------------------------------------
+
+/**
+ * Pool that manages a fixed set of requests. Users are required to release requests that have been
+ * acquired. Usable with \ref FScopedRequestPtr which handles this automatically.
+ */
+struct FRequestPool
+{
+	FRequestPool(FStringView InServiceUrl, uint32 PoolSize)
+	{
+		Pool.AddUninitialized(PoolSize);
+		for (uint8 i = 0; i < Pool.Num(); ++i)
+		{
+			Pool[i].Usage = 0u;
+			Pool[i].Request = new FRequest(InServiceUrl, true);
+		}
+	}
+
+	~FRequestPool()
+	{
+		for (uint8 i = 0; i < Pool.Num(); ++i)
+		{
+			// No requests should be in use by now.
+			check(Pool[i].Usage.load(std::memory_order_acquire) == 0u);
+			delete Pool[i].Request;
+		}
+	}
+
+	/**
+	 * Attempts to get a request is free. Once a request has been returned it is
+	 * "owned by the caller and need to release it to the pool when work has been completed.
+	 * @return Usable request instance if one is available, otherwise null.
+	 */
+	FRequest* GetFreeRequest()
+	{
+		for (uint8 i = 0; i < Pool.Num(); ++i)
+		{
+			if (!Pool[i].Usage.load(std::memory_order_relaxed))
+			{
+				uint8 Expected = 0u;
+				if (Pool[i].Usage.compare_exchange_strong(Expected, 1u))
+				{
+					Pool[i].Request->Reset();
+					return Pool[i].Request;
+				}
+			}
+		}
+		return nullptr;
+	}
+
+	/**
+	 * Block until a request is free. Once a request has been returned it is
+	 * "owned by the caller and need to release it to the pool when work has been completed.
+	 * @return Usable request instance.
+	 */
+	FRequest* WaitForFreeRequest()
+	{
+		TRACE_CPUPROFILER_EVENT_SCOPE(ZenStore_WaitForConnPool);
+		FRequest* Request = nullptr;
+		while (true)
+		{
+			Request = GetFreeRequest();
+			if (Request != nullptr)
+				break;
+			FPlatformProcess::Sleep(UE_ZENHTTP_BACKEND_WAIT_INTERVAL);
+		}
+		return Request;
+	}
+
+	/**
+	 * Release request to the pool.
+	 * @param Request Request that should be freed. Note that any buffer owned by the request can now be reset.
+	 */
+	void ReleaseRequestToPool(FRequest* Request)
+	{
+		for (uint8 i = 0; i < Pool.Num(); ++i)
+		{
+			if (Pool[i].Request == Request)
+			{
+
+				Pool[i].Usage--;
+				return;
+			}
+		}
+		check(false);
+	}
+
+	/**
+	 * While holding a request, make it shared across many users.
+	 */
+	void MakeRequestShared(FRequest* Request, uint8 Users)
+	{
+		check(Users != 0);
+		for (uint8 i = 0; i < Pool.Num(); ++i)
+		{
+			if (Pool[i].Request == Request)
+			{
+
+				Pool[i].Usage = Users;
+				return;
+			}
+		}
+		check(false);
+	}
+
+private:
+
+	struct FEntry
+	{
+		std::atomic<uint8> Usage;
+		FRequest* Request;
+	};
+
+	TArray<FEntry> Pool;
+
+	FRequestPool() = delete;
+};
+
+//----------------------------------------------------------------------------------------------------------
+// FScopedRequestPtr
+//----------------------------------------------------------------------------------------------------------
+
+/**
+ * Utility class to manage requesting and releasing requests from the \ref FRequestPool.
+ */
+struct FScopedRequestPtr
+{
+public:
+	FScopedRequestPtr(FRequestPool* InPool)
+		: Request(InPool->WaitForFreeRequest())
+		, Pool(InPool)
+	{}
+
+	~FScopedRequestPtr()
+	{
+		Pool->ReleaseRequestToPool(Request);
+	}
+
+	bool IsValid() const
+	{
+		return Request != nullptr;
+	}
+
+	FRequest* Get() const
+	{
+		check(IsValid());
+		return Request;
+	}
+
+	FRequest* operator->()
+	{
+		check(IsValid());
+		return Request;
+	}
+
+private:
+	FRequest* Request;
+	FRequestPool* Pool;
+};
+
+} /* Zen */ 
+
+FZenStoreHttpClient::FZenStoreHttpClient(const FStringView InHostName, uint16 InPort)
+: HostName(InHostName)
+, Port(InPort)
+{
+	TStringBuilder<64> Uri;
+	Uri.AppendAnsi("http://");
+	Uri.Append(InHostName);
+	Uri.AppendAnsi(":");
+	Uri << InPort;
+
+	RequestPool = MakeUnique<Zen::FRequestPool>(Uri, 32);
+}
+
+FZenStoreHttpClient::~FZenStoreHttpClient()
+{
+}
+
+void 
+FZenStoreHttpClient::Initialize(FStringView InProjectId, 
+	FStringView InOplogId, 
+	FStringView ServerRoot,
+	FStringView EngineRoot,
+	FStringView ProjectRoot,
+	bool		IsCleanBuild)
+{
+	TRACE_CPUPROFILER_EVENT_SCOPE(ZenStoreHttp_Initialize);
+
+	UE_LOG(LogZenStore, Display, TEXT("Establishing oplog %s / %s"), *FString(InProjectId), *FString(InOplogId));
+
+	// Establish project
+
+	{
+		UE::Zen::FScopedRequestPtr Request(RequestPool.Get());
+
+		TStringBuilder<128> ProjectUri;
+		ProjectUri << "/prj/" << InProjectId;
+		TArray64<uint8> GetBuffer;
+		UE::Zen::FRequest::Result Res = Request->PerformBlockingDownload(ProjectUri, &GetBuffer);
+
+		// TODO: how to handle failure here? This is probably the most likely point of failure
+		// if the service is not up or not responding
+
+		if (Res == Zen::FRequest::Success && Request->GetResponseCode() == 200)
+		{
+			UE_LOG(LogZenStore, Display, TEXT("Zen project '%s' already exists"), *FString(InProjectId));
+		}
+		else
+		{
+			Request->Reset();
+
+			FCbWriter ProjInfo;
+			ProjInfo.BeginObject();
+			ProjInfo << "id" << InProjectId;
+			ProjInfo << "root" << ServerRoot;
+			ProjInfo << "engine" << EngineRoot;
+			ProjInfo << "project" << ProjectRoot;
+			ProjInfo.EndObject();
+
+			Res = Request->PerformBlockingPost(ProjectUri, ProjInfo.Save().AsObject());
+
+			if (Res != Zen::FRequest::Success)
+			{
+				UE_LOG(LogZenStore, Error, TEXT("Zen project '%s' creation FAILED"), *FString(InProjectId));
+
+				// TODO: how to recover / handle this?
+			}
+			else if (Request->GetResponseCode() == 201)
+			{
+				UE_LOG(LogZenStore, Display, TEXT("Zen project '%s' created"), *FString(InProjectId));
+			}
+			else
+			{
+				UE_LOG(LogZenStore, Warning, TEXT("Zen project '%s' creation returned success but not HTTP 201"), *FString(InProjectId));
+			}
+		}
+	}
+
+	// Establish oplog
+
+	{
+		UE::Zen::FScopedRequestPtr Request(RequestPool.Get());
+
+		TStringBuilder<128> OplogUri;
+		OplogUri << "/prj/" << InProjectId << "/oplog/" << InOplogId;
+
+		OplogPath = OplogUri;
+
+		if (IsCleanBuild)
+		{
+			UE_LOG(LogZenStore, Display, TEXT("Deleting oplog '%s'/'%s' if it exists"), *FString(InProjectId), *FString(InOplogId));
+			Request->PerformBlockingQuery<Zen::FRequest::Delete>(OplogUri);
+			Request->Reset();
+		}
+
+		TArray64<uint8> GetBuffer;
+		UE::Zen::FRequest::Result Res = Request->PerformBlockingDownload(OplogUri, &GetBuffer);
+		FCbObjectView OplogInfo;
+
+		if (Res == Zen::FRequest::Success && Request->GetResponseCode() == 200)
+		{
+			UE_LOG(LogZenStore, Display, TEXT("Zen oplog '%s'/'%s' already exists"), *FString(InProjectId), *FString(InOplogId));
+
+			OplogInfo = FCbObjectView(GetBuffer.GetData());
+		}
+		else
+		{
+			TConstArrayView64<uint8> Payload;
+
+			Request->Reset();
+			Res = Request->PerformBlockingUpload<Zen::FRequest::Post>(OplogUri, Payload);
+
+			if (Res != Zen::FRequest::Success)
+			{
+				UE_LOG(LogZenStore, Error, TEXT("Zen oplog '%s'/'%s' creation FAILED"), *FString(InProjectId), *FString(InOplogId));
+
+				// TODO: how to recover / handle this?
+			}
+			else if (Request->GetResponseCode() == 201)
+			{
+				UE_LOG(LogZenStore, Display, TEXT("Zen oplog '%s'/'%s' created"), *FString(InProjectId), *FString(InOplogId));
+			}
+			else
+			{
+				UE_LOG(LogZenStore, Warning, TEXT("Zen oplog '%s'/'%s' creation returned success but not HTTP 201"), *FString(InProjectId), *FString(InOplogId));
+			}
+
+			// Issue another GET to retrieve information
+
+			GetBuffer.Reset();
+			Request->Reset();
+			Res = Request->PerformBlockingDownload(OplogUri, &GetBuffer);
+
+			if (Res == Zen::FRequest::Success && Request->GetResponseCode() == 200)
+			{
+				OplogInfo = FCbObjectView(GetBuffer.GetData());
+			}
+		}
+
+		TempDirPath = FUTF8ToTCHAR(OplogInfo["tempdir"].AsString());
+	}
+
+	{
+		TStringBuilder<128> OplogUri;
+		OplogUri << "/prj/" << InProjectId << "/oplog/" << InOplogId << "/new";
+
+		OplogNewEntryPath = OplogUri;
+	}
+
+	OplogPrepNewEntryPath = TStringBuilder<128>().AppendAnsi("/prj/").Append(InProjectId).AppendAnsi("/oplog/").Append(InOplogId).AppendAnsi("/prep");
+
+	bAllowRead = true;
+	bAllowEdit = true;
+}
+
+void FZenStoreHttpClient::InitializeReadOnly(FStringView InProjectId, FStringView InOplogId)
+{
+	TRACE_CPUPROFILER_EVENT_SCOPE(ZenStoreHttp_InitializeReadOnly);
+
+	UE_LOG(LogZenStore, Display, TEXT("Establishing oplog %s / %s"), *FString(InProjectId), *FString(InOplogId));
+
+	// Establish project
+
+	{
+		UE::Zen::FScopedRequestPtr Request(RequestPool.Get());
+
+		TStringBuilder<128> ProjectUri;
+		ProjectUri << "/prj/" << InProjectId;
+		TArray64<uint8> GetBuffer;
+		UE::Zen::FRequest::Result Res = Request->PerformBlockingDownload(ProjectUri, &GetBuffer);
+
+		// TODO: how to handle failure here? This is probably the most likely point of failure
+		// if the service is not up or not responding
+
+		if (Res != Zen::FRequest::Success || Request->GetResponseCode() != 200)
+		{
+			UE_LOG(LogZenStore, Fatal, TEXT("Zen project '%s' not found"), *FString(InProjectId));
+		}
+	}
+
+	// Establish oplog
+
+	{
+		UE::Zen::FScopedRequestPtr Request(RequestPool.Get());
+
+		TStringBuilder<128> OplogUri;
+		OplogUri << "/prj/" << InProjectId << "/oplog/" << InOplogId;
+
+		OplogPath = OplogUri;
+
+		TArray64<uint8> GetBuffer;
+		UE::Zen::FRequest::Result Res = Request->PerformBlockingDownload(OplogUri, &GetBuffer);
+
+		if (Res != Zen::FRequest::Success || Request->GetResponseCode() != 200)
+		{
+			UE_LOG(LogZenStore, Fatal, TEXT("Zen oplog '%s'/'%s' not found"), *FString(InProjectId), *FString(InOplogId));
+		}
+	}
+
+	bAllowRead = true;
+}
+
+static std::atomic<uint32> GOpCounter;
+
+TFuture<TIoStatusOr<uint64>> FZenStoreHttpClient::AppendOp(FCbPackage OpEntry)
+{
+	check(bAllowEdit);
+
+	TRACE_CPUPROFILER_EVENT_SCOPE(ZenStoreHttp_AppendOp);
+
+	return Async(EAsyncExecution::LargeThreadPool, [this, OpEntry]
+	{
+		TRACE_CPUPROFILER_EVENT_SCOPE(Zen_AppendOp_Async);
+		FLargeMemoryWriter SerializedPackage;
+
+		const uint32 Salt = ++GOpCounter;
+		bool IsUsingTempFiles = false;
+
+		if (TempDirPath.IsEmpty())
+		{
+			// Old-style with all attachments by value
+
+			OpEntry.Save(SerializedPackage);
+		}
+		else
+		{
+			TConstArrayView<FCbAttachment> Attachments = OpEntry.GetAttachments();
+
+			// Prep phase
+
+			TSet<FIoHash> NeedChunks;
+
+			{
+				FCbWriter Writer;
+				Writer.BeginObject();
+				Writer.BeginArray("have");
+
+				for (const FCbAttachment& Attachment : Attachments)
+				{
+					Writer.AddHash(Attachment.GetHash());
+				}
+
+				Writer.EndArray();
+				Writer.EndObject();
+
+				FCbFieldIterator Prep = Writer.Save();
+
+				UE::Zen::FScopedRequestPtr Request(RequestPool.Get());
+
+				bool IsOk = false;
+				
+				const Zen::FRequest::Result Res = Request->PerformBlockingUpload<Zen::FRequest::Post>(OplogPrepNewEntryPath, Prep.AsObjectView());
+
+				if (Res == Zen::FRequest::Success)
+				{
+					FCbObjectView NeedObject;
+
+					if (Res == Zen::FRequest::Success && Request->GetResponseCode() == 200)
+					{
+						NeedObject = FCbObjectView(Request->GetResponseBuffer().GetData());
+
+						for (auto& Entry : NeedObject["need"])
+						{
+							NeedChunks.Add(Entry.AsHash());
+						}
+
+						IsOk = true;
+					}
+				}
+			}
+
+			// This uses a slight variation for package attachment serialization
+			// by writing larger attachments to a file and referencing it in the
+			// core object. Small attachments are serialized inline as normal
+
+			FCbWriter Writer;
+
+			FCbObject PackageObj = OpEntry.GetObject();
+			const FIoHash PackageObjHash = PackageObj.GetHash();
+
+			Writer.AddObject(PackageObj);
+			Writer.AddObjectAttachment(PackageObjHash);
+
+			// Send phase
+
+			for (const FCbAttachment& Attachment : Attachments)
+			{
+				bool IsSerialized = false;
+
+				const FIoHash AttachmentHash = Attachment.GetHash();
+
+				if (NeedChunks.Contains(AttachmentHash))
+				{
+					if (FSharedBuffer AttachView = Attachment.AsBinary())
+					{
+						if (AttachView.GetSize() >= StandaloneThresholdBytes)
+						{
+							// Write to temporary file. To avoid race conditions we derive
+							// the file name from a salt value and the attachment hash
+
+							FIoHash AttachmentSpec[] { FIoHash::HashBuffer(&Salt, sizeof Salt), AttachmentHash };
+							FIoHash AttachmentId = FIoHash::HashBuffer(MakeMemoryView(AttachmentSpec));
+
+							FString TempFilePath = TempDirPath / LexToString(AttachmentId);
+							IPlatformFile& PlatformFile = FPlatformFileManager::Get().GetPlatformFile();
+
+							if (IFileHandle* FileHandle = PlatformFile.OpenWrite(*TempFilePath))
+							{
+								FileHandle->Write((const uint8*)AttachView.GetData(), AttachView.GetSize());
+								delete FileHandle;
+
+								Writer.AddHash(AttachmentHash);
+
+								IsSerialized = true;
+								IsUsingTempFiles = true;
+							}
+							else
+							{
+								// Take the slow path if we can't open the payload file in the
+								// large attachment directory
+
+								UE_LOG(LogZenStore, Warning, TEXT("Could not create file '%s', taking slow path for large attachment"), *TempFilePath);
+							}
+						}
+					}
+
+					if (!IsSerialized)
+					{
+						Attachment.Save(Writer);
+					}
+				}
+				else
+				{
+					Writer.AddHash(AttachmentHash);
+				}
+			}
+			Writer.AddNull();
+
+			Writer.Save(SerializedPackage);
+		}
+
+		UE_LOG(LogZenStore, Verbose, TEXT("Package size: %" UINT64_FMT), SerializedPackage.TotalSize());
+
+		UE::Zen::FScopedRequestPtr Request(RequestPool.Get());
+
+		TStringBuilder<64> NewOpPostUri;
+		NewOpPostUri << OplogNewEntryPath;
+
+		if (IsUsingTempFiles)
+		{
+			NewOpPostUri << "?salt=" << Salt;
+		}
+
+		if (UE::Zen::FRequest::Success == Request->PerformBlockingPost(NewOpPostUri, SerializedPackage.GetView()))
+		{
+			return TIoStatusOr<uint64>(SerializedPackage.TotalSize());
+		}
+		else
+		{
+			return TIoStatusOr<uint64>((FIoStatus)(FIoStatusBuilder(EIoErrorCode::Unknown) << TEXT("Append OpLog failed, NewOpLogPath='") << OplogNewEntryPath << TEXT("'")));
+		}
+	});
+}
+
+TIoStatusOr<uint64> FZenStoreHttpClient::GetChunkSize(const FIoChunkId& Id)
+{
+	TRACE_CPUPROFILER_EVENT_SCOPE(ZenStoreHttp_GetChunkSize);
+
+	check(bAllowRead);
+
+	UE::Zen::FScopedRequestPtr Request(RequestPool.Get());
+	TArray64<uint8> GetBuffer;
+	TStringBuilder<128> ChunkUri;
+	ChunkUri << OplogPath << '/' << Id;
+	UE::Zen::FRequest::Result Res = Request->PerformBlockingQuery<Zen::FRequest::Head>(ChunkUri);
+	FString ContentLengthStr;
+	if (Res == Zen::FRequest::Success && Request->GetResponseCode() == 200 && Request->GetHeader("Content-Length", ContentLengthStr))
+	{
+		return FCStringWide::Atoi64(*ContentLengthStr);
+	}
+	return FIoStatus(EIoErrorCode::NotFound);
+}
+
+TIoStatusOr<FIoBuffer> FZenStoreHttpClient::ReadChunk(const FIoChunkId& Id, uint64 Offset, uint64 Size)
+{
+	TRACE_CPUPROFILER_EVENT_SCOPE(ZenStoreHttp_ReadChunk);
+
+	check(bAllowRead);
+
+	UE::Zen::FScopedRequestPtr Request(RequestPool.Get());
+	TArray64<uint8> GetBuffer;
+	TStringBuilder<128> ChunkUri;
+	ChunkUri << OplogPath << '/' << Id;
+
+	bool bHaveQuery = false;
+
+	auto AppendQueryDelimiter = [&bHaveQuery, &ChunkUri]
+	{
+		if (bHaveQuery)
+		{
+			ChunkUri.Append(TEXT("&"_WSV));
+		}
+		else
+		{
+			ChunkUri.Append(TEXT("?"_WSV));
+			bHaveQuery = true;
+		}
+	};
+
+	if (Offset)
+	{
+		AppendQueryDelimiter();
+		ChunkUri.Appendf(TEXT("offset=%" UINT64_FMT), Offset);
+	}
+
+	if (Size != ~uint64(0))
+	{
+		AppendQueryDelimiter();
+		ChunkUri.Appendf(TEXT("size=%" UINT64_FMT), Size);
+	}
+
+
+	UE::Zen::FRequest::Result Res = Request->PerformBlockingDownload(ChunkUri, &GetBuffer);
+
+	if (Res == Zen::FRequest::Success && Request->GetResponseCode() == 200)
+	{
+		return FIoBuffer(FIoBuffer::Clone, GetBuffer.GetData(), GetBuffer.Num());
+	}
+	return FIoStatus(EIoErrorCode::NotFound);
+}
+
+void 
+FZenStoreHttpClient::StartBuildPass()
+{
+	TRACE_CPUPROFILER_EVENT_SCOPE(ZenStoreHttp_StartBuildPass);
+
+	check(bAllowEdit);
+}
+
+TIoStatusOr<uint64>
+FZenStoreHttpClient::EndBuildPass(FCbPackage OpEntry)
+{
+	TRACE_CPUPROFILER_EVENT_SCOPE(ZenStoreHttp_EndBuildPass);
+
+	check(bAllowEdit);
+
+	FLargeMemoryWriter SerializedPackage;
+	OpEntry.Save(SerializedPackage);
+
+	UE_LOG(LogZenStore, Verbose, TEXT("Package size: %lld"), SerializedPackage.TotalSize());
+
+	UE::Zen::FScopedRequestPtr Request(RequestPool.Get());
+
+	TConstArrayView64<uint8> Payload { SerializedPackage.GetData(), SerializedPackage.TotalSize()};
+	
+	if (UE::Zen::FRequest::Success == Request->PerformBlockingUpload<Zen::FRequest::Post>(OplogNewEntryPath, Payload))
+	{
+		return static_cast<uint64>(Payload.Num());
+	}
+	else
+	{
+		return (FIoStatus)(FIoStatusBuilder(EIoErrorCode::Unknown) << TEXT("End build pass failed, NewOpLogPath='") << OplogNewEntryPath << TEXT("'"));
+	}
+}
+
+} // UE
+
+#else // not PLATFORM_WINDOWS, dummy implementation stub for now
+
+namespace UE {
+namespace Zen {
+	struct FRequestPool
+	{
+	};
+}
+
+FZenStoreHttpClient::FZenStoreHttpClient(const FStringView InHostName, uint16 InPort)
+{
+}
+
+FZenStoreHttpClient::~FZenStoreHttpClient()
+{
+}
+
+void FZenStoreHttpClient::Initialize(
+	FStringView InProjectId,
+	FStringView InOplogId,
+	FStringView ServerRoot,
+	FStringView EngineRoot,
+	FStringView ProjectRoot,
+	bool		IsCleanBuild)
+{
+}
+
+void FZenStoreHttpClient::InitializeReadOnly(FStringView InProjectId, FStringView InOplogId)
+{
+}
+
+TIoStatusOr<uint64> FZenStoreHttpClient::GetChunkSize(const FIoChunkId& Id)
+{
+	return 0;
+}
+
+TIoStatusOr<FIoBuffer> FZenStoreHttpClient::ReadChunk(const FIoChunkId& Id, uint64 Offset, uint64 Size)
+{
+	return FIoBuffer();
+}
+
+void FZenStoreHttpClient::StartBuildPass()
+{
+}
+
+TIoStatusOr<uint64> FZenStoreHttpClient::EndBuildPass(FCbPackage OpEntry)
+{
+	return FIoStatus(EIoErrorCode::Unknown);
+}
+
+TFuture<TIoStatusOr<uint64>> FZenStoreHttpClient::AppendOp(FCbPackage OpEntry)
+{
+	return TFuture<TIoStatusOr<uint64>>();
+}
+
+}
+
+#endif // PLATFORM_WINDOWS
diff --git a/Engine/Source/Developer/IoStoreUtilities/Private/ZenStoreWriter.cpp b/Engine/Source/Developer/IoStoreUtilities/Private/ZenStoreWriter.cpp
new file mode 100644
--- /dev/null
+++ b/Engine/Source/Developer/IoStoreUtilities/Private/ZenStoreWriter.cpp
@@ -0,0 +1,682 @@
+// Copyright Epic Games, Inc. All Rights Reserved.
+
+#include "ZenStoreWriter.h"
+#include "ZenStoreHttpClient.h"
+#include "ZenFileSystemManifest.h"
+#include "PackageStoreOptimizer.h"
+
+#include "Async/Async.h"
+#include "Serialization/CompactBinaryPackage.h"
+#include "Serialization/CompactBinaryWriter.h"
+#include "Serialization/LargeMemoryWriter.h" 
+#include "HAL/FileManager.h"
+#include "HAL/PlatformFileManager.h"
+#include "IO/IoDispatcher.h"
+#include "Misc/App.h"
+#include "Interfaces/ITargetPlatform.h"
+#include "Misc/Paths.h"
+#include "GenericPlatform/GenericPlatformFile.h"
+#include "Misc/ScopeRWLock.h"
+
+DEFINE_LOG_CATEGORY_STATIC(LogZenStoreWriter, Log, All);
+
+using namespace UE;
+
+FCbObjectId ToObjectId(const FIoChunkId& ChunkId)
+{
+	return FCbObjectId(MakeMemoryView(ChunkId.GetData(), ChunkId.GetSize()));
+}
+
+class FZenStoreWriter::FZenStoreHttpQueue
+{
+public:
+
+	struct FHttpRequestResult
+	{
+		double DurationInSeconds = 0.0;
+		uint64 TotalBytes = 0;
+		bool bOk = false;
+	};
+
+	using FHttpRequestCompleted = TFunction<void(FHttpRequestResult)>;
+
+	FZenStoreHttpQueue(UE::FZenStoreHttpClient& InHttpClient)
+		: HttpClient(InHttpClient)
+		, WakeUpEvent(FPlatformProcess::GetSynchEventFromPool())
+		, QueueEmptyEvent(FPlatformProcess::GetSynchEventFromPool(true))
+	{
+		BackgroundThread = AsyncThread([this]()
+		{ 
+			BackgroundThreadEntry();
+		});
+	}
+
+	~FZenStoreHttpQueue()
+	{
+		Shutdown();
+		FPlatformProcess::ReturnSynchEventToPool(WakeUpEvent);
+		FPlatformProcess::ReturnSynchEventToPool(QueueEmptyEvent);
+	}
+
+	void Enqueue(FCbPackage OpEntry, FHttpRequestCompleted&& Callback)
+	{
+		FHttpRequest* Request = Alloc();
+
+		Request->OpEntry	= MoveTemp(OpEntry);
+		Request->Callback	= MoveTemp(Callback);
+		Request->HttpResult = HttpClient
+			.AppendOp(MoveTemp(Request->OpEntry))
+			.Next([StartTime = FPlatformTime::Seconds()](TIoStatusOr<uint64> Status)
+			{
+				return FHttpRequestResult
+				{
+					FPlatformTime::Seconds() - StartTime,
+					Status.IsOk() ? Status.ValueOrDie() : 0,
+					Status.IsOk()
+				};
+			});
+
+		{
+			FScopeLock _(&QueueCriticalSection);
+
+			if (!LastPendingHttpRequest)
+			{
+				check(!NextPendingHttpRequest);
+				NextPendingHttpRequest = LastPendingHttpRequest = Request;
+			}
+			else
+			{
+				LastPendingHttpRequest->Next = Request;
+				LastPendingHttpRequest = Request;
+			}
+		}
+
+		WakeUpEvent->Trigger();
+	}
+
+	void Flush()
+	{
+		QueueEmptyEvent->Reset();
+		WakeUpEvent->Trigger();
+		QueueEmptyEvent->Wait();
+	}
+
+private:
+
+	struct FHttpRequest
+	{
+		FCbPackage OpEntry;
+		TFuture<FHttpRequestResult> HttpResult;
+		FHttpRequestCompleted Callback;
+		FHttpRequest* Next = nullptr;
+	};
+
+	FHttpRequest* Alloc()
+	{
+		FScopeLock _(&RequestsCriticalSection);
+
+		FHttpRequest* Request = nullptr;
+
+		if (FirstFreeHttpRequest)
+		{
+			Request = FirstFreeHttpRequest;
+			FirstFreeHttpRequest = Request->Next;
+			Request->Next = nullptr;
+		}
+		else
+		{
+			const int32 Index = HttpRequests.Add();
+			Request = &HttpRequests[Index];
+		}
+
+		check(!Request->Next);
+
+		return Request;
+	}
+
+	void Free(FHttpRequest* Request)
+	{
+		check(!Request->Next);
+
+		FScopeLock _(&RequestsCriticalSection);
+
+		Request->Next = FirstFreeHttpRequest;
+		FirstFreeHttpRequest = Request;
+	}
+
+	void BackgroundThreadEntry()
+	{
+		while (!bStop.Load())
+		{
+			FHttpRequest* PendingRequests = nullptr;
+			{
+				FScopeLock _(&QueueCriticalSection);
+				PendingRequests = NextPendingHttpRequest;
+				NextPendingHttpRequest = LastPendingHttpRequest = nullptr;
+			}
+
+			if (PendingRequests)
+			{
+				FHttpRequest* Request = PendingRequests;
+				while (Request)
+				{
+					Request->Callback(Request->HttpResult.Get());
+
+					FHttpRequest* FreeRequest = Request;
+					Request = Request->Next;
+					FreeRequest->Next = nullptr;
+					Free(FreeRequest);
+				}
+			}
+			else if (!bStop)
+			{
+				QueueEmptyEvent->Trigger();
+				WakeUpEvent->Wait();
+			}
+		}
+	}
+
+	void Shutdown()
+	{
+		if (!bStop)
+		{
+			bStop = true;
+			WakeUpEvent->Trigger();
+			BackgroundThread.Wait();
+		}
+	}
+
+	UE::FZenStoreHttpClient& HttpClient;
+	FCriticalSection QueueCriticalSection;
+	FEvent* WakeUpEvent;
+	FEvent* QueueEmptyEvent;
+	TFuture<void> BackgroundThread;
+	TAtomic<bool> bStop {false};
+
+	FCriticalSection RequestsCriticalSection;
+	TChunkedArray<FHttpRequest> HttpRequests;
+	FHttpRequest* NextPendingHttpRequest = nullptr;
+	FHttpRequest* LastPendingHttpRequest = nullptr;
+	FHttpRequest* FirstFreeHttpRequest = nullptr;
+};
+
+FZenStoreWriter::FZenStoreWriter(
+	const FString& InOutputPath, 
+	const FString& InMetadataDirectoryPath, 
+	const ITargetPlatform* InTargetPlatform,
+	bool IsCleanBuild)
+	: TargetPlatform(*InTargetPlatform)
+	, OutputPath(InOutputPath)
+	, MetadataDirectoryPath(InMetadataDirectoryPath)
+	, PackageStoreOptimizer(new FPackageStoreOptimizer())
+	, CookMode(IPackageStoreWriter::FCookInfo::CookByTheBookMode)
+{
+	FString HostName = TEXT("localhost");
+	uint16 Port = 1337;
+	FString ProjectId = FApp::GetProjectName();
+	FString OplogId = InTargetPlatform->PlatformName();
+
+	HttpClient = MakeUnique<UE::FZenStoreHttpClient>(HostName, Port);
+
+	FString RootDir = FPaths::RootDir();
+	FString EngineDir = FPaths::EngineDir();
+	FPaths::NormalizeDirectoryName(EngineDir);
+	FString ProjectDir = FPaths::ProjectDir();
+	FPaths::NormalizeDirectoryName(ProjectDir);
+
+	IPlatformFile& PlatformFile = FPlatformFileManager::Get().GetPlatformFile();
+	FString AbsServerRoot = PlatformFile.ConvertToAbsolutePathForExternalAppForRead(*RootDir);
+	FString AbsEngineDir = PlatformFile.ConvertToAbsolutePathForExternalAppForRead(*EngineDir);
+	FString AbsProjectDir = PlatformFile.ConvertToAbsolutePathForExternalAppForRead(*ProjectDir);
+
+	HttpClient->Initialize(	ProjectId, 
+							OplogId, 
+							AbsServerRoot, 
+							AbsEngineDir, 
+							AbsProjectDir, 
+							IsCleanBuild);
+
+	PackageStoreOptimizer->Initialize(InTargetPlatform);
+
+	FPackageStoreManifest::FZenServerInfo& ZenServerInfo = PackageStoreManifest.EditZenServerInfo();
+	ZenServerInfo.HostName = HostName;
+	ZenServerInfo.Port = Port;
+	ZenServerInfo.ProjectId = ProjectId;
+	ZenServerInfo.OplogId = OplogId;
+
+	ZenFileSystemManifest = MakeUnique<FZenFileSystemManifest>(TargetPlatform, OutputPath);
+	
+	HttpQueue = MakeUnique<FZenStoreHttpQueue>(*HttpClient);
+}
+
+FZenStoreWriter::~FZenStoreWriter()
+{
+	FWriteScopeLock _(PackagesLock);
+
+	if (PendingPackages.Num())
+	{
+		UE_LOG(LogZenStoreWriter, Warning, TEXT("Pending packages at shutdown!"));
+	}
+}
+
+void FZenStoreWriter::WritePackageData(const FPackageInfo& Info, const FIoBuffer& PackageData, const TArray<FFileRegion>& FileRegions)
+{
+	check(Info.ChunkId.IsValid());
+
+	TRACE_CPUPROFILER_EVENT_SCOPE(FZenStoreWriter::WritePackageData);
+
+	FIoBuffer CookedHeaderBuffer = FIoBuffer(PackageData.Data(), Info.HeaderSize, PackageData);
+	FIoBuffer CookedExportsBuffer = FIoBuffer(PackageData.Data() + Info.HeaderSize, PackageData.DataSize() - Info.HeaderSize, PackageData);
+	TUniquePtr<FPackageStorePackage> Package{PackageStoreOptimizer->CreatePackageFromCookedHeader(Info.PackageName, CookedHeaderBuffer)};
+	PackageStoreOptimizer->FinalizePackage(Package.Get());
+	TArray<FFileRegion> FileRegionsCopy(FileRegions);
+	for (FFileRegion& Region : FileRegionsCopy)
+	{
+		// Adjust regions so they are relative to the start of the exports buffer
+		Region.Offset -= Info.HeaderSize;
+	}
+	FIoBuffer PackageBuffer = PackageStoreOptimizer->CreatePackageBuffer(Package.Get(), CookedExportsBuffer, &FileRegionsCopy);
+	FIoChunkId ChunkId = CreateIoChunkId(Package->GetId().Value(), 0, EIoChunkType::ExportBundleData);
+	PackageStoreManifest.AddPackageData(Info.PackageName, Info.LooseFilePath, Info.ChunkId);
+	for (FFileRegion& Region : FileRegionsCopy)
+	{
+		// Adjust regions once more so they are relative to the exports bundle buffer
+		Region.Offset -= Package->GetHeaderSize();
+	}
+	//WriteFileRegions(*FPaths::ChangeExtension(Info.LooseFilePath, FString(".uexp") + FFileRegion::RegionsFileExtension), FileRegionsCopy);
+
+	// Commit to Zen build store
+
+	FCbObjectId ChunkOid = ToObjectId(Info.ChunkId);
+
+	// We do not modify the PendingPackages map, so take a shared lock
+	FReadScopeLock _(PackagesLock);
+
+	FPendingPackageState* ExistingState = PendingPackages.Find(Info.PackageName);
+
+	checkf(ExistingState != nullptr, TEXT("WritePackageData called for package which is not pending: '%s'"), *Info.PackageName.ToString());
+
+	FPackageDataEntry& Entry = ExistingState->PackageData;
+
+	Entry.Payload			= PackageBuffer;
+	Entry.Info				= Info;
+	Entry.ChunkId			= ChunkOid;
+	Entry.PackageStoreEntry = PackageStoreOptimizer->CreatePackageStoreEntry(Package.Get());
+	Entry.IsValid			= true;
+}
+
+void FZenStoreWriter::WriteIoStorePackageData(const FPackageInfo& Info, const FIoBuffer& PackageData, const FPackageStoreEntryResource& PackageStoreEntry, const TArray<FFileRegion>& FileRegions)
+{
+	check(Info.ChunkId.IsValid());
+
+	TRACE_CPUPROFILER_EVENT_SCOPE(WriteIoStorePackageData);
+
+	PackageStoreManifest.AddPackageData(Info.PackageName, Info.LooseFilePath, Info.ChunkId);
+	//WriteFileRegions(*FPaths::ChangeExtension(Info.LooseFilePath, FString(".uexp") + FFileRegion::RegionsFileExtension), FileRegionsCopy);
+
+	FCbObjectId ChunkOid = ToObjectId(Info.ChunkId);
+
+	// We do not modify the PendingPackages map, so take a shared lock
+	FReadScopeLock _(PackagesLock);
+
+	FPendingPackageState* ExistingState = PendingPackages.Find(Info.PackageName);
+
+	checkf(ExistingState != nullptr, TEXT("WritePackageData called for package which is not pending: '%s'"), *Info.PackageName.ToString());
+
+	FPackageDataEntry& Entry = ExistingState->PackageData;
+
+	Entry.Payload = PackageData;
+	Entry.Info = Info;
+	Entry.ChunkId = ChunkOid;
+	Entry.PackageStoreEntry = PackageStoreEntry;
+	Entry.IsValid = true;
+}
+
+void FZenStoreWriter::WriteBulkdata(const FBulkDataInfo& Info, const FIoBuffer& BulkData, const TArray<FFileRegion>& FileRegions)
+{
+	check(Info.ChunkId.IsValid());
+
+	FCbObjectId ChunkOid = ToObjectId(Info.ChunkId);
+
+	// We do not modify the PendingPackages map, so take a shared lock
+	FReadScopeLock _(PackagesLock);
+
+	FPendingPackageState* ExistingState = PendingPackages.Find(Info.PackageName);
+
+	checkf(ExistingState != nullptr, TEXT("WriteBulkdata called for package which is not pending: '%s'"), *Info.PackageName.ToString());
+
+	BulkData.MakeOwned();
+
+	FBulkDataEntry BulkEntry;
+
+	BulkEntry.Payload	= BulkData;
+	BulkEntry.Info		= Info;
+	BulkEntry.ChunkId	= ChunkOid;
+	BulkEntry.IsValid	= true;
+
+	ExistingState->BulkData.Emplace(BulkEntry);
+
+	PackageStoreManifest.AddBulkData(Info.PackageName, Info.LooseFilePath, Info.ChunkId);
+
+	//	WriteFileRegions(*(Info.LooseFilePath + FFileRegion::RegionsFileExtension), FileRegions);
+}
+
+bool FZenStoreWriter::WriteAdditionalFile(const FAdditionalFileInfo& Info, const FIoBuffer& FileData)
+{
+	const FZenFileSystemManifest::FManifestEntry& ManifestEntry = ZenFileSystemManifest->AddFile(Info.Filename);
+
+	FFileDataEntry FileEntry;
+
+	FileEntry.Payload				= FIoBuffer(FIoBuffer::Clone, FileData.Data(), FileData.DataSize());
+	FileEntry.Info					= Info;
+	FileEntry.Info.ChunkId			= CreateExternalFileChunkId(0, ManifestEntry.FileId);
+	FileEntry.ZenManifestServerPath = ManifestEntry.ServerPath;
+	FileEntry.ZenManifestClientPath = ManifestEntry.ClientPath;
+
+	// We do not modify the PendingPackages map, so take a shared lock
+	FReadScopeLock _(PackagesLock);
+
+	FPendingPackageState* ExistingState = PendingPackages.Find(Info.PackageName);
+
+	checkf(ExistingState != nullptr, TEXT("WriteAdditionalFile called for package which is not pending: '%s'"), *Info.PackageName.ToString());
+
+	ExistingState->FileData.Add(MoveTemp(FileEntry));
+
+	return true;
+}
+
+void FZenStoreWriter::BeginCook(const FCookInfo& Info)
+{
+	CookMode = Info.CookMode;
+
+	if (Info.CookMode == IPackageStoreWriter::FCookInfo::CookOnTheFlyMode)
+	{
+		FCbPackage Pkg;
+		FCbWriter PackageObj;
+		
+		PackageObj.BeginObject();
+		PackageObj << "key" << "CookOnTheFly";
+
+		const bool bGenerateContainerHeader = false;
+		CreateProjectMetaData(Pkg, PackageObj, bGenerateContainerHeader);
+
+		PackageObj.EndObject();
+		FCbObject Obj = PackageObj.Save().AsObject();
+
+		Pkg.SetObject(Obj);
+
+		TIoStatusOr<uint64> Status = HttpClient->AppendOp(Pkg).Get();
+		UE_CLOG(!Status.IsOk(), LogZenStoreWriter, Fatal, TEXT("Failed to append OpLog"));
+	}
+}
+
+void FZenStoreWriter::EndCook()
+{
+	HttpQueue->Flush();
+
+	FCbPackage Pkg;
+	FCbWriter PackageObj;
+	
+	PackageObj.BeginObject();
+	PackageObj << "key" << "EndCook";
+
+	const bool bGenerateContainerHeader = true;
+	CreateProjectMetaData(Pkg, PackageObj, bGenerateContainerHeader);
+
+	PackageObj.EndObject();
+	FCbObject Obj = PackageObj.Save().AsObject();
+
+	Pkg.SetObject(Obj);
+
+	TIoStatusOr<uint64> Status = HttpClient->EndBuildPass(Pkg);
+	UE_CLOG(!Status.IsOk(), LogZenStoreWriter, Fatal, TEXT("Failed to append OpLog and end the build pass"));
+
+	PackageStoreManifest.Save(*(MetadataDirectoryPath / TEXT("packagestore.manifest")));
+
+	UE_LOG(LogZenStoreWriter, Display, TEXT("Input:\t%d Packages"), PackageStoreOptimizer->GetTotalPackageCount());
+	UE_LOG(LogZenStoreWriter, Display, TEXT("Output:\t%d Export bundles"), PackageStoreOptimizer->GetTotalExportBundleCount());
+	UE_LOG(LogZenStoreWriter, Display, TEXT("Output:\t%d Export bundle entries"), PackageStoreOptimizer->GetTotalExportBundleEntryCount());
+	UE_LOG(LogZenStoreWriter, Display, TEXT("Output:\t%d Internal export bundle arcs"), PackageStoreOptimizer->GetTotalInternalBundleArcsCount());
+	UE_LOG(LogZenStoreWriter, Display, TEXT("Output:\t%d External export bundle arcs"), PackageStoreOptimizer->GetTotalExternalBundleArcsCount());
+	UE_LOG(LogZenStoreWriter, Display, TEXT("Output:\t%d Public runtime script objects"), PackageStoreOptimizer->GetTotalScriptObjectCount());
+	
+	UE_LOG(LogZenStoreWriter, Display, TEXT("%.2llfMiB uploaded to Zen storage in %.2llfmin (Speed=%.2llfMiB/s)"),
+		double(ZenStats.TotalBytes) / 1024.0 / 1024.0,
+		ZenStats.TotalRequestTime / 60.0,
+		(double(ZenStats.TotalBytes) / 1024.0 / 1024.0) / ZenStats.TotalRequestTime);
+}
+
+void FZenStoreWriter::BeginPackage(const FPackageBaseInfo& Info)
+{
+	FWriteScopeLock _(PackagesLock);
+
+	const FPendingPackageState* ExistingState = PendingPackages.Find(Info.PackageName);
+
+	checkf(ExistingState == nullptr, TEXT("BeginPackage called for package which is already pending: '%s'"), *Info.PackageName.ToString());
+
+	FPendingPackageState PackageState;
+	PackageState.PackageName = Info.PackageName;
+
+	PendingPackages.Add(Info.PackageName, PackageState);
+
+	PackageStoreManifest.BeginPackage(Info.PackageName);
+}
+
+void FZenStoreWriter::CommitPackage(const FPackageBaseInfo& Info)
+{
+	TRACE_CPUPROFILER_EVENT_SCOPE(FZenStoreWriter::CommitPackage);
+
+	IPackageStoreWriter::FCommitEventArgs CommitEventArgs;
+
+	CommitEventArgs.PlatformName	= FName(*TargetPlatform.PlatformName()),
+	CommitEventArgs.PackageName		= Info.PackageName;
+	CommitEventArgs.EntryIndex		= INDEX_NONE;
+	
+	FCbPackage Pkg;
+	bool bIsValid = false;
+
+	{
+		FWriteScopeLock _(PackagesLock);
+
+		FPendingPackageState* PackageState = PendingPackages.Find(Info.PackageName);
+
+		checkf(PackageState != nullptr, TEXT("CommitPackage called for package which is not pending: '%s'"), *Info.PackageName.ToString());
+
+		bIsValid = PackageState->PackageData.IsValid;
+
+		if (bIsValid)
+		{
+			FPackageDataEntry& PkgData = PackageState->PackageData;
+
+			CommitEventArgs.EntryIndex = PackageStoreEntries.Num();
+			PackageStoreEntries.Add(PkgData.PackageStoreEntry);
+
+			FCbWriter PackageObj;
+			PackageObj.BeginObject();
+			PackageObj << "key" << Info.PackageName.ToString();
+
+			// Note that this is destructive - we yank out the buffer memory from the 
+			// IoBuffer into the FSharedBuffer
+			auto IoBufferToSharedBuffer = [](FIoBuffer& InBuffer) -> FSharedBuffer {
+				InBuffer.EnsureOwned();
+				const uint64 DataSize = InBuffer.DataSize();
+				uint8* DataPtr = InBuffer.Release().ValueOrDie();
+				return FSharedBuffer{ FSharedBuffer::TakeOwnership(DataPtr, DataSize, FMemory::Free) };
+			};
+
+			FCbAttachment PkgDataAttachment(IoBufferToSharedBuffer(PkgData.Payload));
+			Pkg.AddAttachment(PkgDataAttachment);
+
+			PackageObj.BeginObject("package");
+			PackageObj << "id" << PkgData.ChunkId;
+			PackageObj << "data" << PkgDataAttachment;
+			PackageObj.EndObject();
+
+			if (PackageState->BulkData.Num())
+			{
+				PackageObj.BeginArray("bulkdata");
+
+				for (FBulkDataEntry& Bulk : PackageState->BulkData)
+				{
+					FCbAttachment BulkAttachment(IoBufferToSharedBuffer(Bulk.Payload));
+					Pkg.AddAttachment(BulkAttachment);
+
+					PackageObj.BeginObject();
+					PackageObj << "id" << Bulk.ChunkId;
+					PackageObj << "type" << LexToString(Bulk.Info.BulkdataType);
+					PackageObj << "data" << BulkAttachment;
+					PackageObj.EndObject();
+				}
+
+				PackageObj.EndArray();
+			}
+
+			if (PackageState->FileData.Num())
+			{
+				PackageObj.BeginArray("files");
+
+				for (FFileDataEntry& File : PackageState->FileData)
+				{
+					FCbAttachment FileDataAttachment(IoBufferToSharedBuffer(File.Payload));
+					Pkg.AddAttachment(FileDataAttachment);
+
+					PackageObj.BeginObject();
+					PackageObj << "id" << ToObjectId(File.Info.ChunkId);
+					PackageObj << "data" << FileDataAttachment;
+					PackageObj << "serverpath" << File.ZenManifestServerPath;
+					PackageObj << "clientpath" << File.ZenManifestClientPath;
+					PackageObj.EndObject();
+
+					CommitEventArgs.AdditionalFiles.Add(FAdditionalFileInfo
+					{ 
+						Info.PackageName,
+						File.ZenManifestClientPath,
+						File.Info.ChunkId
+					});
+				}
+
+				PackageObj.EndArray();
+			}
+
+			PackageObj.EndObject();
+
+			FCbObject Obj = PackageObj.Save().AsObject();
+
+			Pkg.SetObject(Obj);
+		}
+	}
+
+	if (bIsValid)
+	{
+		HttpQueue->Enqueue(Pkg, [this, EventArgs = MoveTemp(CommitEventArgs)](FZenStoreHttpQueue::FHttpRequestResult HttpResult) mutable
+		{
+			ZenStats.TotalBytes += HttpResult.TotalBytes;
+			ZenStats.TotalRequestTime += HttpResult.DurationInSeconds;
+			BroadcastCommit(EventArgs);
+		});
+	}
+	else
+	{
+		BroadcastCommit(CommitEventArgs);
+	}
+
+	{
+		FWriteScopeLock _(PackagesLock);
+		PendingPackages.Remove(Info.PackageName);
+	}
+}
+void FZenStoreWriter::GetEntries(TFunction<void(TArrayView<const FPackageStoreEntryResource>)>&& Callback)
+{
+	FReadScopeLock _(PackagesLock);
+	Callback(PackageStoreEntries);
+}
+
+void FZenStoreWriter::Flush()
+{
+	HttpQueue->Flush();
+}
+
+void FZenStoreWriter::CreateProjectMetaData(FCbPackage& Pkg, FCbWriter& PackageObj, bool bGenerateContainerHeader)
+{
+	// File Manifest
+	{
+		// Append all project and cooked files
+		ZenFileSystemManifest->Generate();
+
+		{
+			PackageObj.BeginArray("files");
+
+			for (const FZenFileSystemManifest::FManifestEntry& ManifestEntry : ZenFileSystemManifest->ManifestEntries())
+			{
+				FCbObjectId FileOid = ToObjectId(CreateExternalFileChunkId(0, ManifestEntry.FileId));
+
+				PackageObj.BeginObject();
+				PackageObj << "id" << FileOid;
+				PackageObj << "data" << FIoHash::Zero;
+				PackageObj << "serverpath" << ManifestEntry.ServerPath;
+				PackageObj << "clientpath" << ManifestEntry.ClientPath;
+				PackageObj.EndObject();
+			}
+
+			PackageObj.EndArray(); // End of files
+		}
+
+		ZenFileSystemManifest->Save(*FPaths::Combine(MetadataDirectoryPath, TEXT("zenfs.manifest")));
+	}
+
+	// Metadata section
+	{
+		PackageObj.BeginArray("meta");
+
+		// Summarize Script Objects
+		FIoBuffer ScriptObjectsBuffer = PackageStoreOptimizer->CreateScriptObjectsBuffer();
+		FCbObjectId ScriptOid = ToObjectId(CreateIoChunkId(0, 0, EIoChunkType::ScriptObjects));
+
+		FCbAttachment ScriptAttachment(FSharedBuffer::MakeView(ScriptObjectsBuffer.Data(), ScriptObjectsBuffer.DataSize()));
+		Pkg.AddAttachment(ScriptAttachment);
+
+		PackageObj.BeginObject();
+		PackageObj << "id" << ScriptOid;
+		PackageObj << "name" << "ScriptObjects";
+		PackageObj << "data" << ScriptAttachment;
+		PackageObj.EndObject();
+
+		// Generate Container Header
+		if (bGenerateContainerHeader)
+		{
+			FCbObjectId HeaderOid = ToObjectId(CreateIoChunkId(ContainerId.Value(), 0, EIoChunkType::ContainerHeader));
+			FIoBuffer HeaderBuffer;
+
+			{
+				FContainerHeader Header = PackageStoreOptimizer->CreateContainerHeader(ContainerId, PackageStoreEntries);
+				FLargeMemoryWriter HeaderAr(0, true);
+				HeaderAr << Header;
+				int64 DataSize = HeaderAr.TotalSize();
+				HeaderBuffer = FIoBuffer(FIoBuffer::AssumeOwnership, HeaderAr.ReleaseOwnership(), DataSize);
+			}
+
+			FCbAttachment HeaderAttachment(FSharedBuffer::MakeView(HeaderBuffer.Data(), HeaderBuffer.DataSize()));
+			Pkg.AddAttachment(HeaderAttachment);
+			
+			PackageObj.BeginObject();
+			PackageObj << "id" << HeaderOid;
+			PackageObj << "name" << "ContainerHeader";
+			PackageObj << "data" << HeaderAttachment;
+			PackageObj.EndObject();
+		}
+
+		PackageObj.EndArray();	// End of Meta array
+	}
+}
+
+void FZenStoreWriter::BroadcastCommit(IPackageStoreWriter::FCommitEventArgs& EventArgs)
+{
+	FScopeLock CommitEventLock(&CommitEventCriticalSection);
+	
+	if (CommitEvent.IsBound())
+	{
+		FReadScopeLock EntriesLock(PackagesLock);
+		EventArgs.Entries = PackageStoreEntries;
+		CommitEvent.Broadcast(EventArgs);
+	}
+}
diff --git a/Engine/Source/Developer/IoStoreUtilities/Public/FilePackageStoreWriter.h b/Engine/Source/Developer/IoStoreUtilities/Public/FilePackageStoreWriter.h
new file mode 100644
--- /dev/null
+++ b/Engine/Source/Developer/IoStoreUtilities/Public/FilePackageStoreWriter.h
@@ -0,0 +1,103 @@
+// Copyright Epic Games, Inc. All Rights Reserved.
+
+#pragma once
+
+#include "Serialization/AsyncLoading2.h"
+#include "IO/PackageStore.h"
+#include "IO/PackageStoreWriter.h"
+
+class FIoStoreWriterContext;
+class FIoStoreWriter;
+class FPackageStoreOptimizer;
+class FPackageStorePackage;
+class FPackageStoreContainerHeaderEntry;
+
+class FPackageStoreManifest
+{
+public:
+	struct FFileInfo
+	{
+		FString FileName;
+		FIoChunkId ChunkId;
+	};
+
+	struct FPackageInfo
+	{
+		FName PackageName;
+		FIoChunkId PackageChunkId;
+		TArray<FIoChunkId> BulkDataChunkIds;
+	};
+	
+	struct FZenServerInfo
+	{
+		FString HostName;
+		uint16 Port;
+		FString ProjectId;
+		FString OplogId;
+	};
+
+	IOSTOREUTILITIES_API FPackageStoreManifest() = default;
+	IOSTOREUTILITIES_API ~FPackageStoreManifest() = default;
+
+	IOSTOREUTILITIES_API void BeginPackage(FName PackageName);
+	IOSTOREUTILITIES_API void AddPackageData(FName PackageName, const FString& FileName, const FIoChunkId& ChunkId);
+	IOSTOREUTILITIES_API void AddBulkData(FName PackageName, const FString& FileName, const FIoChunkId& ChunkId);
+
+	IOSTOREUTILITIES_API FIoStatus Save(const TCHAR* Filename) const;
+	IOSTOREUTILITIES_API FIoStatus Load(const TCHAR* Filename);
+
+	IOSTOREUTILITIES_API TArray<FFileInfo> GetFiles() const;
+	IOSTOREUTILITIES_API TArray<FPackageInfo> GetPackages() const;
+
+	IOSTOREUTILITIES_API FZenServerInfo& EditZenServerInfo();
+	IOSTOREUTILITIES_API const FZenServerInfo* ReadZenServerInfo() const;
+
+private:
+	mutable FCriticalSection CriticalSection;
+	TMap<FName, FPackageInfo> PackageInfoByNameMap;
+	TMap<FIoChunkId, FString> FileNameByChunkIdMap;
+	TUniquePtr<FZenServerInfo> ZenServerInfo;
+};
+
+class FFilePackageStoreWriter
+	: public IPackageStoreWriter
+{
+public:
+	IOSTOREUTILITIES_API FFilePackageStoreWriter(const FString& OutputPath, const FString& MetadataDirectoryPath, const ITargetPlatform* TargetPlatform);
+	IOSTOREUTILITIES_API ~FFilePackageStoreWriter();
+	IOSTOREUTILITIES_API virtual void BeginPackage(const FPackageBaseInfo& Info) override;
+	IOSTOREUTILITIES_API virtual void CommitPackage(const FPackageBaseInfo& Info) override;
+	IOSTOREUTILITIES_API virtual void WritePackageData(const FPackageInfo& Info, const FIoBuffer& PackageData, const TArray<FFileRegion>& FileRegions) override;
+	IOSTOREUTILITIES_API virtual void WriteBulkdata(const FBulkDataInfo& Info, const FIoBuffer& BulkData, const TArray<FFileRegion>& FileRegions) override;
+	IOSTOREUTILITIES_API virtual bool WriteAdditionalFile(const FAdditionalFileInfo& Info, const FIoBuffer& FileData) override { return false; }
+
+	IOSTOREUTILITIES_API virtual void BeginCook(const FCookInfo& Info);
+	IOSTOREUTILITIES_API virtual void EndCook();
+
+	virtual void GetEntries(TFunction<void(TArrayView<const FPackageStoreEntryResource>)>&& Callback) override
+	{
+		check(false);
+	}
+
+	DECLARE_DERIVED_EVENT(FFilePackageStoreWriter, IPackageStoreWriter::FCommitEvent, FCommitEvent);
+	virtual FCommitEvent& OnCommit() override
+	{
+		check(false);
+		return CommitEvent;
+	}
+
+	virtual void Flush() override
+	{ }
+
+private:
+	TUniquePtr<FPackageStoreOptimizer> PackageStoreOptimizer;
+	TUniquePtr<FIoStoreWriterContext> IoStoreWriterContext;
+	TUniquePtr<FIoStoreWriter> IoStoreWriter;
+	const ITargetPlatform& TargetPlatform;
+	FString OutputPath;
+	FString MetadataDirectoryPath;
+	FPackageStoreManifest PackageStoreManifest;
+	FIoContainerId ContainerId = FIoContainerId::FromName(TEXT("global"));
+	TArray<FPackageStoreEntryResource> PackageStoreEntries;
+	FCommitEvent CommitEvent;
+};
diff --git a/Engine/Source/Developer/IoStoreUtilities/Public/ZenStoreWriter.h b/Engine/Source/Developer/IoStoreUtilities/Public/ZenStoreWriter.h
new file mode 100644
--- /dev/null
+++ b/Engine/Source/Developer/IoStoreUtilities/Public/ZenStoreWriter.h
@@ -0,0 +1,124 @@
+// Copyright Epic Games, Inc. All Rights Reserved.
+
+#pragma once
+
+#include "Serialization/AsyncLoading2.h"
+#include "IO/IoDispatcher.h"
+#include "Containers/Map.h"
+#include "Serialization/CompactBinary.h"
+#include "IO/PackageStore.h"
+#include "IO/PackageStoreWriter.h"
+#include "FilePackageStoreWriter.h"
+
+namespace UE {
+	class FZenStoreHttpClient;
+}
+
+class FPackageStoreOptimizer;
+class FZenFileSystemManifest;
+class FCbPackage;
+class FCbWriter;
+
+/** 
+ * Client for interfacing with Zen storage service
+ */
+
+class FZenStoreWriter
+	: public IPackageStoreWriter
+{
+public:
+	IOSTOREUTILITIES_API FZenStoreWriter(	const FString& OutputPath, 
+											const FString& MetadataDirectoryPath, 
+											const ITargetPlatform* TargetPlatform, 
+											bool IsCleanBuild);
+
+	IOSTOREUTILITIES_API ~FZenStoreWriter();
+
+	IOSTOREUTILITIES_API virtual void BeginPackage(const FPackageBaseInfo& Info) override;
+	IOSTOREUTILITIES_API virtual void CommitPackage(const FPackageBaseInfo& Info) override;
+
+	IOSTOREUTILITIES_API virtual void WritePackageData(const FPackageInfo& Info, const FIoBuffer& PackageData, const TArray<FFileRegion>& FileRegions) override;
+	IOSTOREUTILITIES_API virtual bool WriteAdditionalFile(const FAdditionalFileInfo& Info, const FIoBuffer& FileData) override;
+	IOSTOREUTILITIES_API virtual void WriteBulkdata(const FBulkDataInfo& Info, const FIoBuffer& BulkData, const TArray<FFileRegion>& FileRegions) override;
+	IOSTOREUTILITIES_API virtual void BeginCook(const FCookInfo& Info) override;
+	IOSTOREUTILITIES_API virtual void EndCook() override;
+
+	IOSTOREUTILITIES_API virtual void GetEntries(TFunction<void(TArrayView<const FPackageStoreEntryResource>)>&& Callback) override;
+
+	DECLARE_DERIVED_EVENT(FZenStoreWriter, IPackageStoreWriter::FCommitEvent, FCommitEvent);
+	IOSTOREUTILITIES_API virtual FCommitEvent& OnCommit() override
+	{
+		return CommitEvent;
+	}
+
+	virtual void Flush() override;
+
+	IOSTOREUTILITIES_API void WriteIoStorePackageData(const FPackageInfo& Info, const FIoBuffer& PackageData, const FPackageStoreEntryResource& PackageStoreEntry, const TArray<FFileRegion>& FileRegions);
+
+private:
+	void CreateProjectMetaData(FCbPackage& Pkg, FCbWriter& PackageObj, bool bGenerateContainerHeader);
+	void BroadcastCommit(IPackageStoreWriter::FCommitEventArgs& EventArgs);
+
+	struct FBulkDataEntry
+	{
+		FIoBuffer Payload;
+		FBulkDataInfo Info;
+		FCbObjectId ChunkId;
+		bool IsValid = false;
+	};
+
+	struct FPackageDataEntry
+	{
+		FIoBuffer Payload;
+		FPackageInfo Info;
+		FCbObjectId ChunkId;
+		FPackageStoreEntryResource PackageStoreEntry;
+		bool IsValid = false;
+	};
+
+	struct FFileDataEntry
+	{
+		FIoBuffer Payload;
+		FAdditionalFileInfo Info;
+		FString ZenManifestServerPath;
+		FString ZenManifestClientPath;
+	};
+
+	struct FPendingPackageState
+	{
+		FName PackageName;
+		FPackageDataEntry PackageData;
+		TArray<FBulkDataEntry> BulkData;
+		TArray<FFileDataEntry> FileData;
+	};
+
+	struct FZenStats
+	{
+		uint64 TotalBytes = 0;
+		double TotalRequestTime = 0.0;
+	};
+
+	FRWLock								PackagesLock;
+	TMap<FName,FPendingPackageState>	PendingPackages;
+	TUniquePtr<UE::FZenStoreHttpClient>	HttpClient;
+
+	const ITargetPlatform&				TargetPlatform;
+	FString								OutputPath;
+	FString								MetadataDirectoryPath;
+	FIoContainerId						ContainerId = FIoContainerId::FromName(TEXT("global"));
+
+	FPackageStoreManifest				PackageStoreManifest;
+	TUniquePtr<FPackageStoreOptimizer>	PackageStoreOptimizer;
+	TArray<FPackageStoreEntryResource>	PackageStoreEntries;
+	TUniquePtr<FZenFileSystemManifest>	ZenFileSystemManifest;
+	
+	FCriticalSection					CommitEventCriticalSection;
+	FCommitEvent						CommitEvent;
+
+	class FZenStoreHttpQueue;
+	TUniquePtr<FZenStoreHttpQueue>		HttpQueue;
+	
+	IPackageStoreWriter::FCookInfo::ECookMode CookMode;
+
+	FZenStats							ZenStats;
+};
diff --git a/Engine/Source/Developer/PakFileUtilities/Private/PakFileUtilities.cpp b/Engine/Source/Developer/PakFileUtilities/Private/PakFileUtilities.cpp
--- a/Engine/Source/Developer/PakFileUtilities/Private/PakFileUtilities.cpp
+++ b/Engine/Source/Developer/PakFileUtilities/Private/PakFileUtilities.cpp
@@ -63,101 +63,6 @@ int64 GDDCHits = 0;
 int64 GDDCMisses = 0;
 #endif
 
-class FThreadLocalScratchSpace
-{
-public:
-
-	static FThreadLocalScratchSpace& Get()
-	{
-		static FThreadLocalScratchSpace Result;
-		return Result;
-	}
-	struct FScratchSpace 
-	{
-	public:
-		FScratchSpace()
-		{
-			Size = 0;
-			Buffer = nullptr;
-			Next = nullptr;
-		}
-		int64 Size;
-		uint8* Buffer;
-		TArray<uint8> Array;
-		FScratchSpace * Next; // Link is protected by Lock
-
-		virtual ~FScratchSpace()
-		{
-			// should be a NOP, CleanUpAll is called manually before we destruct
-			//	if we destruct before CleanUpAll the linked list points to freed memory
-			//	could call CleanUpAll here to prevent that
-			//	just assert instead
-			check( Buffer == nullptr );
-			CleanUp();
-		}
-
-		void Embiggen(int64 NewSize)
-		{
-			if (NewSize > Size)
-			{
-				Buffer = (uint8*)FMemory::Realloc(Buffer, NewSize);
-				Size = NewSize;
-			}
-		}
-		void CleanUp()
-		{
-			if ( Buffer )
-			{
-			FMemory::Free(Buffer);
-			Buffer = nullptr;
-			}
-			Size = 0;
-			Array.Empty();
-			// do NOT clear Next
-		}
-	};
-
-	void GetScratchSpace(FScratchSpace*& ScratchSpace)
-	{
-		ScratchSpace = (FScratchSpace*)FPlatformTLS::GetTlsValue(TLSSlot);
-		if (ScratchSpace == nullptr)
-		{
-			ScratchSpace = new FScratchSpace;
-			FPlatformTLS::SetTlsValue(TLSSlot, (void*)ScratchSpace);
-
-			FScopeLock Lock(&ScratchSpacesLinkLock);
-			ScratchSpace->Next = ScratchSpacesLink;
-			ScratchSpacesLink = ScratchSpace;
-		}
-	}
-
-	void CleanUpAll()
-	{
-		// free scratch memory of all threads scratch spaces :
-		//	(they must all be quiescent now)
-
-		FScopeLock Lock(&ScratchSpacesLinkLock);
-		for(FScratchSpace * Link = ScratchSpacesLink;Link != nullptr;Link = Link->Next)
-		{
-			Link->CleanUp();
-		}
-	}
-private:
-
-	FThreadLocalScratchSpace()
-	{
-		TLSSlot = FPlatformTLS::AllocTlsSlot();
-		ScratchSpacesLink = nullptr;
-	}
-
-	uint32 TLSSlot;
-	uint8 PadForFalseSharing[60];
-	FCriticalSection  ScratchSpacesLinkLock;
-	FScratchSpace * ScratchSpacesLink;
-	// ScratchSpacesLinkLock protects ScratchSpacesLink and the Next field of FScratchSpace
-};
-
-
 class FMemoryCompressor;
 
 /**
@@ -790,13 +695,6 @@ FString FCompressedFileBuffer::GetDDCKeyString(const uint8* UncompressedFile, co
 
 bool FCompressedFileBuffer::CompressFileToWorkingBuffer(const FPakInputPair& InFile, FName CompressionMethod, const int32 CompressionBlockSize)
 {
-
-	FThreadLocalScratchSpace::FScratchSpace* ScratchSpace = nullptr;
-	FThreadLocalScratchSpace::Get().GetScratchSpace(ScratchSpace);
-	check(ScratchSpace);
-
-
-
 	TUniquePtr<FArchive> FileHandle(IFileManager::Get().CreateFileReader(*InFile.Source));
 	if(!FileHandle)
 	{
@@ -808,28 +706,23 @@ bool FCompressedFileBuffer::CompressFileToWorkingBuffer(const FPakInputPair& InF
 	const int64 FileSize = OriginalSize;
 	const int64 PaddedEncryptedFileSize = Align(FileSize,FAES::AESBlockSize);
 	check(PaddedEncryptedFileSize >= FileSize);
-	if(ScratchSpace->Size < PaddedEncryptedFileSize)
-	{
-		ScratchSpace->Embiggen(PaddedEncryptedFileSize);
-	}
 
-
-	uint8*& InOutPersistentBuffer = ScratchSpace->Buffer;
-	int64& InOutBufferSize = ScratchSpace->Size;
+	TArray<uint8, TInlineAllocator<4096>> UncompressedBuffer;
+	UncompressedBuffer.SetNumUninitialized(PaddedEncryptedFileSize);
 
 	// Load to buffer
-	FileHandle->Serialize(InOutPersistentBuffer,FileSize);
-
+	FileHandle->Serialize(UncompressedBuffer.GetData(), FileSize);
 
+#if USE_DDC_FOR_COMPRESSED_FILES
+	const bool bShouldUseDDC = true; // && (FileSize > 20 * 1024 ? true : false);
 	FString DDCKey;
-	TArray<uint8>& GetData = ScratchSpace->Array;
-	const bool bShouldUseDDC = USE_DDC_FOR_COMPRESSED_FILES; // && (FileSize > 20 * 1024 ? true : false);
+	TArray<uint8, TInlineAllocator<4096>> GetData;
 	if (bShouldUseDDC)
 	{
 #if DETAILED_UNREALPAK_TIMING
 		FUnrealPakScopeCycleCounter Scope(GDDCSyncReadTime);
 #endif
-		DDCKey = GetDDCKeyString(InOutPersistentBuffer, FileSize, CompressionMethod, CompressionBlockSize);
+		DDCKey = GetDDCKeyString(UncompressedBuffer.GetData(), FileSize, CompressionMethod, CompressionBlockSize);
 
 		if (GetDerivedDataCacheRef().CachedDataProbablyExists(*DDCKey))
 		{
@@ -846,13 +739,14 @@ bool FCompressedFileBuffer::CompressFileToWorkingBuffer(const FPakInputPair& InF
 			}
 		}
 	}
+#endif
 
 	{
 #if DETAILED_UNREALPAK_TIMING
 		FUnrealPakScopeCycleCounter Scope(GCompressionTime);
 #endif
 		// Start parallel compress
-		FMemoryCompressor MemoryCompressor(InOutPersistentBuffer, FileSize, CompressionMethod, CompressionBlockSize);
+		FMemoryCompressor MemoryCompressor(UncompressedBuffer.GetData(), FileSize, CompressionMethod, CompressionBlockSize);
 
 		// Build buffers for working
 		int64 UncompressedSize = FileSize;
@@ -870,7 +764,7 @@ bool FCompressedFileBuffer::CompressFileToWorkingBuffer(const FPakInputPair& InF
 			int32 CompressedBlockSize = FMath::Max<int32>(CompressionBufferSize, MaxCompressedBlockSize);
 			FileCompressionBlockSize = FMath::Max<uint32>(BlockSize, FileCompressionBlockSize);
 			EnsureBufferSpace(Align(TotalCompressedSize + CompressedBlockSize, FAES::AESBlockSize));
-			if (!MemoryCompressor.CompressMemory(CompressionMethod, CompressedBuffer.Get() + TotalCompressedSize, CompressedBlockSize, InOutPersistentBuffer + UncompressedBytes, BlockSize))
+			if (!MemoryCompressor.CompressMemory(CompressionMethod, CompressedBuffer.Get() + TotalCompressedSize, CompressedBlockSize, UncompressedBuffer.GetData() + UncompressedBytes, BlockSize))
 			{
 				return false;
 			}
@@ -897,6 +791,7 @@ bool FCompressedFileBuffer::CompressFileToWorkingBuffer(const FPakInputPair& InF
 		}
 
 	}
+#if USE_DDC_FOR_COMPRESSED_FILES
 	if (bShouldUseDDC)
 	{
 #if DETAILED_UNREALPAK_TIMING
@@ -908,7 +803,7 @@ bool FCompressedFileBuffer::CompressFileToWorkingBuffer(const FPakInputPair& InF
 		SerializeDDCData(Ar);
 		GetDerivedDataCacheRef().Put(*DDCKey, GetData, InFile.Dest);
 	}
-
+#endif
 
 	return true;
 }
@@ -2169,8 +2064,6 @@ bool CreatePakFile(const TCHAR* Filename, TArray<FPakInputPair>& FilesToAdd, con
 	TArray<FAsyncCompressor> AsyncCompressors;
 	AsyncCompressors.AddZeroed(FilesToAdd.Num());
 	
-	FThreadLocalScratchSpace::Get();
-
 	class FRunCompressionTask : public FNonAbandonableTask
 	{
 		FAsyncCompressor* AsyncCompressor;
@@ -2518,8 +2411,6 @@ bool CreatePakFile(const TCHAR* Filename, TArray<FPakInputPair>& FilesToAdd, con
 	FMemory::Free(ReadBuffer);
 	ReadBuffer = NULL;
 
-	FThreadLocalScratchSpace::Get().CleanUpAll();
-
 	if (RequiredPatchPadding > 0)
 	{
 		for (const FPakEntryPair& Pair : Index)
diff --git a/Engine/Source/Developer/Virtualization/Private/VirtualizationJupiterBackend.cpp b/Engine/Source/Developer/Virtualization/Private/VirtualizationJupiterBackend.cpp
--- a/Engine/Source/Developer/Virtualization/Private/VirtualizationJupiterBackend.cpp
+++ b/Engine/Source/Developer/Virtualization/Private/VirtualizationJupiterBackend.cpp
@@ -1246,13 +1246,13 @@ struct FJupiterServiceStatus : public FJsonSerializable
 	{
 		UE_LOG(LogVirtualization, Log, TEXT("Jupiter Service Status:"));
 		UE_LOG(LogVirtualization, Log, TEXT("Version: %s"), *Version);
+		UE_LOG(LogVirtualization, Log, TEXT("Site Id: %s"), *SiteIdentifier);
 		UE_LOG(LogVirtualization, Log, TEXT("GitHash: %s"), *GitHash);		
 		UE_LOG(LogVirtualization, Log, TEXT("Capabilities:"));
 		for (const FString& Capability : Capabilities)
 		{
 			UE_LOG(LogVirtualization, Log, TEXT("\t%s"), *Capability);
-		}
-		UE_LOG(LogVirtualization, Log, TEXT("Site Id: %s"), *SiteIdentifier);
+		}	
 	}
 
 private:
diff --git a/Engine/Source/Editor/UnrealEd/Classes/Commandlets/CookCommandlet.h b/Engine/Source/Editor/UnrealEd/Classes/Commandlets/CookCommandlet.h
--- a/Engine/Source/Editor/UnrealEd/Classes/Commandlets/CookCommandlet.h
+++ b/Engine/Source/Editor/UnrealEd/Classes/Commandlets/CookCommandlet.h
@@ -55,6 +55,8 @@ class UCookCommandlet
 	bool bPartialGC;
 	/** Do not cook any shaders.  Shader maps will be empty */
 	bool bNoShaderCooking;
+	/** Whether to cook for I/O store */
+	bool bUseIoStore;
 	/** All commandline tokens */
 	TArray<FString> Tokens;
 	/** All commandline switches */
diff --git a/Engine/Source/Editor/UnrealEd/Classes/CookOnTheSide/CookOnTheFlyServer.h b/Engine/Source/Editor/UnrealEd/Classes/CookOnTheSide/CookOnTheFlyServer.h
--- a/Engine/Source/Editor/UnrealEd/Classes/CookOnTheSide/CookOnTheFlyServer.h
+++ b/Engine/Source/Editor/UnrealEd/Classes/CookOnTheSide/CookOnTheFlyServer.h
@@ -73,6 +73,7 @@ enum class ECookByTheBookOptions
 	IoStore =							0x00008000, // Cook for IoStore
 	CookAgainstFixedBase =				0x00010000, // If cooking DLC, assume that the base content can not be modified. 
 	DlcLoadMainAssetRegistry =			0x00020000, // If cooking DLC, populate the main game asset registry
+	ZenStore =							0x00040000, // Store cooked data in Zen Store
 
 	// Deprecated flags
 	DisableUnsolicitedPackages UE_DEPRECATED(4.26, "Use SkipSoftReferences and/or SkipHardReferences instead") = SkipSoftReferences | SkipHardReferences,
@@ -123,6 +124,7 @@ namespace Cook
 	struct FPlatformManager;
 	struct FTickStackData;
 	struct FPopulatePackageContext;
+	class ICookOnTheFlyRequestManager;
 }
 }
 
@@ -137,7 +139,25 @@ class UNREALED_API UCookOnTheFlyServer : public UObject, public FTickableEditorO
 public:
 	struct FCookByTheBookOptions;
 
+	//////////////////////////////////////////////////////////////////////////
+	// Cook on the fly startup options
+	struct FCookOnTheFlyOptions
+	{
+		/** Wether the network file server or the I/O store connection server should bind to any port */
+		bool bBindAnyPort = false;
+		/** Wether to cook for I/O store */
+		bool bIoStore = false;
+		/** Target platforms */
+		TArray<ITargetPlatform*> TargetPlatforms;
+	};
+
 private:
+	//////////////////////////////////////////////////////////////////////////
+	// Cook on the fly server interface adapter
+	class FCookOnTheFlyServerInterface;
+	friend class FCookOnTheFlyServerInterface;
+	TUniquePtr<FCookOnTheFlyServerInterface> CookOnTheFlyServerInterface;
+
 	/** Current cook mode the cook on the fly server is running in */
 	ECookMode::Type CurrentCookMode = ECookMode::CookOnTheFly;
 	/** Directory to output to instead of the default should be empty in the case of DLC cooking */ 
@@ -149,9 +169,11 @@ private:
 	//////////////////////////////////////////////////////////////////////////
 	// Cook on the fly options
 
+	FCookOnTheFlyOptions CookOnTheFlyOptions;
 	/** Cook on the fly server uses the NetworkFileServer */
 	TArray<class INetworkFileServer*> NetworkFileServers;
 	FOnFileModifiedDelegate FileModifiedDelegate;
+	TUniquePtr<UE::Cook::ICookOnTheFlyRequestManager> CookOnTheFlyRequestManager;
 
 	//////////////////////////////////////////////////////////////////////////
 	// General cook options
@@ -264,7 +286,7 @@ private:
 	void OnTargetPlatformChangedSupportedFormats(const ITargetPlatform* TargetPlatform);
 
 	/* Version of AddCookOnTheFlyPlatform that takes the Platform name instead of an ITargetPlatform*.  Returns the Platform if found */
-	const ITargetPlatform* AddCookOnTheFlyPlatform(const FString& PlatformName);
+	const ITargetPlatform* AddCookOnTheFlyPlatform(const FName& PlatformName);
 	/* Internal helper for AddCookOnTheFlyPlatform.  Initializing Platforms must be done on the tickloop thread; Platform data is read only on other threads */
 	void AddCookOnTheFlyPlatformFromGameThread(ITargetPlatform* TargetPlatform);
 
@@ -381,25 +403,23 @@ public:
 	void Initialize( ECookMode::Type DesiredCookMode, ECookInitializationFlags InCookInitializationFlags, const FString& OutputDirectoryOverride = FString() );
 
 	/**
-	* Cook on the side, cooks while also running the editor...
+	* Initialize cook on the fly server
 	*
-	* @param BindAnyPort					Whether to bind on any port or the default port.
-	* @param TargetPlatforms				If nonempty, cooking will be prepared (generate AssetRegistry, etc) for each platform in the array
+	* @param InCookOnTheFlyOptions Cook on the fly startup options
 	*
 	* @return true on success, false otherwise.
 	*/
-	bool StartNetworkFileServer( bool BindAnyPort, const TArray<ITargetPlatform*>& TargetPlatforms = TArray<ITargetPlatform*>() );
+	bool StartCookOnTheFly(FCookOnTheFlyOptions InCookOnTheFlyOptions); 
 
 	/**
 	* Broadcast our the fileserver presence on the network
 	*/
 	bool BroadcastFileserverPresence( const FGuid &InstanceId );
+	
 	/** 
-	* Stop the network file server
-	*
+	* Shutdown cook on the fly server
 	*/
-	void EndNetworkFileServer();
-
+	void ShutdownCookOnTheFly();
 
 	struct FCookByTheBookStartupOptions
 	{
@@ -552,6 +572,7 @@ public:
 	bool IsUsingShaderCodeLibrary() const;
 
 	bool IsUsingIoStore() const;
+	bool IsUsingZenStore() const;
 
 	/**
 	* Helper function returns if we are in any cook on the fly mode
@@ -732,48 +753,8 @@ private:
 	//////////////////////////////////////////////////////////////////////////
 	// cook on the fly specific functions
 
-	/**
-	 * When we get a new connection from the network make sure the version is compatible 
-	 *		Will terminate the connection if return false
-	 * 
-	 * @return return false if not compatible, true if it is
-	 */
-	bool HandleNetworkFileServerNewConnection( const FString& VersionInfo, const FString& PlatformName );
-
 	void GetCookOnTheFlyUnsolicitedFiles(const ITargetPlatform* TargetPlatform, const FString& PlatformName, TArray<FString>& UnsolicitedFiles, const FString& Filename, bool bIsCookable);
 
-	/**
-	* Cook requests for a package from network
-	*  blocks until cook is complete
-	* 
-	* @param  Filename	requested file to cook. This can be altered if cooking produces a primary output asset with a different name as with uasset/umap extension changes.
-	* @param  PlatformName platform to cook for
-	* @param  out UnsolicitedFiles return a list of files which were cooked as a result of cooking the requested package
-	*/
-	void HandleNetworkFileServerFileRequest( FString& Filename, const FString& PlatformName, TArray<FString>& UnsolicitedFiles );
-
-	/**
-	* Shader recompile request from network
-	*  blocks until shader recompile complete
-	*
-	* @param  RecompileData input params for shader compile and compiled shader output
-	*/
-	void HandleNetworkFileServerRecompileShaders(const struct FShaderRecompileData& RecompileData);
-
-	/**
-	 * Get the sandbox path we want the network file server to use
-	 */
-	FString HandleNetworkGetSandboxPath();
-
-	/**
-	 * HandleNetworkGetPrecookedList 
-	 * this is used specifically for cook on the fly with shared cooked builds
-	 * returns the list of files which are still valid in the pak file which was initially loaded
-	 * 
-	 * @param PrecookedFileList all the files which are still valid in the client pak file
-	 */
-	void HandleNetworkGetPrecookedList( const FString& PlatformName, TMap<FString, FDateTime>& PrecookedFileList );
-
 	//////////////////////////////////////////////////////////////////////////
 	// general functions
 
@@ -828,6 +809,11 @@ private:
 	*/
 	void InitializeTargetPlatforms(const TArrayView<ITargetPlatform* const>& NewTargetPlatforms);
 
+	/**
+	* Initialize shaders for the specified platforms when running cook on the fly.
+	*/
+	void InitializeShadersForCookOnTheFly(const TArrayView<ITargetPlatform* const>& NewTargetPlatforms);
+
 	/**
 	* Some content plugins does not support all target platforms.
 	* Build up a map of unsupported packages per platform that can be checked before saving.
diff --git a/Engine/Source/Editor/UnrealEd/Classes/Settings/LevelEditorPlaySettings.h b/Engine/Source/Editor/UnrealEd/Classes/Settings/LevelEditorPlaySettings.h
--- a/Engine/Source/Editor/UnrealEd/Classes/Settings/LevelEditorPlaySettings.h
+++ b/Engine/Source/Editor/UnrealEd/Classes/Settings/LevelEditorPlaySettings.h
@@ -274,6 +274,10 @@ public:
 	UPROPERTY(config, EditAnywhere, Category = PlayInEditor, meta = (DisplayName="Stream Sub-Levels during Play in Editor", ToolTip="Prefer to stream sub-levels from the disk instead of duplicating editor sub-levels"))
 	uint32 bPreferToStreamLevelsInPIE:1;
 
+	/** Whether to use Zen Loader/Storage server when running cook-on-the-fly. */
+	UPROPERTY(config, EditAnywhere, Category=PlayInEditor)
+	bool EnableIoStoreCookOnTheFlyMode;
+
 public:
 	/** The width of the new view port window in pixels (0 = use the desktop's screen resolution). */
 	UPROPERTY(config, EditAnywhere, Category=GameViewportSettings, meta=(ClampMin=0))
diff --git a/Engine/Source/Editor/UnrealEd/Private/Commandlets/CookCommandlet.cpp b/Engine/Source/Editor/UnrealEd/Private/Commandlets/CookCommandlet.cpp
--- a/Engine/Source/Editor/UnrealEd/Private/Commandlets/CookCommandlet.cpp
+++ b/Engine/Source/Editor/UnrealEd/Private/Commandlets/CookCommandlet.cpp
@@ -418,9 +418,12 @@ bool UCookCommandlet::CookOnTheFly( FGuid InstanceId, int32 Timeout, bool bForce
 	CookFlags |= bUnversioned ? ECookInitializationFlags::Unversioned : ECookInitializationFlags::None;
 	CookOnTheFlyServer->Initialize( ECookMode::CookOnTheFly, CookFlags );
 
-	bool BindAnyPort = InstanceId.IsValid();
+	UCookOnTheFlyServer::FCookOnTheFlyOptions CookOnTheFlyStartupOptions;
+	CookOnTheFlyStartupOptions.bBindAnyPort = InstanceId.IsValid();
+	CookOnTheFlyStartupOptions.bIoStore = bUseIoStore;
+	CookOnTheFlyStartupOptions.TargetPlatforms = TargetPlatforms;
 
-	if ( CookOnTheFlyServer->StartNetworkFileServer(BindAnyPort, TargetPlatforms) == false )
+	if (CookOnTheFlyServer->StartCookOnTheFly(MoveTemp(CookOnTheFlyStartupOptions)) == false)
 	{
 		return false;
 	}
@@ -600,7 +603,7 @@ bool UCookCommandlet::CookOnTheFly( FGuid InstanceId, int32 Timeout, bool bForce
 		GShaderCompilingManager->SkipShaderCompilation(false);
 	}
 
-	CookOnTheFlyServer->EndNetworkFileServer();
+	CookOnTheFlyServer->ShutdownCookOnTheFly();
 	return true;
 }
 
@@ -628,6 +631,7 @@ int32 UCookCommandlet::Main(const FString& CmdLineParams)
 	bPartialGC = Switches.Contains(TEXT("Partialgc"));
 	ShowErrorCount = !Switches.Contains(TEXT("DIFFONLY"));
 	ShowProgress = !Switches.Contains(TEXT("DIFFONLY"));
+	bUseIoStore = Switches.Contains(TEXT("IoStore"));
 	bNoShaderCooking = bCookOnTheFly; // Do not cook any shaders into the shader maps. Always true if we are running w/ cook on the fly
 
 	COOK_STAT(DetailedCookStats::CookProject = FApp::GetProjectName());
@@ -891,6 +895,7 @@ bool UCookCommandlet::CookByTheBook( const TArray<ITargetPlatform*>& Platforms)
 	CookOptions |= Switches.Contains(TEXT("NODEV")) ? ECookByTheBookOptions::NoDevContent : ECookByTheBookOptions::None;
 	CookOptions |= Switches.Contains(TEXT("FullLoadAndSave")) ? ECookByTheBookOptions::FullLoadAndSave : ECookByTheBookOptions::None;
 	CookOptions |= Switches.Contains(TEXT("IoStore")) ? ECookByTheBookOptions::IoStore : ECookByTheBookOptions::None;
+	CookOptions |= Switches.Contains(TEXT("ZenStore")) ? (ECookByTheBookOptions::ZenStore | ECookByTheBookOptions::IoStore) : ECookByTheBookOptions::None;
 	CookOptions |= Switches.Contains(TEXT("NoGameAlwaysCook")) ? ECookByTheBookOptions::NoGameAlwaysCookPackages : ECookByTheBookOptions::None;
 	CookOptions |= Switches.Contains(TEXT("DisableUnsolicitedPackages")) ? (ECookByTheBookOptions::SkipHardReferences | ECookByTheBookOptions::SkipSoftReferences) : ECookByTheBookOptions::None;
 	CookOptions |= Switches.Contains(TEXT("NoDefaultMaps")) ? ECookByTheBookOptions::NoDefaultMaps : ECookByTheBookOptions::None;
diff --git a/Engine/Source/Editor/UnrealEd/Private/Commandlets/PackageUtilities.cpp b/Engine/Source/Editor/UnrealEd/Private/Commandlets/PackageUtilities.cpp
--- a/Engine/Source/Editor/UnrealEd/Private/Commandlets/PackageUtilities.cpp
+++ b/Engine/Source/Editor/UnrealEd/Private/Commandlets/PackageUtilities.cpp
@@ -53,6 +53,7 @@
 #include "Misc/OutputDeviceHelper.h"
 #include "Misc/OutputDeviceFile.h"
 #include "UObject/UObjectThreadContext.h"
+#include "Internationalization/GatherableTextData.h"
 
 DEFINE_LOG_CATEGORY(LogPackageHelperFunctions);
 DEFINE_LOG_CATEGORY_STATIC(LogPackageUtilities, Log, All);
diff --git a/Engine/Source/Editor/UnrealEd/Private/Commandlets/SavePackageUtilitiesCommandlet.cpp b/Engine/Source/Editor/UnrealEd/Private/Commandlets/SavePackageUtilitiesCommandlet.cpp
--- a/Engine/Source/Editor/UnrealEd/Private/Commandlets/SavePackageUtilitiesCommandlet.cpp
+++ b/Engine/Source/Editor/UnrealEd/Private/Commandlets/SavePackageUtilitiesCommandlet.cpp
@@ -14,6 +14,7 @@
 #include "UObject/LinkerDiff.h"
 #include "UObject/LinkerSave.h"
 #include "UObject/Package.h"
+#include "UObject/SavePackage.h"
 #include "UObject/UObjectGlobals.h"
 #include "UObject/UObjectHash.h"
 
