IoStore: Improved warning message for missing shader maps

#rnx

diff --git a/Engine/Source/Developer/IoStoreUtilities/Private/IoStoreUtilities.cpp b/Engine/Source/Developer/IoStoreUtilities/Private/IoStoreUtilities.cpp
--- a/Engine/Source/Developer/IoStoreUtilities/Private/IoStoreUtilities.cpp
+++ b/Engine/Source/Developer/IoStoreUtilities/Private/IoStoreUtilities.cpp
@@ -1366,7 +1366,7 @@ void ProcessShaderLibraries(const FIoStoreArguments& Arguments, TArray<FContaine
 				const TArray<FIoChunkId>* FindChunkIds = ShaderChunkIdsByShaderMapHash.Find(ShaderMapHash);
 				if (!FindChunkIds)
 				{
-					UE_LOG(LogIoStore, Warning, TEXT("Package referencing unknown shader map"));
+					UE_LOG(LogIoStore, Warning, TEXT("Package '%s' in '%s' referencing missing shader map '%s'"), *Package->PackageName.ToString(), *ContainerTarget->Name.ToString(), *ShaderMapHash.ToString());
 					continue;
 				}
 				for (const FIoChunkId& ShaderChunkId : *FindChunkIds)
