IoStore: Added some missing cpu scopes and fixed very poor performance when creating the disk layout
#rnx

diff --git a/Engine/Source/Developer/IoStoreUtilities/Private/IoStoreUtilities.cpp b/Engine/Source/Developer/IoStoreUtilities/Private/IoStoreUtilities.cpp
--- a/Engine/Source/Developer/IoStoreUtilities/Private/IoStoreUtilities.cpp
+++ b/Engine/Source/Developer/IoStoreUtilities/Private/IoStoreUtilities.cpp
@@ -722,6 +722,8 @@ static void AssignPackagesDiskOrder(
 	bool bClusterByOrderFilePriority
 	)
 {
+	IOSTORE_CPU_SCOPE(AssignPackagesDiskOrder);
+
 	struct FCluster
 	{
 		TArray<FLegacyCookedPackage*> Packages;
@@ -972,8 +974,11 @@ static void CreateDiskLayout(
 						}
 					}
 				}
-				SortedTargetFiles.Insert(PackageUniqueShaders, Index + 1);
-				Index += PackageUniqueShaders.Num();
+				if (!PackageUniqueShaders.IsEmpty())
+				{
+					SortedTargetFiles.Insert(PackageUniqueShaders, Index + 1);
+					Index += PackageUniqueShaders.Num();
+				}
 			}
 			++Index;
 		}
@@ -4156,6 +4161,8 @@ static bool ParsePakResponseFile(const TCHAR* FilePath, TArray<FContainerSourceF
 
 static bool ParsePakOrderFile(const TCHAR* FilePath, FFileOrderMap& Map)
 {
+	IOSTORE_CPU_SCOPE(ParsePakOrderFile);
+
 	TArray<FString> OrderFileContents;
 	if (!FFileHelper::LoadFileToStringArray(OrderFileContents, FilePath))
 	{
@@ -4681,6 +4688,7 @@ int32 CreateIoStoreContainerFiles(const TCHAR* CmdLine)
 		}
 		else
 		{
+			IOSTORE_CPU_SCOPE(FindCookedAssets);
 			UE_LOG(LogIoStore, Display, TEXT("Searching for cooked assets in folder '%s'"), *Arguments.CookedDir);
 			FCookedFileVisitor CookedFileVistor(Arguments.CookedFileStatMap, nullptr, Arguments.bFileRegions);
 			IFileManager::Get().IterateDirectoryStatRecursively(*Arguments.CookedDir, CookedFileVistor);
