Iterative cooking support when using Zen loader/storage server

Zen: https://github.com/EpicGames/zen/commit/af7ff3f1c66628c8005fb7d5d8e86e985d2d7b05

#rnx

diff --git a/Engine/Source/Developer/IoStoreUtilities/Private/FilePackageStoreWriter.cpp b/Engine/Source/Developer/IoStoreUtilities/Private/FilePackageStoreWriter.cpp
--- a/Engine/Source/Developer/IoStoreUtilities/Private/FilePackageStoreWriter.cpp
+++ b/Engine/Source/Developer/IoStoreUtilities/Private/FilePackageStoreWriter.cpp
@@ -348,11 +348,11 @@ void FFilePackageStoreWriter::EndCook()
 	UE_LOG(LogPackageStoreWriter, Display, TEXT("Output:\t%d Public runtime script objects"), PackageStoreOptimizer->GetTotalScriptObjectCount());
 }
 
-void FFilePackageStoreWriter::BeginPackage(const FPackageBaseInfo& Info)
+void FFilePackageStoreWriter::BeginPackage(const FBeginPackageInfo& Info)
 {
 	PackageStoreManifest.BeginPackage(Info.PackageName);
 }
 
-void FFilePackageStoreWriter::CommitPackage(const FPackageBaseInfo& Info)
+void FFilePackageStoreWriter::CommitPackage(const FCommitPackageInfo& Info)
 {
 }
diff --git a/Engine/Source/Developer/IoStoreUtilities/Private/IoStoreUtilities.cpp b/Engine/Source/Developer/IoStoreUtilities/Private/IoStoreUtilities.cpp
--- a/Engine/Source/Developer/IoStoreUtilities/Private/IoStoreUtilities.cpp
+++ b/Engine/Source/Developer/IoStoreUtilities/Private/IoStoreUtilities.cpp
@@ -4065,10 +4065,10 @@ int32 Staged2Zen(const FString& BuildPath, const FKeyChain& KeyChain, const ITar
 	{
 		const FPackageInfo& PackageInfo = PackagesArray[Index];
 
-		IPackageStoreWriter::FPackageBaseInfo PackageStorePackageBaseInfo;
-		PackageStorePackageBaseInfo.PackageName = PackageInfo.PackageName;
+		IPackageStoreWriter::FBeginPackageInfo BeginPackageInfo;
+		BeginPackageInfo.PackageName = PackageInfo.PackageName;
 
-		ZenStoreWriter->BeginPackage(PackageStorePackageBaseInfo);
+		ZenStoreWriter->BeginPackage(BeginPackageInfo);
 
 		IPackageStoreWriter::FPackageInfo PackageStorePackageInfo;
 		PackageStorePackageInfo.PackageName = PackageInfo.PackageName;
@@ -4088,7 +4088,10 @@ int32 Staged2Zen(const FString& BuildPath, const FKeyChain& KeyChain, const ITar
 			FIoBuffer BulkDataBuffer = BulkDataInfo.Chunk.Key->Read(BulkDataInfo.Chunk.Value, FIoReadOptions()).ValueOrDie();
 			ZenStoreWriter->WriteBulkdata(PackageStoreBulkDataInfo, BulkDataBuffer, TArray<FFileRegion>());
 		}
-		ZenStoreWriter->CommitPackage(PackageStorePackageBaseInfo);
+		
+		IPackageStoreWriter::FCommitPackageInfo CommitInfo;
+		CommitInfo.PackageName = PackageInfo.PackageName;
+		ZenStoreWriter->CommitPackage(CommitInfo);
 
 		int32 LocalUploadCount = UploadCount.IncrementExchange() + 1;
 		UE_CLOG(LocalUploadCount % 1000 == 0, LogIoStore, Display, TEXT("Uploading package %d/%d"), LocalUploadCount, PackagesArray.Num());
@@ -4803,4 +4806,4 @@ int32 CreateIoStoreContainerFiles(const TCHAR* CmdLine)
 	return 0;
 }
 
-#endif
\ No newline at end of file
+#endif
diff --git a/Engine/Source/Developer/IoStoreUtilities/Private/ZenFileSystemManifest.cpp b/Engine/Source/Developer/IoStoreUtilities/Private/ZenFileSystemManifest.cpp
--- a/Engine/Source/Developer/IoStoreUtilities/Private/ZenFileSystemManifest.cpp
+++ b/Engine/Source/Developer/IoStoreUtilities/Private/ZenFileSystemManifest.cpp
@@ -8,7 +8,7 @@
 #include "Misc/Paths.h"
 #include "Misc/DataDrivenPlatformInfoRegistry.h"
 
-const FZenFileSystemManifest::FManifestEntry FZenFileSystemManifest::InvalidEntry;
+const FZenFileSystemManifest::FManifestEntry FZenFileSystemManifest::InvalidEntry = FManifestEntry();
 
 FZenFileSystemManifest::FZenFileSystemManifest(const ITargetPlatform& InTargetPlatform, FString InCookDirectory)
 	: TargetPlatform(InTargetPlatform)
@@ -20,7 +20,7 @@ FZenFileSystemManifest::FZenFileSystemManifest(const ITargetPlatform& InTargetPl
 	FPaths::NormalizeDirectoryName(ServerRoot);
 }
 
-void FZenFileSystemManifest::Generate()
+int32 FZenFileSystemManifest::Generate()
 {
 	TRACE_CPUPROFILER_EVENT_SCOPE(GenerateStorageServerFileSystemManifest);
 
@@ -105,6 +105,7 @@ void FZenFileSystemManifest::Generate()
 						RelativePath.RightChopInline(LocalDirectory.Len() + 1);
 
 						AddManifestEntry(
+							CreateExternalFileChunkId(0, Entries.Num() + 1),
 							FPaths::Combine(ServerRelativeDirectory, RelativePath.GetData()),
 							FPaths::Combine(ClientDirectory, RelativePath.GetData()));
 					}
@@ -119,6 +120,8 @@ void FZenFileSystemManifest::Generate()
 			}
 		};
 
+	const int32 PreviousEntryCount = NumEntries();
+
 	AddFilesFromDirectory(TEXT("/{engine}"), FPaths::Combine(CookDirectory, TEXT("Engine")), true);
 	AddFilesFromDirectory(TEXT("/{project}"), FPaths::Combine(CookDirectory, FApp::GetProjectName()), true);
 	AddFilesFromDirectory(TEXT("/{project}"), ProjectDir, false);
@@ -148,9 +151,13 @@ void FZenFileSystemManifest::Generate()
 		AddFilesFromDirectory(FPaths::Combine(TEXT("/{engine}/Platforms"), PlatformDirectory), FPaths::Combine(EngineDir, TEXT("Platforms"), PlatformDirectory), true, &Filter);
 		AddFilesFromDirectory(FPaths::Combine(TEXT("/{project}/Platforms"), PlatformDirectory), FPaths::Combine(ProjectDir, TEXT("Platforms"), PlatformDirectory), true, &Filter);
 	}
+
+	const int32 CurrentEntryCount = NumEntries();
+	
+	return CurrentEntryCount - PreviousEntryCount;
 }
 
-const FZenFileSystemManifest::FManifestEntry& FZenFileSystemManifest::AddFile(const FString& Filename)
+const FZenFileSystemManifest::FManifestEntry& FZenFileSystemManifest::CreateManifestEntry(const FString& Filename)
 {
 	FString CookedEngineDirectory = FPaths::Combine(CookDirectory, TEXT("Engine"));
 
@@ -164,6 +171,7 @@ const FZenFileSystemManifest::FManifestEntry& FZenFileSystemManifest::AddFile(co
 		ServerRelativeDirectory = TEXT("/") + ServerRelativeDirectory;
 
 		return AddManifestEntry(
+			CreateExternalFileChunkId(0, Entries.Num() + 1),
 			FPaths::Combine(ServerRelativeDirectory, RelativePath.GetData()),
 			FPaths::Combine(ClientDirectory, RelativePath.GetData()));
 	};
@@ -182,11 +190,19 @@ const FZenFileSystemManifest::FManifestEntry& FZenFileSystemManifest::AddFile(co
 	return InvalidEntry;
 }
 
-const FZenFileSystemManifest::FManifestEntry& FZenFileSystemManifest::AddManifestEntry(FString ServerPath, FString ClientPath)
+const FZenFileSystemManifest::FManifestEntry& FZenFileSystemManifest::AddManifestEntry(const FIoChunkId& FileChunkId, FString ServerPath, FString ClientPath)
 {
+	check(ServerPath.Len() > 0 && ClientPath.Len() > 0);
+
 	ServerPath.ReplaceInline(TEXT("\\"), TEXT("/"));
 	ClientPath.ReplaceInline(TEXT("\\"), TEXT("/"));
 
+	// The server path is always relative to project root
+	if (ServerPath[0] == '/')
+	{
+		ServerPath.RightChopInline(1);
+	}
+
 	int32& EntryIndex = ServerPathToEntry.FindOrAdd(ServerPath, INDEX_NONE);
 
 	if (EntryIndex != INDEX_NONE)
@@ -199,14 +215,14 @@ const FZenFileSystemManifest::FManifestEntry& FZenFileSystemManifest::AddManifes
 	FManifestEntry Entry;
 	Entry.ServerPath = MoveTemp(ServerPath);
 	Entry.ClientPath = MoveTemp(ClientPath);
-	Entry.FileId = static_cast<uint32>(EntryIndex + 1);
+	Entry.FileChunkId = FileChunkId;
 
 	Entries.Add(MoveTemp(Entry));
 
 	return Entries[EntryIndex];
 }
 
-void FZenFileSystemManifest::Save(const TCHAR* Filename)
+bool FZenFileSystemManifest::Save(const TCHAR* Filename)
 {
 	check(Filename);
 
@@ -214,10 +230,13 @@ void FZenFileSystemManifest::Save(const TCHAR* Filename)
 	CsvLines.Add(FString::Printf(TEXT(";ServerRoot=%s, Platform=%s, CookDirectory=%s"), *ServerRoot, *TargetPlatform.PlatformName(), *CookDirectory));
 	CsvLines.Add(TEXT("FileId, ServerPath, ClientPath"));
 
+	TStringBuilder<2048> Sb;
 	for (const FManifestEntry& Entry : Entries)
 	{
-		CsvLines.Add(FString::Printf(TEXT("%d, %s, %s"), Entry.FileId, *Entry.ServerPath, *Entry.ClientPath));
+		Sb.Reset();
+		Sb << Entry.FileChunkId << TEXT(", ") << Entry.ServerPath << TEXT(", ") << *Entry.ClientPath;
+		CsvLines.Add(Sb.ToString());
 	}
 
-	FFileHelper::SaveStringArrayToFile(CsvLines, Filename);
+	return FFileHelper::SaveStringArrayToFile(CsvLines, Filename);
 }
diff --git a/Engine/Source/Developer/IoStoreUtilities/Private/ZenFileSystemManifest.h b/Engine/Source/Developer/IoStoreUtilities/Private/ZenFileSystemManifest.h
--- a/Engine/Source/Developer/IoStoreUtilities/Private/ZenFileSystemManifest.h
+++ b/Engine/Source/Developer/IoStoreUtilities/Private/ZenFileSystemManifest.h
@@ -3,6 +3,7 @@
 #pragma once
 
 #include "CoreMinimal.h"
+#include "IO/IoDispatcher.h"
 
 class ITargetPlatform;
 class FArchive;
@@ -15,30 +16,30 @@ public:
 	{
 		FString ServerPath;
 		FString ClientPath;
-		uint32 FileId = ~uint32(0);
+		FIoChunkId FileChunkId;
 	};
 	
 	FZenFileSystemManifest(const ITargetPlatform& InTargetPlatform, FString InCookDirectory);
 	
-	void Generate();
+	int32 Generate();
 
-	const FManifestEntry& AddFile(const FString& Filename);
+	const FManifestEntry& CreateManifestEntry(const FString& Filename);
+	
+	const FManifestEntry& AddManifestEntry(const FIoChunkId& FileChunkId, FString ServerPath, FString ClientPath);
 
 	TArrayView<const FManifestEntry> ManifestEntries() const
 	{
 		return Entries;
 	}
 
-	int32 ManifestVersion() const
+	bool Save(const TCHAR* Filename);
+
+	int32 NumEntries() const
 	{
 		return Entries.Num();
 	}
 
-	void Save(const TCHAR* Filename);
-
 private:
-	const FManifestEntry& AddManifestEntry(FString ServerPath, FString ClientPath);
-
 	const ITargetPlatform& TargetPlatform;
 	FString CookDirectory;
 	FString ServerRoot;
diff --git a/Engine/Source/Developer/IoStoreUtilities/Private/ZenStoreHttpClient.cpp b/Engine/Source/Developer/IoStoreUtilities/Private/ZenStoreHttpClient.cpp
--- a/Engine/Source/Developer/IoStoreUtilities/Private/ZenStoreHttpClient.cpp
+++ b/Engine/Source/Developer/IoStoreUtilities/Private/ZenStoreHttpClient.cpp
@@ -2,8 +2,6 @@
 
 #include "ZenStoreHttpClient.h"
 
-#include "Serialization/CompactBinaryPackage.h"
-
 #if PLATFORM_WINDOWS
 
 #if PLATFORM_WINDOWS || PLATFORM_HOLOLENS
@@ -1241,6 +1239,54 @@ TIoStatusOr<FIoBuffer> FZenStoreHttpClient::ReadChunk(const FIoChunkId& Id, uint
 	return FIoStatus(EIoErrorCode::NotFound);
 }
 
+TFuture<TIoStatusOr<FCbObject>> FZenStoreHttpClient::GetOplog()
+{
+	return Async(EAsyncExecution::LargeThreadPool, [this]
+	{
+		UE::Zen::FScopedRequestPtr Request(RequestPool.Get());
+		
+		TStringBuilder<128> Uri;
+		Uri << OplogPath << "/entries";
+
+		TArray64<uint8> GetBuffer;
+		UE::Zen::FRequest::Result Res = Request->PerformBlockingDownload(Uri, &GetBuffer);
+
+		if (Res == Zen::FRequest::Success && Request->GetResponseCode() == 200)
+		{
+			FCbObjectView Response(GetBuffer.GetData());
+			return TIoStatusOr<FCbObject>(FCbObject::Clone(Response));
+		}
+		else
+		{
+			return TIoStatusOr<FCbObject>(FIoStatus(EIoErrorCode::NotFound));
+		}
+	});
+}
+
+TFuture<TIoStatusOr<FCbObject>> FZenStoreHttpClient::GetFiles()
+{
+	return Async(EAsyncExecution::LargeThreadPool, [this]
+	{
+		UE::Zen::FScopedRequestPtr Request(RequestPool.Get());
+		
+		TStringBuilder<128> Uri;
+		Uri << OplogPath << "/files";
+
+		TArray64<uint8> GetBuffer;
+		UE::Zen::FRequest::Result Res = Request->PerformBlockingDownload(Uri, &GetBuffer);
+
+		if (Res == Zen::FRequest::Success && Request->GetResponseCode() == 200)
+		{
+			FCbObjectView Response(GetBuffer.GetData());
+			return TIoStatusOr<FCbObject>(FCbObject::Clone(Response));
+		}
+		else
+		{
+			return TIoStatusOr<FCbObject>(FIoStatus(EIoErrorCode::NotFound));
+		}
+	});
+}
+
 void 
 FZenStoreHttpClient::StartBuildPass()
 {
@@ -1332,6 +1378,16 @@ TFuture<TIoStatusOr<uint64>> FZenStoreHttpClient::AppendOp(FCbPackage OpEntry)
 	return TFuture<TIoStatusOr<uint64>>();
 }
 
+TFuture<TIoStatusOr<FCbObject>> FZenStoreHttpClient::GetOplog()
+{
+	return TFuture<TIoStatusOr<FCbObject>>();
+}
+
+TFuture<TIoStatusOr<FCbObject>> FZenStoreHttpClient::GetFiles()
+{
+	return TFuture<TIoStatusOr<FCbObject>>();
+}
+
 }
 
 #endif // PLATFORM_WINDOWS
diff --git a/Engine/Source/Developer/IoStoreUtilities/Private/ZenStoreHttpClient.h b/Engine/Source/Developer/IoStoreUtilities/Private/ZenStoreHttpClient.h
--- a/Engine/Source/Developer/IoStoreUtilities/Private/ZenStoreHttpClient.h
+++ b/Engine/Source/Developer/IoStoreUtilities/Private/ZenStoreHttpClient.h
@@ -4,6 +4,7 @@
 
 #include "IO/IoDispatcher.h"
 #include "Async/Future.h"
+#include "Serialization/CompactBinaryPackage.h"
 
 class FCbPackage;
 class FCbObject;
@@ -39,6 +40,9 @@ public:
 	TIoStatusOr<uint64> GetChunkSize(const FIoChunkId& Id);
 	TIoStatusOr<FIoBuffer> ReadChunk(const FIoChunkId& Id, uint64 Offset = 0, uint64 Size = ~0ull);
 
+	TFuture<TIoStatusOr<FCbObject>> GetOplog();
+	TFuture<TIoStatusOr<FCbObject>> GetFiles();
+
 private:
 	FString	HostName;
 	uint16 Port;
diff --git a/Engine/Source/Developer/IoStoreUtilities/Private/ZenStoreWriter.cpp b/Engine/Source/Developer/IoStoreUtilities/Private/ZenStoreWriter.cpp
--- a/Engine/Source/Developer/IoStoreUtilities/Private/ZenStoreWriter.cpp
+++ b/Engine/Source/Developer/IoStoreUtilities/Private/ZenStoreWriter.cpp
@@ -27,6 +27,19 @@ FCbObjectId ToObjectId(const FIoChunkId& ChunkId)
 	return FCbObjectId(MakeMemoryView(ChunkId.GetData(), ChunkId.GetSize()));
 }
 
+FMD5Hash IoHashToMD5(const FIoHash& IoHash)
+{
+	const FIoHash::ByteArray& Bytes = IoHash.GetBytes();
+	
+	FMD5 MD5Gen;
+	MD5Gen.Update(Bytes, sizeof(FIoHash::ByteArray));
+	
+	FMD5Hash Hash;
+	Hash.Set(MD5Gen);
+
+	return Hash;
+}
+
 class FZenStoreWriter::FZenStoreHttpQueue
 {
 public:
@@ -247,6 +260,86 @@ FZenStoreWriter::FZenStoreWriter(
 	ZenFileSystemManifest = MakeUnique<FZenFileSystemManifest>(TargetPlatform, OutputPath);
 	
 	HttpQueue = MakeUnique<FZenStoreHttpQueue>(*HttpClient);
+
+	if (!IsCleanBuild)
+	{
+		UE_LOG(LogZenStoreWriter, Display, TEXT("Fetching oplog..."), *ProjectId, *OplogId);
+
+		TFuture<FIoStatus> FutureOplogStatus = HttpClient->GetOplog().Next([this](TIoStatusOr<FCbObject> OplogStatus)
+		{
+			if (!OplogStatus.IsOk())
+			{
+				return OplogStatus.Status();
+			}
+			
+			FCbObject Oplog = OplogStatus.ConsumeValueOrDie();
+			
+			if (Oplog["entries"])
+			{
+				for (FCbField& OplogEntry : Oplog["entries"].AsArray())
+				{
+					FCbObject OplogObj						= OplogEntry.AsObject();
+
+					if (OplogObj["package"])
+					{
+						FCbObject PackageObj				= OplogObj["package"].AsObject();
+
+						const FGuid PkgGuid					= PackageObj["guid"].AsUuid();
+						const FIoHash PkgHash				= PackageObj["data"].AsHash();
+						const int64	PkgDiskSize				= PackageObj["disksize"].AsUInt64();
+						FPackageStoreEntryResource Entry	= FPackageStoreEntryResource::FromCbObject(OplogObj["packagestoreentry"].AsObject());
+						const FName PackageName				= Entry.PackageName;
+
+						const int32 Index					= PackageStoreEntries.Num();
+
+						PackageStoreEntries.Add(MoveTemp(Entry));
+						CookedPackagesInfo.Add(FCookedPackageInfo { PackageName, IoHashToMD5(PkgHash), PkgGuid, PkgDiskSize });
+						PackageNameToIndex.Add(PackageName, Index);
+					}
+				}
+			}
+
+			return FIoStatus::Ok;
+		});
+
+		UE_LOG(LogZenStoreWriter, Display, TEXT("Fetching file manifest..."), *ProjectId, *OplogId);
+
+		TIoStatusOr<FCbObject> FileStatus = HttpClient->GetFiles().Get();
+		if (FileStatus .IsOk())
+		{
+			FCbObject FilesObj = FileStatus.ConsumeValueOrDie();
+			if (FilesObj["files"])
+			{
+				for (FCbField& FileEntry : FilesObj["files"].AsArray())
+				{
+					FCbObject FileObj	= FileEntry.AsObject();
+					FCbObjectId FileId	= FileObj["id"].AsObjectId();
+					FString ServerPath	= FString(FileObj["serverpath"].AsString());
+					FString ClientPath	= FString(FileObj["clientpath"].AsString());
+
+					FIoChunkId FileChunkId;
+					FileChunkId.Set(FileId.GetView());
+					
+					ZenFileSystemManifest->AddManifestEntry(FileChunkId, MoveTemp(ServerPath), MoveTemp(ClientPath));
+				}
+			}
+
+			UE_LOG(LogZenStoreWriter, Display, TEXT("Fetched '%d' files(s) from oplog '%s/%s'"), ZenFileSystemManifest->NumEntries(), *ProjectId, *OplogId);
+		}
+		else
+		{
+			UE_LOG(LogZenStoreWriter, Warning, TEXT("Failed to fetch file(s) from oplog '%s/%s'"), *ProjectId, *OplogId);
+		}
+
+		if (FutureOplogStatus.Get().IsOk())
+		{
+			UE_LOG(LogZenStoreWriter, Display, TEXT("Fetched '%d' packges(s) from oplog '%s/%s'"), PackageStoreEntries.Num(), *ProjectId, *OplogId);
+		}
+		else
+		{
+			UE_LOG(LogZenStoreWriter, Warning, TEXT("Failed to fetch oplog '%s/%s'"), *ProjectId, *OplogId);
+		}
+	}
 }
 
 FZenStoreWriter::~FZenStoreWriter()
@@ -363,13 +456,13 @@ void FZenStoreWriter::WriteBulkdata(const FBulkDataInfo& Info, const FIoBuffer&
 
 bool FZenStoreWriter::WriteAdditionalFile(const FAdditionalFileInfo& Info, const FIoBuffer& FileData)
 {
-	const FZenFileSystemManifest::FManifestEntry& ManifestEntry = ZenFileSystemManifest->AddFile(Info.Filename);
+	const FZenFileSystemManifest::FManifestEntry& ManifestEntry = ZenFileSystemManifest->CreateManifestEntry(Info.Filename);
 
 	FFileDataEntry FileEntry;
 
 	FileEntry.Payload				= FIoBuffer(FIoBuffer::Clone, FileData.Data(), FileData.DataSize());
 	FileEntry.Info					= Info;
-	FileEntry.Info.ChunkId			= CreateExternalFileChunkId(0, ManifestEntry.FileId);
+	FileEntry.Info.ChunkId			= ManifestEntry.FileChunkId;
 	FileEntry.ZenManifestServerPath = ManifestEntry.ServerPath;
 	FileEntry.ZenManifestClientPath = ManifestEntry.ClientPath;
 
@@ -446,7 +539,7 @@ void FZenStoreWriter::EndCook()
 		(double(ZenStats.TotalBytes) / 1024.0 / 1024.0) / ZenStats.TotalRequestTime);
 }
 
-void FZenStoreWriter::BeginPackage(const FPackageBaseInfo& Info)
+void FZenStoreWriter::BeginPackage(const FBeginPackageInfo& Info)
 {
 	FWriteScopeLock _(PackagesLock);
 
@@ -462,7 +555,7 @@ void FZenStoreWriter::BeginPackage(const FPackageBaseInfo& Info)
 	PackageStoreManifest.BeginPackage(Info.PackageName);
 }
 
-void FZenStoreWriter::CommitPackage(const FPackageBaseInfo& Info)
+void FZenStoreWriter::CommitPackage(const FCommitPackageInfo& Info)
 {
 	TRACE_CPUPROFILER_EVENT_SCOPE(FZenStoreWriter::CommitPackage);
 
@@ -473,7 +566,7 @@ void FZenStoreWriter::CommitPackage(const FPackageBaseInfo& Info)
 	CommitEventArgs.EntryIndex		= INDEX_NONE;
 	
 	FCbPackage Pkg;
-	bool bIsValid = false;
+	bool bIsValid = Info.bSucceeded;
 
 	{
 		FWriteScopeLock _(PackagesLock);
@@ -482,18 +575,9 @@ void FZenStoreWriter::CommitPackage(const FPackageBaseInfo& Info)
 
 		checkf(PackageState != nullptr, TEXT("CommitPackage called for package which is not pending: '%s'"), *Info.PackageName.ToString());
 
-		bIsValid = PackageState->PackageData.IsValid;
-
 		if (bIsValid)
 		{
-			FPackageDataEntry& PkgData = PackageState->PackageData;
-
-			CommitEventArgs.EntryIndex = PackageStoreEntries.Num();
-			PackageStoreEntries.Add(PkgData.PackageStoreEntry);
-
-			FCbWriter PackageObj;
-			PackageObj.BeginObject();
-			PackageObj << "key" << Info.PackageName.ToString();
+			check(PackageState->PackageData.IsValid);
 
 			// Note that this is destructive - we yank out the buffer memory from the 
 			// IoBuffer into the FSharedBuffer
@@ -504,14 +588,31 @@ void FZenStoreWriter::CommitPackage(const FPackageBaseInfo& Info)
 				return FSharedBuffer{ FSharedBuffer::TakeOwnership(DataPtr, DataSize, FMemory::Free) };
 			};
 
+			FPackageDataEntry& PkgData = PackageState->PackageData;
+			
+			const int64 PkgDiskSize = PkgData.Payload.DataSize();
+			
 			FCbAttachment PkgDataAttachment(IoBufferToSharedBuffer(PkgData.Payload));
 			Pkg.AddAttachment(PkgDataAttachment);
+			
+			CommitEventArgs.EntryIndex = PackageStoreEntries.Num();
+			PackageStoreEntries.Add(PkgData.PackageStoreEntry);
+			CookedPackagesInfo.Add(FCookedPackageInfo { Info.PackageName, IoHashToMD5(PkgDataAttachment.GetHash()), Info.PackageGuid, PkgDiskSize });
 
+			FCbWriter PackageObj;
+			PackageObj.BeginObject();
+			PackageObj << "key" << Info.PackageName.ToString();
+
+			// NOTE: The package GUID and disk size are used for iterative cooks when comparing asset registry package data
 			PackageObj.BeginObject("package");
 			PackageObj << "id" << PkgData.ChunkId;
+			PackageObj << "guid" << Info.PackageGuid;
 			PackageObj << "data" << PkgDataAttachment;
+			PackageObj << "disksize" << PkgDiskSize; 
 			PackageObj.EndObject();
 
+			PackageObj << "packagestoreentry" << PkgData.PackageStoreEntry;
+			
 			if (PackageState->BulkData.Num())
 			{
 				PackageObj.BeginArray("bulkdata");
@@ -596,32 +697,85 @@ void FZenStoreWriter::Flush()
 	HttpQueue->Flush();
 }
 
+void FZenStoreWriter::GetCookedPackages(TArray<FCookedPackageInfo>& OutCookedPackages) 
+{
+	OutCookedPackages.Append(CookedPackagesInfo);
+}
+
+void FZenStoreWriter::RemoveCookedPackages(TArrayView<const FName> PackageNamesToRemove)
+{
+	TSet<int32> PackageIndicesToKeep;
+	for (int32 Idx = 0, Num = PackageStoreEntries.Num(); Idx < Num; ++Idx)
+	{
+		PackageIndicesToKeep.Add(Idx);
+	}
+	
+	for (const FName& PackageName : PackageNamesToRemove)
+	{
+		if (const int32* Idx = PackageNameToIndex.Find(PackageName))
+		{
+			PackageIndicesToKeep.Remove(*Idx);
+		}
+	}
+
+	const int32 NumPackagesToKeep = PackageIndicesToKeep.Num();
+	
+	TArray<FPackageStoreEntryResource> PreviousPackageStoreEntries = MoveTemp(PackageStoreEntries);
+	TArray<FCookedPackageInfo> PreviousCookedPackageInfo = MoveTemp(CookedPackagesInfo);
+	PackageNameToIndex.Empty(false);
+
+	if (NumPackagesToKeep > 0)
+	{
+		PackageStoreEntries.Reserve(NumPackagesToKeep);
+		CookedPackagesInfo.Reserve(NumPackagesToKeep);
+		PackageNameToIndex.Reserve(NumPackagesToKeep);
+
+		int32 EntryIndex = 0;
+		for (int32 Idx : PackageIndicesToKeep)
+		{
+			const FName PackageName = PreviousCookedPackageInfo[Idx].PackageName;
+
+			PackageStoreEntries.Add(MoveTemp(PreviousPackageStoreEntries[Idx]));
+			CookedPackagesInfo.Add(MoveTemp(PreviousCookedPackageInfo[Idx]));
+			PackageNameToIndex.Add(PackageName, EntryIndex++);
+		}
+	}
+}
+
 void FZenStoreWriter::CreateProjectMetaData(FCbPackage& Pkg, FCbWriter& PackageObj, bool bGenerateContainerHeader)
 {
 	// File Manifest
 	{
-		// Append all project and cooked files
-		ZenFileSystemManifest->Generate();
+		// Only append new file entries to the Oplog
+		
+		int32 NumEntries = ZenFileSystemManifest->NumEntries();
+		const int32 NumNewEntries = ZenFileSystemManifest->Generate();
 
+		if (NumNewEntries > 0)
 		{
+			TArrayView<FZenFileSystemManifest::FManifestEntry const> Entries = ZenFileSystemManifest->ManifestEntries();
+			TArrayView<FZenFileSystemManifest::FManifestEntry const> NewEntries = Entries.Slice(NumEntries, NumNewEntries);
+			
 			PackageObj.BeginArray("files");
 
-			for (const FZenFileSystemManifest::FManifestEntry& ManifestEntry : ZenFileSystemManifest->ManifestEntries())
+			for (const FZenFileSystemManifest::FManifestEntry& NewEntry : NewEntries)
 			{
-				FCbObjectId FileOid = ToObjectId(CreateExternalFileChunkId(0, ManifestEntry.FileId));
+				FCbObjectId FileOid = ToObjectId(NewEntry.FileChunkId);
 
 				PackageObj.BeginObject();
 				PackageObj << "id" << FileOid;
 				PackageObj << "data" << FIoHash::Zero;
-				PackageObj << "serverpath" << ManifestEntry.ServerPath;
-				PackageObj << "clientpath" << ManifestEntry.ClientPath;
+				PackageObj << "serverpath" << NewEntry.ServerPath;
+				PackageObj << "clientpath" << NewEntry.ClientPath;
 				PackageObj.EndObject();
 			}
 
-			PackageObj.EndArray(); // End of files
+			PackageObj.EndArray();
 		}
 
-		ZenFileSystemManifest->Save(*FPaths::Combine(MetadataDirectoryPath, TEXT("zenfs.manifest")));
+		FString ManifestPath = FPaths::Combine(MetadataDirectoryPath, TEXT("zenfs.manifest"));
+		UE_LOG(LogZenStoreWriter, Display, TEXT("Saving Zen filesystem manifest '%s'"), *ManifestPath);
+		ZenFileSystemManifest->Save(*ManifestPath);
 	}
 
 	// Metadata section
diff --git a/Engine/Source/Developer/IoStoreUtilities/Public/FilePackageStoreWriter.h b/Engine/Source/Developer/IoStoreUtilities/Public/FilePackageStoreWriter.h
--- a/Engine/Source/Developer/IoStoreUtilities/Public/FilePackageStoreWriter.h
+++ b/Engine/Source/Developer/IoStoreUtilities/Public/FilePackageStoreWriter.h
@@ -65,8 +65,8 @@ class FFilePackageStoreWriter
 public:
 	IOSTOREUTILITIES_API FFilePackageStoreWriter(const FString& OutputPath, const FString& MetadataDirectoryPath, const ITargetPlatform* TargetPlatform);
 	IOSTOREUTILITIES_API ~FFilePackageStoreWriter();
-	IOSTOREUTILITIES_API virtual void BeginPackage(const FPackageBaseInfo& Info) override;
-	IOSTOREUTILITIES_API virtual void CommitPackage(const FPackageBaseInfo& Info) override;
+	IOSTOREUTILITIES_API virtual void BeginPackage(const FBeginPackageInfo& Info) override;
+	IOSTOREUTILITIES_API virtual void CommitPackage(const FCommitPackageInfo& Info) override;
 	IOSTOREUTILITIES_API virtual void WritePackageData(const FPackageInfo& Info, const FIoBuffer& PackageData, const TArray<FFileRegion>& FileRegions) override;
 	IOSTOREUTILITIES_API virtual void WriteBulkdata(const FBulkDataInfo& Info, const FIoBuffer& BulkData, const TArray<FFileRegion>& FileRegions) override;
 	IOSTOREUTILITIES_API virtual bool WriteAdditionalFile(const FAdditionalFileInfo& Info, const FIoBuffer& FileData) override { return false; }
@@ -89,6 +89,12 @@ public:
 	virtual void Flush() override
 	{ }
 
+	virtual void GetCookedPackages(TArray<FCookedPackageInfo>& OutCookedPackages)
+	{ }
+
+	virtual void RemoveCookedPackages(TArrayView<const FName> PackageNamesToRemove)
+	{ }
+
 private:
 	TUniquePtr<FPackageStoreOptimizer> PackageStoreOptimizer;
 	TUniquePtr<FIoStoreWriterContext> IoStoreWriterContext;
diff --git a/Engine/Source/Developer/IoStoreUtilities/Public/ZenStoreWriter.h b/Engine/Source/Developer/IoStoreUtilities/Public/ZenStoreWriter.h
--- a/Engine/Source/Developer/IoStoreUtilities/Public/ZenStoreWriter.h
+++ b/Engine/Source/Developer/IoStoreUtilities/Public/ZenStoreWriter.h
@@ -34,8 +34,8 @@ public:
 
 	IOSTOREUTILITIES_API ~FZenStoreWriter();
 
-	IOSTOREUTILITIES_API virtual void BeginPackage(const FPackageBaseInfo& Info) override;
-	IOSTOREUTILITIES_API virtual void CommitPackage(const FPackageBaseInfo& Info) override;
+	IOSTOREUTILITIES_API virtual void BeginPackage(const FBeginPackageInfo& Info) override;
+	IOSTOREUTILITIES_API virtual void CommitPackage(const FCommitPackageInfo& Info) override;
 
 	IOSTOREUTILITIES_API virtual void WritePackageData(const FPackageInfo& Info, const FIoBuffer& PackageData, const TArray<FFileRegion>& FileRegions) override;
 	IOSTOREUTILITIES_API virtual bool WriteAdditionalFile(const FAdditionalFileInfo& Info, const FIoBuffer& FileData) override;
@@ -55,6 +55,9 @@ public:
 
 	IOSTOREUTILITIES_API void WriteIoStorePackageData(const FPackageInfo& Info, const FIoBuffer& PackageData, const FPackageStoreEntryResource& PackageStoreEntry, const TArray<FFileRegion>& FileRegions);
 
+	IOSTOREUTILITIES_API virtual void GetCookedPackages(TArray<FCookedPackageInfo>& OutCookedPackages) override;
+	IOSTOREUTILITIES_API virtual void RemoveCookedPackages(TArrayView<const FName> PackageNamesToRemove);
+
 private:
 	void CreateProjectMetaData(FCbPackage& Pkg, FCbWriter& PackageObj, bool bGenerateContainerHeader);
 	void BroadcastCommit(IPackageStoreWriter::FCommitEventArgs& EventArgs);
@@ -110,6 +113,9 @@ private:
 	FPackageStoreManifest				PackageStoreManifest;
 	TUniquePtr<FPackageStoreOptimizer>	PackageStoreOptimizer;
 	TArray<FPackageStoreEntryResource>	PackageStoreEntries;
+	TArray<FCookedPackageInfo>			CookedPackagesInfo;
+	TMap<FName, int32>					PackageNameToIndex;
+
 	TUniquePtr<FZenFileSystemManifest>	ZenFileSystemManifest;
 	
 	FCriticalSection					CommitEventCriticalSection;
diff --git a/Engine/Source/Editor/UnrealEd/Classes/CookOnTheSide/CookOnTheFlyServer.h b/Engine/Source/Editor/UnrealEd/Classes/CookOnTheSide/CookOnTheFlyServer.h
--- a/Engine/Source/Editor/UnrealEd/Classes/CookOnTheSide/CookOnTheFlyServer.h
+++ b/Engine/Source/Editor/UnrealEd/Classes/CookOnTheSide/CookOnTheFlyServer.h
@@ -27,6 +27,7 @@ class FSavePackageContext;
 class ITargetPlatform;
 class IAssetRegistry;
 class IPlugin;
+class IPackageStoreWriter;
 
 enum class ECookInitializationFlags
 {
@@ -1052,6 +1053,11 @@ private:
 	* @param Platforms to process
 	*/
 	void PopulateCookedPackagesFromDisk(const TArrayView<const ITargetPlatform* const>& Platforms);
+	
+	/**
+	 * Populate cooked packages from meta data available in the package store.
+	 */
+	void PopulateCookedPackagesFromPackageStore(const TArrayView<const ITargetPlatform* const>& Platforms);
 
 	/**
 	* Searches the disk for all the cooked files in the sandbox path provided
@@ -1107,6 +1113,8 @@ private:
 
 	uint32 FullLoadAndSave(uint32& CookedPackageCount);
 
+	IPackageStoreWriter* GetPackageStoreWriter(const FName& PlatformName) const;
+
 	uint32		StatLoadedPackageCount = 0;
 	uint32		StatSavedPackageCount = 0;
 
diff --git a/Engine/Source/Editor/UnrealEd/Private/Commandlets/AssetRegistryGenerator.cpp b/Engine/Source/Editor/UnrealEd/Private/Commandlets/AssetRegistryGenerator.cpp
--- a/Engine/Source/Editor/UnrealEd/Private/Commandlets/AssetRegistryGenerator.cpp
+++ b/Engine/Source/Editor/UnrealEd/Private/Commandlets/AssetRegistryGenerator.cpp
@@ -1065,7 +1065,7 @@ void FAssetRegistryGenerator::Initialize(const TArray<FName> &InStartupPackages)
 	FGameDelegates::Get().GetAssignLayerChunkDelegate() = FAssignLayerChunkDelegate::CreateStatic(AssignLayerChunkDelegate);
 }
 
-void FAssetRegistryGenerator::ComputePackageDifferences(TSet<FName>& ModifiedPackages, TSet<FName>& NewPackages, TSet<FName>& RemovedPackages, TSet<FName>& IdenticalCookedPackages, TSet<FName>& IdenticalUncookedPackages, bool bRecurseModifications, bool bRecurseScriptModifications)
+void FAssetRegistryGenerator::ComputePackageDifferences(const FComputeDifferenceOptions& Options, const TMap<FName, const FAssetPackageData*>& PreviousAssetPackageDataMap, FAssetRegistryDifference& OutDifference)
 {
 	TArray<FName> ModifiedScriptPackages;
 
@@ -1074,23 +1074,23 @@ void FAssetRegistryGenerator::ComputePackageDifferences(TSet<FName>& ModifiedPac
 		FName PackageName = PackagePair.Key;
 		const FAssetPackageData* CurrentPackageData = PackagePair.Value;
 
-		const FAssetPackageData* PreviousPackageData = PreviousState.GetAssetPackageData(PackageName);
+		const FAssetPackageData* PreviousPackageData = PreviousAssetPackageDataMap.FindRef(PackageName);
 
 		if (!PreviousPackageData)
 		{
-			NewPackages.Add(PackageName);
+			OutDifference.NewPackages.Add(PackageName);
 		}
 		PRAGMA_DISABLE_DEPRECATION_WARNINGS
 		else if (CurrentPackageData->PackageGuid == PreviousPackageData->PackageGuid)
 		PRAGMA_ENABLE_DEPRECATION_WARNINGS
 		{
-			if (PreviousPackageData->DiskSize < 0)
+			if (!Options.bIgnoreDiskSize && PreviousPackageData->DiskSize < 0)
 			{
-				IdenticalUncookedPackages.Add(PackageName);
+				OutDifference.IdenticalUncookedPackages.Add(PackageName);
 			}
 			else
 			{
-				IdenticalCookedPackages.Add(PackageName);
+				OutDifference.IdenticalCookedPackages.Add(PackageName);
 			}
 		}
 		else
@@ -1101,12 +1101,12 @@ void FAssetRegistryGenerator::ComputePackageDifferences(TSet<FName>& ModifiedPac
 			}
 			else
 			{
-			ModifiedPackages.Add(PackageName);
+				OutDifference.ModifiedPackages.Add(PackageName);
+			}
 		}
 	}
-	}
 
-	for (const TPair<FName, const FAssetPackageData*>& PackagePair : PreviousState.GetAssetPackageDataMap())
+	for (const TPair<FName, const FAssetPackageData*>& PackagePair : PreviousAssetPackageDataMap)
 	{
 		FName PackageName = PackagePair.Key;
 		const FAssetPackageData* PreviousPackageData = PackagePair.Value;
@@ -1115,16 +1115,16 @@ void FAssetRegistryGenerator::ComputePackageDifferences(TSet<FName>& ModifiedPac
 
 		if (!CurrentPackageData)
 		{
-			RemovedPackages.Add(PackageName);
+			OutDifference.RemovedPackages.Add(PackageName);
 		}
 	}
 
-	if (bRecurseModifications)
+	if (Options.bRecurseModifications)
 	{
 		// Recurse modified packages to their dependencies. This is needed because we only compare package guids
-		TArray<FName> ModifiedPackagesToRecurse = ModifiedPackages.Array();
+		TArray<FName> ModifiedPackagesToRecurse = OutDifference.ModifiedPackages.Array();
 
-		if (bRecurseScriptModifications)
+		if (Options.bRecurseScriptModifications)
 		{
 			ModifiedPackagesToRecurse.Append(ModifiedScriptPackages);
 		}
@@ -1138,13 +1138,13 @@ void FAssetRegistryGenerator::ComputePackageDifferences(TSet<FName>& ModifiedPac
 			for (const FAssetIdentifier& Referencer : Referencers)
 			{
 				FName ReferencerPackage = Referencer.PackageName;
-				if (!ModifiedPackages.Contains(ReferencerPackage) && (IdenticalCookedPackages.Contains(ReferencerPackage) || IdenticalUncookedPackages.Contains(ReferencerPackage)))
+				if (!OutDifference.ModifiedPackages.Contains(ReferencerPackage) && (OutDifference.IdenticalCookedPackages.Contains(ReferencerPackage) || OutDifference.IdenticalUncookedPackages.Contains(ReferencerPackage)))
 				{
 					// Remove from identical list
-					IdenticalCookedPackages.Remove(ReferencerPackage);
-					IdenticalUncookedPackages.Remove(ReferencerPackage);
+					OutDifference.IdenticalCookedPackages.Remove(ReferencerPackage);
+					OutDifference.IdenticalUncookedPackages.Remove(ReferencerPackage);
 
-					ModifiedPackages.Add(ReferencerPackage);
+					OutDifference.ModifiedPackages.Add(ReferencerPackage);
 					ModifiedPackagesToRecurse.Add(ReferencerPackage);
 				}
 			}
@@ -1152,6 +1152,22 @@ void FAssetRegistryGenerator::ComputePackageDifferences(TSet<FName>& ModifiedPac
 	}
 }
 
+void FAssetRegistryGenerator::ComputePackageDifferences(TSet<FName>& ModifiedPackages, TSet<FName>& NewPackages, TSet<FName>& RemovedPackages, TSet<FName>& IdenticalCookedPackages, TSet<FName>& IdenticalUncookedPackages, bool bRecurseModifications, bool bRecurseScriptModifications)
+{
+	FComputeDifferenceOptions Options;
+	Options.bRecurseModifications		= bRecurseModifications;
+	Options.bRecurseScriptModifications	= bRecurseScriptModifications;
+	
+	FAssetRegistryDifference Difference;
+	ComputePackageDifferences(Options, PreviousState.GetAssetPackageDataMap(), Difference);
+
+	ModifiedPackages			= MoveTemp(Difference.ModifiedPackages);
+	NewPackages					= MoveTemp(Difference.NewPackages);
+	RemovedPackages				= MoveTemp(Difference.RemovedPackages);
+	IdenticalCookedPackages		= MoveTemp(Difference.IdenticalCookedPackages);
+	IdenticalUncookedPackages	= MoveTemp(Difference.IdenticalUncookedPackages);
+}
+
 void FAssetRegistryGenerator::UpdateKeptPackages(const TArray<FName>& InKeptPackages)
 {
 	KeptPackages.Append(InKeptPackages);
@@ -1160,6 +1176,21 @@ void FAssetRegistryGenerator::UpdateKeptPackages(const TArray<FName>& InKeptPack
 	// Delay update of AssetData with TagsAndValues, this data may be modified up until serialization in SaveAssetRegistry
 }
 
+void FAssetRegistryGenerator::UpdateKeptPackages(const TArray<FName>& InKeptPackages, TFunction<void(const FName&, FAssetPackageData*)>&& Update)
+{
+	KeptPackages.Append(InKeptPackages);
+
+	for (const FName& PackageName : InKeptPackages)
+	{
+		// Get mutable PackageData without creating it when it does not exist
+		FAssetPackageData* PackageData = State.GetAssetPackageData(PackageName)
+			? State.CreateOrGetAssetPackageData(PackageName)
+			: nullptr;
+
+		Update(PackageName, PackageData);
+	}
+}
+
 void FAssetRegistryGenerator::BuildChunkManifest(const TSet<FName>& InCookedPackages, const TSet<FName>& InDevelopmentOnlyPackages, FSandboxPlatformFile* InSandboxFile, bool bGenerateStreamingInstallManifest)
 {
 	// If we were asked to generate a streaming install manifest explicitly and we did not have bGenerateNoChunks set, we will generate chunks.
diff --git a/Engine/Source/Editor/UnrealEd/Private/CookOnTheFlyServer.cpp b/Engine/Source/Editor/UnrealEd/Private/CookOnTheFlyServer.cpp
--- a/Engine/Source/Editor/UnrealEd/Private/CookOnTheFlyServer.cpp
+++ b/Engine/Source/Editor/UnrealEd/Private/CookOnTheFlyServer.cpp
@@ -594,18 +594,7 @@ public:
 
 	virtual IPackageStoreWriter* GetPackageStoreWriter(const FName& PlatformName) override
 	{
-		if (Cooker.CookOnTheFlyOptions.bIoStore)
-		{
-			for (FSavePackageContext* Context : Cooker.SavePackageContexts)
-			{
-				if (FName(*Context->TargetPlatform->PlatformName()) == PlatformName)
-				{
-					return Context->PackageStoreWriter;
-				}
-			}
-		}
-
-		return nullptr;
+		return Cooker.GetPackageStoreWriter(PlatformName);
 	}
 
 	DECLARE_DERIVED_EVENT(UCookOnTheFlyServer::FCookOnTheFlyServerInterface, UE::Cook::ICookOnTheFlyServer::FFlushEvent, FFlushEvent);
@@ -4643,12 +4632,12 @@ void UCookOnTheFlyServer::SaveCookedPackage(UE::Cook::FPackageData& PackageData,
 					{
 						FSavePackageContext* const SavePackageContext = SavePackageContexts.Num() > 0 ? SavePackageContexts[PlatformIndex] : nullptr;
 
-						IPackageStoreWriter::FPackageBaseInfo Info;
-						Info.PackageName = Package->GetFName();
-
 						IPackageStoreWriter* PackageStoreWriter = SavePackageContext ? SavePackageContext->PackageStoreWriter : nullptr;
 						if (PackageStoreWriter)
 						{
+							IPackageStoreWriter::FBeginPackageInfo Info;
+							Info.PackageName = Package->GetFName();
+							
 							PackageStoreWriter->BeginPackage(Info);
 						}
 
@@ -4659,6 +4648,17 @@ void UCookOnTheFlyServer::SaveCookedPackage(UE::Cook::FPackageData& PackageData,
 
 						if (PackageStoreWriter)
 						{
+							FAssetRegistryGenerator* Generator = PlatformManager->GetPlatformData(Target)->RegistryGenerator.Get();
+							FAssetPackageData* AssetPackageData = Generator->GetAssetPackageData(Package->GetFName());
+							check(AssetPackageData);
+
+							IPackageStoreWriter::FCommitPackageInfo Info;
+							Info.bSucceeded		= Result.IsSuccessful();
+							Info.PackageName	= Package->GetFName();
+							PRAGMA_DISABLE_DEPRECATION_WARNINGS
+							Info.PackageGuid	= AssetPackageData->PackageGuid;
+							PRAGMA_ENABLE_DEPRECATION_WARNINGS
+
 							PackageStoreWriter->CommitPackage(Info);
 						}
 					}
@@ -6172,6 +6172,100 @@ void UCookOnTheFlyServer::PopulateCookedPackagesFromDisk(const TArrayView<const
 	}
 }
 
+void UCookOnTheFlyServer::PopulateCookedPackagesFromPackageStore(const TArrayView<const ITargetPlatform* const>& Platforms)
+{
+	TRACE_CPUPROFILER_EVENT_SCOPE(UCookOnTheFlyServer::PopulateCookedPackagesFromPackgeStore);
+
+	struct FPreviousAssetPackageData
+	{
+		FPreviousAssetPackageData(const TArray<IPackageStoreWriter::FCookedPackageInfo>& InCookedPackages)
+		{
+			AssetPackageDatas.SetNum(InCookedPackages.Num());
+			AssetPackageDataMap.Reserve(InCookedPackages.Num());
+
+			for (int32 Idx = 0, Count = InCookedPackages.Num(); Idx < Count; ++Idx)
+			{
+				const IPackageStoreWriter::FCookedPackageInfo& CookedPackageInfo = InCookedPackages[Idx];
+				FAssetPackageData& AssetPackageData = AssetPackageDatas[Idx];
+				
+				AssetPackageData.CookedHash = CookedPackageInfo.Hash;
+				PRAGMA_DISABLE_DEPRECATION_WARNINGS
+				AssetPackageData.PackageGuid = CookedPackageInfo.PackageGuid;
+				PRAGMA_ENABLE_DEPRECATION_WARNINGS
+				AssetPackageData.DiskSize = CookedPackageInfo.DiskSize;
+
+				AssetPackageDataMap.Add(CookedPackageInfo.PackageName, &AssetPackageData);
+			}
+		}
+
+		TArray<FAssetPackageData> AssetPackageDatas;
+		TMap<FName, const FAssetPackageData*> AssetPackageDataMap;
+	};
+
+	for (const ITargetPlatform* Platform : Platforms)
+	{
+		UE_LOG(LogCook, Display, TEXT("Populating cooked package(s) from package store on platform '%s'"), *Platform->PlatformName());
+		
+		UE::Cook::FPlatformData* PlatformData = PlatformManager->GetPlatformData(Platform);
+		FAssetRegistryGenerator* PlatformAssetRegistry = PlatformData->RegistryGenerator.Get();
+		IPackageStoreWriter* PackageStoreWriter = GetPackageStoreWriter(FName(*Platform->PlatformName()));
+		
+		check(PlatformData);
+		check(PlatformAssetRegistry);
+		check(PackageStoreWriter);
+
+		TArray<IPackageStoreWriter::FCookedPackageInfo> CookedPackages;
+		PackageStoreWriter->GetCookedPackages(CookedPackages);
+		
+		UE_LOG(LogCook, Display, TEXT("Found '%d' cooked package(s) in package store"), CookedPackages.Num());
+
+		FPreviousAssetPackageData PreviousAssetPackageData(CookedPackages);
+
+		FAssetRegistryGenerator::FComputeDifferenceOptions Options;
+		Options.bRecurseModifications		= true;
+		Options.bRecurseScriptModifications	= !IsCookFlagSet(ECookInitializationFlags::IgnoreScriptPackagesOutOfDate);
+		
+		FAssetRegistryGenerator::FAssetRegistryDifference Difference;
+		PlatformAssetRegistry->ComputePackageDifferences(Options, PreviousAssetPackageData.AssetPackageDataMap, Difference);
+
+		TArray<FName> PackagesToRemove;
+		PackagesToRemove.Reserve(Difference.ModifiedPackages.Num() + Difference.RemovedPackages.Num());
+		
+		for (const FName& ModifiedPackageName : Difference.ModifiedPackages)
+		{
+			PackagesToRemove.Add(ModifiedPackageName);
+		}
+		
+		for (const FName& RemovedPackageName : Difference.RemovedPackages)
+		{
+			PackagesToRemove.Add(RemovedPackageName );
+		}
+		
+		UE_LOG(LogCook, Display, TEXT("Keeping '%d' and removing '%d' cooked package(s)"), Difference.IdenticalCookedPackages.Num(), PackagesToRemove.Num());
+		
+		PackageStoreWriter->RemoveCookedPackages(PackagesToRemove);
+		
+		TArray<FName> PackagesToKeep;
+		for (const FName& IdenticalPackageName : Difference.IdenticalCookedPackages)
+		{
+			UE::Cook::FPackageData* PackageData = PackageDatas->TryAddPackageDataByPackageName(IdenticalPackageName);
+			if (PackageData)
+			{
+				const bool bSucceded = true;
+				PackageData->AddCookedPlatforms({ Platform }, bSucceded);
+				PackagesToKeep.Add(IdenticalPackageName);
+			}
+		}
+
+		PlatformAssetRegistry->UpdateKeptPackages(PackagesToKeep, [&PreviousAssetPackageData](const FName& PackageName, FAssetPackageData* Current)
+		{
+			const FAssetPackageData* Previous = PreviousAssetPackageData.AssetPackageDataMap.FindRef(PackageName);
+			Current->CookedHash = Previous->CookedHash;
+			Current->DiskSize = Previous->DiskSize;
+		});
+	}
+}
+
 const FString ExtractPackageNameFromObjectPath( const FString ObjectPath )
 {
 	// get the path 
@@ -7665,9 +7759,11 @@ void UCookOnTheFlyServer::InitializeSandbox(const TArrayView<const ITargetPlatfo
 		}
 
 		// Don't populate platforms that were already initialized; we already populated them when we first initialized them
-		RefreshPlatforms.RemoveAllSwap([&AlreadyInitializedPlatforms](const ITargetPlatform* TargetPlatform) {
+		RefreshPlatforms.RemoveAllSwap([&AlreadyInitializedPlatforms](const ITargetPlatform* TargetPlatform)
+		{
 			return AlreadyInitializedPlatforms.Contains(TargetPlatform);
-			});
+		});
+
 		if (RefreshPlatforms.Num() != 0)
 		{
 			for (UE::Cook::FPackageData* PackageData : *PackageDatas.Get())
@@ -7676,7 +7772,7 @@ void UCookOnTheFlyServer::InitializeSandbox(const TArrayView<const ITargetPlatfo
 			}
 			// The asset registry makes populating cooked packages from disk fast, and populating is a performance benefit
 			// Don't populate however if we are looking for deterministic cooking differences; start from an empty list of cooked packages
-			if (!bIsDiffOnly) 
+			if (!bIsDiffOnly && !IsUsingIoStore()) 
 			{
 				PopulateCookedPackagesFromDisk(RefreshPlatforms);
 			}
@@ -7731,10 +7827,6 @@ void UCookOnTheFlyServer::InitializePackageStore(const TArrayView<const ITargetP
 			{
 				PackageStoreWriter = new FFilePackageStoreWriter(ResolvedRootPath, MetadataDirectoryPath, TargetPlatform);
 			}
-
-			IPackageStoreWriter::FCookInfo CookInfo;
-			CookInfo.CookMode = IsCookOnTheFlyMode() ? IPackageStoreWriter::FCookInfo::CookOnTheFlyMode : IPackageStoreWriter::FCookInfo::CookByTheBookMode;
-			PackageStoreWriter->BeginCook(CookInfo);
 		}
 
 		FConfigFile PlatformEngineIni;
@@ -7746,6 +7838,21 @@ void UCookOnTheFlyServer::InitializePackageStore(const TArrayView<const ITargetP
 		FSavePackageContext* SavePackageContext = new FSavePackageContext(TargetPlatform, PackageStoreWriter, bLegacyBulkDataOffsets);
 		SavePackageContexts.Add(SavePackageContext);
 	}
+
+	if (IsUsingIoStore())
+	{
+		if (!bIsCleanBuild)
+		{
+			PopulateCookedPackagesFromPackageStore(TargetPlatforms);
+		}
+
+		for (const FSavePackageContext* Context : SavePackageContexts)
+		{
+			IPackageStoreWriter::FCookInfo CookInfo;
+			CookInfo.CookMode = IsCookOnTheFlyMode() ? IPackageStoreWriter::FCookInfo::CookOnTheFlyMode : IPackageStoreWriter::FCookInfo::CookByTheBookMode;
+			Context->PackageStoreWriter->BeginCook(CookInfo);
+		}
+	}
 }
 
 void UCookOnTheFlyServer::FinalizePackageStore()
@@ -8027,7 +8134,6 @@ void UCookOnTheFlyServer::StartCookByTheBook( const FCookByTheBookStartupOptions
 	// This will either delete the sandbox or iteratively clean it
 	InitializeSandbox(TargetPlatforms);
 	InitializeTargetPlatforms(TargetPlatforms);
-
 	InitializePackageStore(TargetPlatforms);
 
 	if (CurrentCookMode == ECookMode::CookByTheBook && !IsCookFlagSet(ECookInitializationFlags::Iterative))
@@ -9094,4 +9200,17 @@ uint32 UCookOnTheFlyServer::FullLoadAndSave(uint32& CookedPackageCount)
 	return Result;
 }
 
+IPackageStoreWriter* UCookOnTheFlyServer::GetPackageStoreWriter(const FName& PlatformName) const
+{
+	for (const FSavePackageContext* Context : SavePackageContexts)
+	{
+		if (FName(*Context->TargetPlatform->PlatformName()) == PlatformName)
+		{
+			return Context->PackageStoreWriter;
+		}
+	}
+
+	return nullptr;
+}
+
 #undef LOCTEXT_NAMESPACE
diff --git a/Engine/Source/Editor/UnrealEd/Private/Cooker/IoStoreCookOnTheFlyRequestManager.cpp b/Engine/Source/Editor/UnrealEd/Private/Cooker/IoStoreCookOnTheFlyRequestManager.cpp
--- a/Engine/Source/Editor/UnrealEd/Private/Cooker/IoStoreCookOnTheFlyRequestManager.cpp
+++ b/Engine/Source/Editor/UnrealEd/Private/Cooker/IoStoreCookOnTheFlyRequestManager.cpp
@@ -12,8 +12,6 @@
 #include "UObject/SavePackage.h"
 #include "ProfilingDebugging/CountersTrace.h"
 
-PRAGMA_DISABLE_OPTIMIZATION
-
 DEFINE_LOG_CATEGORY_STATIC(LogCookOnTheFlyTracker, Log, All);
 
 class FCookOnTheFlyPackageTracker
@@ -189,6 +187,27 @@ public:
 		return MoveTemp(CompletedPackages);
 	}
 
+	void AddCompletedPackages(TArrayView<const FPackageStoreEntryResource> Entries)
+	{
+		Packages.Reserve(Entries.Num());
+		CookedEntries.Reserve(Entries.Num());
+
+		for (int32 EntryIndex = 0, Num = Entries.Num(); EntryIndex < Num; ++EntryIndex)
+		{
+			const FPackageStoreEntryResource& Entry = Entries[EntryIndex];
+			const FPackageId PackageId = Entry.GetPackageId();
+
+			FPackage& Package	= GetPackage(PackageId);
+
+			Package.PackageId	= PackageId;
+			Package.PackageName = Entry.PackageName;
+			Package.Status		= EPackageStatus::Cooked;
+			Package.EntryIndex	= EntryIndex;
+
+			CookedEntries.Add(Package.EntryIndex);
+		}
+	}
+
 	FCompletedPackages Flush()
 	{
 		UE_LOG(LogCookOnTheFlyTracker, Log, TEXT("Flushing '%d' waiting package(s)"), WaitingPackages.Num());
@@ -425,8 +444,16 @@ private:
 				{
 					TUniquePtr<FPlatformContext>& Context = PlatformContexts.Add(Client.PlatformName, MakeUnique<FPlatformContext>());
 					Context->PlatformName = Client.PlatformName;
-					Context->PackageStoreWriter = CookOnTheFlyServer.GetPackageStoreWriter(Client.PlatformName);
-					Context->PackageStoreWriter->OnCommit().AddRaw(this, &FIoStoreCookOnTheFlyRequestManager::OnPackageCooked); 
+					
+					IPackageStoreWriter* PackageStoreWriter = CookOnTheFlyServer.GetPackageStoreWriter(Client.PlatformName);
+					Context->PackageStoreWriter = PackageStoreWriter;
+					
+					PackageStoreWriter->GetEntries([&Context](TArrayView<const FPackageStoreEntryResource> Entries)
+					{
+						Context->PackageTracker.AddCompletedPackages(Entries);
+					});
+
+					PackageStoreWriter->OnCommit().AddRaw(this, &FIoStoreCookOnTheFlyRequestManager::OnPackageCooked);
 				}
 
 				return true;
@@ -747,5 +774,3 @@ TUniquePtr<ICookOnTheFlyRequestManager> MakeIoStoreCookOnTheFlyRequestManager(IC
 }
 
 }}
-
-PRAGMA_ENABLE_OPTIMIZATION
diff --git a/Engine/Source/Editor/UnrealEd/Public/Commandlets/AssetRegistryGenerator.h b/Engine/Source/Editor/UnrealEd/Public/Commandlets/AssetRegistryGenerator.h
--- a/Engine/Source/Editor/UnrealEd/Public/Commandlets/AssetRegistryGenerator.h
+++ b/Engine/Source/Editor/UnrealEd/Public/Commandlets/AssetRegistryGenerator.h
@@ -43,6 +43,37 @@ public:
 	 * Loads asset registry from a previous run that is used for iterative or DLC cooking
 	 */
 	bool LoadPreviousAssetRegistry(const FString& Filename);
+	
+	/**
+	 * Options when computing the differences between current and previous state.
+	 */
+	struct FComputeDifferenceOptions
+	{
+		bool bRecurseModifications;
+		bool bRecurseScriptModifications;
+		bool bIgnoreDiskSize = false;
+	};
+
+	/**
+	 * Differences between the current and the previous state.
+	 */
+	struct FAssetRegistryDifference
+	{
+		TSet<FName> ModifiedPackages;
+		TSet<FName> NewPackages;
+		TSet<FName> RemovedPackages;
+		TSet<FName> IdenticalCookedPackages;
+		TSet<FName> IdenticalUncookedPackages;
+	};
+
+	/**
+	 * Computes differences between the current asset registry state and the provided previous state.
+	 *
+	 * @param Options options to use when computing the differences
+	 * @param PreviousAssetPackageDataMap previously cooked asset package data
+	 * @param OutDifference the differences between the current and the previous state
+	 */
+	void ComputePackageDifferences(const FComputeDifferenceOptions& Options, const TMap<FName, const FAssetPackageData*>& PreviousAssetPackageDataMap, FAssetRegistryDifference& OutDifference);
 
 	/**
 	 * Computes differences between the previous asset registry and the current one 
@@ -62,6 +93,7 @@ public:
 	 * Updates the current asset registry from the previous one for all kept packages.
 	 */
 	void UpdateKeptPackages(const TArray<FName>& InKeptPackages);
+	void UpdateKeptPackages(const TArray<FName>& InKeptPackages, TFunction<void(const FName&, FAssetPackageData*)>&& Update);
 
 	/**
 	 * GenerateChunkManifest 
@@ -408,4 +440,4 @@ private:
 	 * Helper function to find or create asset data for the input object. If the asset is not in the registry it will be added.
 	 */
 	const FAssetData* CreateOrFindAssetData(UObject& Object);
-};
\ No newline at end of file
+};
diff --git a/Engine/Source/Runtime/Core/Private/IO/PackageStore.cpp b/Engine/Source/Runtime/Core/Private/IO/PackageStore.cpp
--- a/Engine/Source/Runtime/Core/Private/IO/PackageStore.cpp
+++ b/Engine/Source/Runtime/Core/Private/IO/PackageStore.cpp
@@ -2,6 +2,8 @@
 
 #include "IO/PackageStore.h"
 #include "Serialization/StructuredArchive.h"
+#include "Serialization/CompactBinary.h"
+#include "Serialization/CompactBinaryWriter.h"
 
 FArchive& operator<<(FArchive& Ar, FPackageStoreExportInfo& ExportInfo)
 {
@@ -13,6 +15,30 @@ FArchive& operator<<(FArchive& Ar, FPackageStoreExportInfo& ExportInfo)
 	return Ar;
 }
 
+FCbWriter& operator<<(FCbWriter& Writer, const FPackageStoreExportInfo& ExportInfo)
+{
+	Writer.BeginObject();
+	Writer << "exportbundlessize" << ExportInfo.ExportBundlesSize;
+	Writer << "exportcount" << ExportInfo.ExportCount;
+	Writer << "exportbundlecount" << ExportInfo.ExportBundleCount;
+	Writer << "loadorder" << ExportInfo.LoadOrder;
+	Writer.EndObject();
+	
+	return Writer;
+}
+
+FPackageStoreExportInfo FPackageStoreExportInfo::FromCbObject(const FCbObject& Obj)
+{
+	FPackageStoreExportInfo ExportInfo;
+	
+	ExportInfo.ExportBundlesSize	= Obj["exportbundlessize"].AsUInt64();
+	ExportInfo.ExportCount			= Obj["exportcount"].AsInt32();
+	ExportInfo.ExportBundleCount	= Obj["exportbundlecount"].AsInt32();
+	ExportInfo.LoadOrder			= Obj["loadorder"].AsUInt32();
+
+	return ExportInfo;
+}
+
 FArchive& operator<<(FArchive& Ar, FPackageStoreEntryResource& PackageStoreEntry)
 {
 	uint32 Flags = static_cast<uint32>(PackageStoreEntry.Flags);
@@ -31,3 +57,69 @@ FArchive& operator<<(FArchive& Ar, FPackageStoreEntryResource& PackageStoreEntry
 
 	return Ar;
 }
+
+FCbWriter& operator<<(FCbWriter& Writer, const FPackageStoreEntryResource& PackageStoreEntry)
+{
+	Writer.BeginObject();
+
+	Writer << "flags" << static_cast<uint32>(PackageStoreEntry.Flags);
+	Writer << "packagename" << PackageStoreEntry.PackageName.ToString();
+	Writer << "sourcepackagename" << PackageStoreEntry.SourcePackageName.ToString();
+	Writer << "region" << PackageStoreEntry.Region.ToString();
+	Writer << "exportinfo" << PackageStoreEntry.ExportInfo;
+
+	if (PackageStoreEntry.ImportedPackageIds.Num())
+	{
+		Writer.BeginArray("importedpackageids");
+		for (const FPackageId& ImportedPackageId : PackageStoreEntry.ImportedPackageIds)
+		{
+			Writer << ImportedPackageId.Value();
+		}
+		Writer.EndArray();
+	}
+
+	if (PackageStoreEntry.ShaderMapHashes.Num())
+	{
+		Writer.BeginArray("shadermaphashes");
+		for (const FSHAHash& ShaderMapHash : PackageStoreEntry.ShaderMapHashes)
+		{
+			Writer << ShaderMapHash.ToString();
+		}
+		Writer.EndArray();
+	}
+
+	Writer.EndObject();
+
+	return Writer;
+}
+
+FPackageStoreEntryResource FPackageStoreEntryResource::FromCbObject(const FCbObject& Obj)
+{
+	FPackageStoreEntryResource Entry;
+
+	Entry.Flags				= static_cast<EPackageStoreEntryFlags>(Obj["flags"].AsUInt32());
+	Entry.PackageName		= FName(Obj["packagename"].AsString());
+	Entry.SourcePackageName	= FName(Obj["sourcepackagename"].AsString());
+	Entry.Region			= FName(Obj["region"].AsString());
+	Entry.ExportInfo		= FPackageStoreExportInfo::FromCbObject(Obj["exportinfo"].AsObject());
+	
+	if (Obj["importedpackageids"])
+	{
+		for (FCbFieldView ArrayField : Obj["importedpackageids"])
+		{
+			Entry.ImportedPackageIds.Add(FPackageId::FromValue(ArrayField .AsUInt64()));
+		}
+	}
+	
+	if (Obj["shadermaphashes"])
+	{
+		for (FCbField& ArrayField : Obj["shadermaphashes"].AsArray())
+		{
+			FSHAHash& ShaderMapHash = Entry.ShaderMapHashes.AddDefaulted_GetRef();
+			FAnsiStringView AnsiStringView = ArrayField.AsString();
+			ShaderMapHash.FromString(FStringView(ANSI_TO_TCHAR(AnsiStringView.GetData()), AnsiStringView.Len()));
+		}
+	}
+
+	return Entry;
+}
diff --git a/Engine/Source/Runtime/Core/Public/IO/IoDispatcher.h b/Engine/Source/Runtime/Core/Public/IO/IoDispatcher.h
--- a/Engine/Source/Runtime/Core/Public/IO/IoDispatcher.h
+++ b/Engine/Source/Runtime/Core/Public/IO/IoDispatcher.h
@@ -634,6 +634,31 @@ private:
 	uint8	Hash[32];
 };
 
+/**
+ * Addressable chunk types.
+ * 
+ * The enumerators have explicitly defined values here to encourage backward/forward
+ * compatible updates. 
+ * 
+ * Also note that for certain discriminators, Zen Store will assume certain things
+ * about the structure of the chunk ID so changes must be made carefully.
+ * 
+ */
+enum class EIoChunkType : uint8
+{
+	Invalid = 0,
+	ExportBundleData = 1,
+	BulkData = 2,
+	OptionalBulkData = 3,
+	MemoryMappedBulkData = 4,
+	ScriptObjects = 5,
+	ContainerHeader = 6,
+	ExternalFile = 7,
+	ShaderCodeLibrary = 8,
+	ShaderCode = 9,
+	PackageStoreEntry = 10
+};
+
 /**
  * Identifier to a chunk of data.
  */
@@ -695,6 +720,11 @@ public:
 	inline const uint8* GetData() const { return Id; }
 	inline uint32		GetSize() const { return sizeof Id; }
 
+	EIoChunkType GetChunkType() const
+	{
+		return static_cast<EIoChunkType>(Id[11]);
+	}
+
 private:
 	static inline FIoChunkId CreateEmptyId()
 	{
@@ -708,29 +738,6 @@ private:
 	uint8	Id[12];
 };
 
-/**
- * Addressable chunk types.
- * 
- * The enumerators have explicitly defined values here to encourage backward/forward
- * compatible updates. 
- * 
- * Also note that for certain discriminators, Zen Store will assume certain things
- * about the structure of the chunk ID so changes must be made carefully.
- * 
- */
-enum class EIoChunkType : uint8
-{
-	Invalid = 0,
-	ExportBundleData = 1,
-	BulkData = 2,
-	OptionalBulkData = 3,
-	MemoryMappedBulkData = 4,
-	ScriptObjects = 5,
-	ContainerHeader = 6,
-	ExternalFile = 7,
-	ShaderCodeLibrary = 8,
-	ShaderCode = 9
-};
 
 /**
  * Creates a chunk identifier (generic -- prefer specialized versions where possible)
diff --git a/Engine/Source/Runtime/Core/Public/IO/PackageId.h b/Engine/Source/Runtime/Core/Public/IO/PackageId.h
--- a/Engine/Source/Runtime/Core/Public/IO/PackageId.h
+++ b/Engine/Source/Runtime/Core/Public/IO/PackageId.h
@@ -19,6 +19,11 @@ public:
 
 	CORE_API static FPackageId FromName(const FName& Name);
 
+	static FPackageId FromValue(const uint64 Value)
+	{
+		return FPackageId(Value);
+	}
+
 	inline bool IsValid() const
 	{
 		return Id != InvalidId;
diff --git a/Engine/Source/Runtime/Core/Public/IO/PackageStore.h b/Engine/Source/Runtime/Core/Public/IO/PackageStore.h
--- a/Engine/Source/Runtime/Core/Public/IO/PackageStore.h
+++ b/Engine/Source/Runtime/Core/Public/IO/PackageStore.h
@@ -8,6 +8,8 @@
 
 class FArchive;
 class FStructuredArchiveSlot;
+class FCbObject;
+class FCbWriter;
 
 /**
  * Package export information.
@@ -20,6 +22,10 @@ struct FPackageStoreExportInfo
 	uint32 LoadOrder = 0;
 
 	CORE_API friend FArchive& operator<<(FArchive& Ar, FPackageStoreExportInfo& ExportInfo);
+	
+	CORE_API friend FCbWriter& operator<<(FCbWriter& Writer, const FPackageStoreExportInfo& ExportInfo);
+	
+	CORE_API static FPackageStoreExportInfo FromCbObject(const FCbObject& Obj);
 };
 
 /**
@@ -85,6 +91,10 @@ struct FPackageStoreEntryResource
 	}
 
 	CORE_API friend FArchive& operator<<(FArchive& Ar, FPackageStoreEntryResource& PackageStoreEntry);
+	
+	CORE_API friend FCbWriter& operator<<(FCbWriter& Writer, const FPackageStoreEntryResource& PackageStoreEntry);
+	
+	CORE_API static FPackageStoreEntryResource FromCbObject(const FCbObject& Obj);
 };
 
 /**
diff --git a/Engine/Source/Runtime/Core/Public/IO/PackageStoreWriter.h b/Engine/Source/Runtime/Core/Public/IO/PackageStoreWriter.h
--- a/Engine/Source/Runtime/Core/Public/IO/PackageStoreWriter.h
+++ b/Engine/Source/Runtime/Core/Public/IO/PackageStoreWriter.h
@@ -3,6 +3,7 @@
 #pragma once
 
 #include "IO/IoDispatcher.h"
+#include "Misc/SecureHash.h"
 
 struct FPackageStoreEntryResource;
 
@@ -11,7 +12,7 @@ class IPackageStoreWriter
 public:
 	virtual ~IPackageStoreWriter() = default;
 
-	struct FPackageBaseInfo
+	struct FBeginPackageInfo
 	{
 		FName	PackageName;
 	};
@@ -20,11 +21,18 @@ public:
 
 		This must be called before any data is produced for a given package
 	  */
-	virtual void BeginPackage(const FPackageBaseInfo& Info) = 0;
+	virtual void BeginPackage(const FBeginPackageInfo& Info) = 0;
+
+	struct FCommitPackageInfo
+	{
+		FName PackageName;
+		FGuid PackageGuid;
+		bool bSucceeded = false;
+	};
 
 	/** Finalize a package started with BeginPackage()
 	  */
-	virtual void CommitPackage(const FPackageBaseInfo& Info) = 0;
+	virtual void CommitPackage(const FCommitPackageInfo& Info) = 0;
 
 	struct FPackageInfo
 	{
@@ -118,6 +126,24 @@ public:
 	 * Flush any outstanding writes.
 	 */
 	virtual void Flush() = 0;
+
+	struct FCookedPackageInfo
+	{
+		FName PackageName;
+		FMD5Hash Hash;
+		FGuid PackageGuid;
+		int64 DiskSize = -1;
+	};
+
+	/**
+	 * Returns a list of cooked package(s).
+	 */
+	virtual void GetCookedPackages(TArray<FCookedPackageInfo>& OutCookedPackages) = 0;
+
+	/**
+	 * Remove cooked package(s) that has been modified since the last cook.
+	 */
+	virtual void RemoveCookedPackages(TArrayView<const FName> PackageNamesToRemove) = 0;
 };
 
 static inline const ANSICHAR* LexToString(IPackageStoreWriter::FBulkDataInfo::EType Value)
diff --git a/Engine/Source/Runtime/StorageServerClient/Private/StorageServerConnection.cpp b/Engine/Source/Runtime/StorageServerClient/Private/StorageServerConnection.cpp
--- a/Engine/Source/Runtime/StorageServerClient/Private/StorageServerConnection.cpp
+++ b/Engine/Source/Runtime/StorageServerClient/Private/StorageServerConnection.cpp
@@ -337,7 +337,7 @@ int32 FStorageServerConnection::HandshakeRequest()
 void FStorageServerConnection::FileManifestRequest(TFunctionRef<void(FIoChunkId Id, FStringView Path)> Callback)
 {
 	TAnsiStringBuilder<256> ResourceBuilder;
-	ResourceBuilder.Append(OplogPath).Append("/files");
+	ResourceBuilder.Append(OplogPath).Append("/files?filter=client");
 	FStorageServerRequest Request("GET", *ResourceBuilder, Hostname);
 	FSocket* Socket = Request.Send(*this);
 	if (!Socket)
@@ -357,7 +357,7 @@ void FStorageServerConnection::FileManifestRequest(TFunctionRef<void(FIoChunkId
 			FCbObjectId Id = Entry["id"].AsObjectId();
 
 			TStringBuilder<128> WidePath;
-			WidePath.Append(FUTF8ToTCHAR(Entry["path"].AsString()));
+			WidePath.Append(FUTF8ToTCHAR(Entry["clientpath"].AsString()));
 
 			FIoChunkId ChunkId;
 			ChunkId.Set(Id.GetView());
