iostore fix crash in logging error

diff --git a/Engine/Source/Developer/IoStoreUtilities/Private/IoStoreUtilities.cpp b/Engine/Source/Developer/IoStoreUtilities/Private/IoStoreUtilities.cpp
--- a/Engine/Source/Developer/IoStoreUtilities/Private/IoStoreUtilities.cpp
+++ b/Engine/Source/Developer/IoStoreUtilities/Private/IoStoreUtilities.cpp
@@ -1326,7 +1326,7 @@ void ProcessShaderLibraries(const FIoStoreArguments& Arguments, TArray<FContaine
 							ShaderTargetFile.ChunkId = ShaderChunkId;
 							ShaderTargetFile.ChunkType = EContainerChunkType::ShaderCode;
 							ShaderTargetFile.bForceUncompressed = true;
-
+							
 							FShaderInfo& ShaderInfo = Shaders.AddDefaulted_GetRef();
 							ShaderInfo.ChunkId = ShaderChunkId;
 							ShaderInfo.LibraryChunkId = TargetFile.ChunkId;
@@ -1338,7 +1338,7 @@ void ProcessShaderLibraries(const FIoStoreArguments& Arguments, TArray<FContaine
 							ShaderTargetFile.SourceBuffer.Emplace(*FindCodeIoBuffer);
 							ShaderTargetFile.SourceSize = FindCodeIoBuffer->DataSize();
 							ShaderInfo.Size = FindCodeIoBuffer->DataSize();
-
+							
 							AddedShaderChunkIds.Add(ShaderChunkId);
 						}
 					}
@@ -1348,9 +1348,9 @@ void ProcessShaderLibraries(const FIoStoreArguments& Arguments, TArray<FContaine
 		}
 
 		check(ShaderTargetFiles.Num() == Shaders.Num());
-
+		
 		if (ShaderTargetFiles.Num() == 0)
-		{ 
+		{
 			continue;
 		}
 
@@ -2399,7 +2399,7 @@ int32 CreateTarget(const FIoStoreArguments& Arguments, const FIoStoreWriterSetti
 			HeaderAr << ContainerTarget->Header;
 			int64 DataSize = HeaderAr.TotalSize();
 			FIoBuffer ContainerHeaderBuffer(FIoBuffer::AssumeOwnership, HeaderAr.ReleaseOwnership(), DataSize);
-
+			
 			FIoWriteOptions WriteOptions;
 			WriteOptions.DebugName = TEXT("ContainerHeader");
 			ContainerTarget->IoStoreWriter->Append(
@@ -2472,7 +2472,7 @@ int32 CreateTarget(const FIoStoreArguments& Arguments, const FIoStoreWriterSetti
 	uint64 ImportedPackagesCount = 0;
 	uint64 NoImportedPackagesCount = 0;
 	uint64 NameMapCount = 0;
-
+	
 	for (const FLegacyCookedPackage* Package : Packages)
 	{
 		UExpSize += Package->UExpSize;
@@ -2482,7 +2482,7 @@ int32 CreateTarget(const FIoStoreArguments& Arguments, const FIoStoreWriterSetti
 		ImportedPackagesCount += PackageImportedPackagesCount;
 		NoImportedPackagesCount += PackageImportedPackagesCount == 0;
 	}
-
+	
 	uint64 GlobalShaderSize = 0;
 	uint64 SharedShaderSize = 0;
 	uint64 UniqueShaderSize = 0;
@@ -4107,7 +4107,7 @@ static bool ParsePakResponseFile(const TCHAR* FilePath, TArray<FContainerSourceF
 	TArray<FString> ResponseFileContents;
 	if (!FFileHelper::LoadFileToStringArray(ResponseFileContents, FilePath))
 	{
-		UE_LOG(LogIoStore, Error, TEXT("Failed to read response file '%s'."), *FilePath);
+		UE_LOG(LogIoStore, Error, TEXT("Failed to read response file '%s'."), FilePath);
 		return false;
 	}
 
