Cooking: Store build dependencies, runtime dependencies, and TargetDomainKey in an attachment in the oplog.
#rnx

diff --git a/Engine/Config/BaseEditor.ini b/Engine/Config/BaseEditor.ini
--- a/Engine/Config/BaseEditor.ini
+++ b/Engine/Config/BaseEditor.ini
@@ -330,6 +330,7 @@ MaxMemoryAllowance=16384
 EditorDomainEnabled=false
 EditorDomainSaveUnversioned=false
 BulkDataRegistryEnabled=false
+TargetDomainEnabled=false
 PackagesPerGC=0
 IdleTimeToGC=20
 MemoryMaxUsedPhysical=16384
diff --git a/Engine/Source/Developer/IoStoreUtilities/Private/ZenStoreWriter.cpp b/Engine/Source/Developer/IoStoreUtilities/Private/ZenStoreWriter.cpp
--- a/Engine/Source/Developer/IoStoreUtilities/Private/ZenStoreWriter.cpp
+++ b/Engine/Source/Developer/IoStoreUtilities/Private/ZenStoreWriter.cpp
@@ -308,20 +308,17 @@ FZenStoreWriter::FZenStoreWriter(
 		if (FileStatus .IsOk())
 		{
 			FCbObject FilesObj = FileStatus.ConsumeValueOrDie();
-			if (FilesObj["files"])
+			for (FCbField& FileEntry : FilesObj["files"])
 			{
-				for (FCbField& FileEntry : FilesObj["files"].AsArray())
-				{
-					FCbObject FileObj	= FileEntry.AsObject();
-					FCbObjectId FileId	= FileObj["id"].AsObjectId();
-					FString ServerPath	= FString(FileObj["serverpath"].AsString());
-					FString ClientPath	= FString(FileObj["clientpath"].AsString());
+				FCbObject FileObj	= FileEntry.AsObject();
+				FCbObjectId FileId	= FileObj["id"].AsObjectId();
+				FString ServerPath	= FString(FileObj["serverpath"].AsString());
+				FString ClientPath	= FString(FileObj["clientpath"].AsString());
 
-					FIoChunkId FileChunkId;
-					FileChunkId.Set(FileId.GetView());
+				FIoChunkId FileChunkId;
+				FileChunkId.Set(FileId.GetView());
 					
-					ZenFileSystemManifest->AddManifestEntry(FileChunkId, MoveTemp(ServerPath), MoveTemp(ClientPath));
-				}
+				ZenFileSystemManifest->AddManifestEntry(FileChunkId, MoveTemp(ServerPath), MoveTemp(ClientPath));
 			}
 
 			UE_LOG(LogZenStoreWriter, Display, TEXT("Fetched '%d' files(s) from oplog '%s/%s'"), ZenFileSystemManifest->NumEntries(), *ProjectId, *OplogId);
@@ -609,12 +606,12 @@ void FZenStoreWriter::CommitPackage(const FCommitPackageInfo& Info)
 			PackageObj.BeginObject();
 			PackageObj << "key" << Info.PackageName.ToString();
 
-			// NOTE: The package GUID and disk size are used for iterative cooks when comparing asset registry package data
+			// NOTE: The package GUID and disk size are used for legacy iterative cooks when comparing asset registry package data
 			PackageObj.BeginObject("package");
 			PackageObj << "id" << PkgData.ChunkId;
 			PackageObj << "guid" << Info.PackageGuid;
 			PackageObj << "data" << PkgDataAttachment;
-			PackageObj << "disksize" << PkgDiskSize; 
+			PackageObj << "disksize" << PkgDiskSize;
 			PackageObj.EndObject();
 
 			PackageObj << "packagestoreentry" << PkgData.PackageStoreEntry;
@@ -665,6 +662,13 @@ void FZenStoreWriter::CommitPackage(const FCommitPackageInfo& Info)
 				PackageObj.EndArray();
 			}
 
+			if (Info.TargetDomainDependencies)
+			{
+				FCbAttachment DependencyAttachment(Info.TargetDomainDependencies);
+				Pkg.AddAttachment(DependencyAttachment);
+				PackageObj << "dependencies" << DependencyAttachment;
+			}
+
 			PackageObj.EndObject();
 
 			FCbObject Obj = PackageObj.Save().AsObject();
diff --git a/Engine/Source/Editor/UnrealEd/Classes/CookOnTheSide/CookOnTheFlyServer.h b/Engine/Source/Editor/UnrealEd/Classes/CookOnTheSide/CookOnTheFlyServer.h
--- a/Engine/Source/Editor/UnrealEd/Classes/CookOnTheSide/CookOnTheFlyServer.h
+++ b/Engine/Source/Editor/UnrealEd/Classes/CookOnTheSide/CookOnTheFlyServer.h
@@ -1117,6 +1117,8 @@ private:
 	bool bSaveBusy = false;
 	/** If preloading is enabled, we call TryPreload until it returns true before sending the package to LoadReady, otherwise we skip TryPreload and it goes immediately. */
 	bool bPreloadingEnabled = false;
+	/** If TargetDomain is enabled, we load/save TargetDomainKey hashes and use those to test whether packages have already been cooked in hybrid-iterative builds. */
+	bool bTargetDomainEnabled = true;
 
 	/** Timers for tracking how long we have been busy, to manage retries and warnings of deadlock */
 	float SaveBusyTimeLastRetry = 0.f;
diff --git a/Engine/Source/Editor/UnrealEd/Private/CookOnTheFlyServer.cpp b/Engine/Source/Editor/UnrealEd/Private/CookOnTheFlyServer.cpp
--- a/Engine/Source/Editor/UnrealEd/Private/CookOnTheFlyServer.cpp
+++ b/Engine/Source/Editor/UnrealEd/Private/CookOnTheFlyServer.cpp
@@ -90,6 +90,7 @@
 #include "MeshCardRepresentation.h"
 #include "Misc/PathViews.h"
 #include "String/Find.h"
+#include "TargetDomain/TargetDomainUtils.h"
 
 #include "AssetRegistryModule.h"
 #include "AssetRegistryState.h"
@@ -4734,12 +4735,20 @@ void FSaveCookedPackageContext::FinishPlatform()
 		FAssetPackageData* AssetPackageData = Generator->GetAssetPackageData(Package->GetFName());
 		check(AssetPackageData);
 
+		FCbObject TargetDomainDependencies;
+		if (COTFS.bTargetDomainEnabled)
+		{
+			UE_SCOPED_HIERARCHICAL_COOKTIMER(TargetDomainDependencies);
+			TargetDomainDependencies = UE::TargetDomain::GetDependenciesObject(Package, TargetPlatform, nullptr /* ErrorMessage */);
+		}
+
 		IPackageStoreWriter::FCommitPackageInfo Info;
 		Info.bSucceeded = bSuccessful;
 		Info.PackageName = Package->GetFName();
 		PRAGMA_DISABLE_DEPRECATION_WARNINGS
 		Info.PackageGuid = AssetPackageData->PackageGuid;
 		PRAGMA_ENABLE_DEPRECATION_WARNINGS
+		Info.TargetDomainDependencies = TargetDomainDependencies;
 
 		PackageStoreWriter->CommitPackage(Info);
 	}
@@ -4974,6 +4983,7 @@ void UCookOnTheFlyServer::Initialize( ECookMode::Type DesiredCookMode, ECookInit
 		bPreloadingEnabled = true;
 		FLinkerLoad::SetPreloadingEnabled(true);
 	}
+	GConfig->GetBool(TEXT("CookSettings"), TEXT("TargetDomainEnabled"), bTargetDomainEnabled, GEditorIni);
 
 	{
 		const FConfigSection* CacheSettings = GConfig->GetSectionPrivate(TEXT("CookPlatformDataCacheSettings"), false, true, GEditorIni);
diff --git a/Engine/Source/Editor/UnrealEd/Private/TargetDomain/TargetDomainUtils.h b/Engine/Source/Editor/UnrealEd/Private/TargetDomain/TargetDomainUtils.h
new file mode 100644
--- /dev/null
+++ b/Engine/Source/Editor/UnrealEd/Private/TargetDomain/TargetDomainUtils.h
@@ -0,0 +1,64 @@
+// Copyright Epic Games, Inc. All Rights Reserved.
+
+#pragma once
+
+#include "Containers/Array.h"
+#include "IO/IoHash.h"
+#include "Misc/StringBuilder.h"
+#include "Serialization/CompactBinary.h"
+#include "Serialization/CompactBinaryWriter.h"
+#include "UObject/NameTypes.h"
+
+class ITargetPlatform;
+class UPackage;
+struct FIoHash;
+
+namespace UE::TargetDomain
+{
+
+inline bool TryCollectKeyAndDependencies(UPackage* Package, const ITargetPlatform* TargetPlatform,
+	FIoHash* OutHash, TArray<FName>* OutBuildDependencies, TArray<FName>* OutRuntimeOnlyDependencies, FString* OutErrorMessage)
+{
+	// Implementation is coming in a separate changelist that mostly involves the PackageBuildDependencyTracker
+	if (OutErrorMessage) *OutErrorMessage = TEXT("Not yet implemented.");
+	return false;
+}
+
+inline FCbObject GetDependenciesObject(UPackage* Package, const ITargetPlatform* TargetPlatform, FString* ErrorMessage)
+{
+	FIoHash TargetDomainKey;
+	TArray<FName> BuildDependencies;
+	TArray<FName> RuntimeOnlyDependencies;
+	if (!UE::TargetDomain::TryCollectKeyAndDependencies(Package, TargetPlatform, &TargetDomainKey, &BuildDependencies, &RuntimeOnlyDependencies, ErrorMessage))
+	{
+		return FCbObject();
+	}
+
+	FCbWriter Writer;
+	Writer.BeginObject();
+	Writer << "targetdomainkey" << TargetDomainKey;
+	TStringBuilder<128> PackageNameBuffer;
+	if (!BuildDependencies.IsEmpty())
+	{
+		Writer.BeginArray("builddependencies");
+		for (FName DependencyName : BuildDependencies)
+		{
+			DependencyName.ToString(PackageNameBuffer);
+			Writer << PackageNameBuffer;
+		}
+		Writer.EndArray();
+	}
+	if (!RuntimeOnlyDependencies.IsEmpty())
+	{
+		Writer.BeginArray("runtimeonlydependencies");
+		for (FName DependencyName : RuntimeOnlyDependencies)
+		{
+			DependencyName.ToString(PackageNameBuffer);
+			Writer << PackageNameBuffer;
+		}
+		Writer.EndArray();
+	}
+	return Writer.Save().AsObject();
+}
+
+}
diff --git a/Engine/Source/Runtime/Core/Public/IO/PackageStoreWriter.h b/Engine/Source/Runtime/Core/Public/IO/PackageStoreWriter.h
--- a/Engine/Source/Runtime/Core/Public/IO/PackageStoreWriter.h
+++ b/Engine/Source/Runtime/Core/Public/IO/PackageStoreWriter.h
@@ -4,6 +4,7 @@
 
 #include "IO/IoDispatcher.h"
 #include "Misc/SecureHash.h"
+#include "Serialization/CompactBinary.h"
 
 struct FPackageStoreEntryResource;
 
@@ -43,6 +44,7 @@ public:
 	{
 		FName PackageName;
 		FGuid PackageGuid;
+		FCbObject TargetDomainDependencies;
 		bool bSucceeded = false;
 	};
 
