PackageStore: Add IPackageStoreWriter::GetTargetDomainDependencies and FZenStoreHttpClient::ReadOplogAttachment.
#rnx

diff --git a/Engine/Source/Developer/IoStoreUtilities/Private/ZenStoreHttpClient.cpp b/Engine/Source/Developer/IoStoreUtilities/Private/ZenStoreHttpClient.cpp
--- a/Engine/Source/Developer/IoStoreUtilities/Private/ZenStoreHttpClient.cpp
+++ b/Engine/Source/Developer/IoStoreUtilities/Private/ZenStoreHttpClient.cpp
@@ -1194,13 +1194,24 @@ TIoStatusOr<uint64> FZenStoreHttpClient::GetChunkSize(const FIoChunkId& Id)
 TIoStatusOr<FIoBuffer> FZenStoreHttpClient::ReadChunk(const FIoChunkId& Id, uint64 Offset, uint64 Size)
 {
 	TRACE_CPUPROFILER_EVENT_SCOPE(ZenStoreHttp_ReadChunk);
+	TStringBuilder<128> ChunkUri;
+	ChunkUri << OplogPath << '/' << Id;
+	return ReadOpLogUri(ChunkUri, Offset, Size);
+}
 
+TIoStatusOr<FIoBuffer> FZenStoreHttpClient::ReadOpLogAttachment(FStringView Id, uint64 Offset, uint64 Size)
+{
+	TStringBuilder<128> ChunkUri;
+	ChunkUri << OplogPath << '/' << Id;
+	return ReadOpLogUri(ChunkUri, Offset, Size);
+}
+
+TIoStatusOr<FIoBuffer> FZenStoreHttpClient::ReadOpLogUri(FStringBuilderBase& ChunkUri, uint64 Offset, uint64 Size)
+{
 	check(bAllowRead);
 
 	UE::Zen::FScopedRequestPtr Request(RequestPool.Get());
 	TArray64<uint8> GetBuffer;
-	TStringBuilder<128> ChunkUri;
-	ChunkUri << OplogPath << '/' << Id;
 
 	bool bHaveQuery = false;
 
diff --git a/Engine/Source/Developer/IoStoreUtilities/Private/ZenStoreHttpClient.h b/Engine/Source/Developer/IoStoreUtilities/Private/ZenStoreHttpClient.h
--- a/Engine/Source/Developer/IoStoreUtilities/Private/ZenStoreHttpClient.h
+++ b/Engine/Source/Developer/IoStoreUtilities/Private/ZenStoreHttpClient.h
@@ -2,8 +2,9 @@
 
 #pragma once
 
-#include "IO/IoDispatcher.h"
 #include "Async/Future.h"
+#include "IO/IoDispatcher.h"
+#include "Misc/StringBuilder.h"
 #include "Serialization/CompactBinaryPackage.h"
 
 class FCbPackage;
@@ -39,11 +40,14 @@ public:
 
 	TIoStatusOr<uint64> GetChunkSize(const FIoChunkId& Id);
 	TIoStatusOr<FIoBuffer> ReadChunk(const FIoChunkId& Id, uint64 Offset = 0, uint64 Size = ~0ull);
+	TIoStatusOr<FIoBuffer> ReadOpLogAttachment(FStringView Id, uint64 Offset = 0, uint64 Size = ~0ull);
 
 	TFuture<TIoStatusOr<FCbObject>> GetOplog();
 	TFuture<TIoStatusOr<FCbObject>> GetFiles();
 
 private:
+	TIoStatusOr<FIoBuffer> ReadOpLogUri(FStringBuilderBase& ChunkUri, uint64 Offset, uint64 Size);
+
 	FString	HostName;
 	uint16 Port;
 	TUniquePtr<Zen::FRequestPool> RequestPool;
diff --git a/Engine/Source/Developer/IoStoreUtilities/Private/ZenStoreWriter.cpp b/Engine/Source/Developer/IoStoreUtilities/Private/ZenStoreWriter.cpp
--- a/Engine/Source/Developer/IoStoreUtilities/Private/ZenStoreWriter.cpp
+++ b/Engine/Source/Developer/IoStoreUtilities/Private/ZenStoreWriter.cpp
@@ -13,6 +13,7 @@
 #include "HAL/PlatformFileManager.h"
 #include "IO/IoDispatcher.h"
 #include "Misc/App.h"
+#include "Misc/StringBuilder.h"
 #include "Interfaces/ITargetPlatform.h"
 #include "Misc/Paths.h"
 #include "GenericPlatform/GenericPlatformFile.h"
@@ -22,6 +23,16 @@ DEFINE_LOG_CATEGORY_STATIC(LogZenStoreWriter, Log, All);
 
 using namespace UE;
 
+// Note that this is destructive - we yank out the buffer memory from the 
+// IoBuffer into the FSharedBuffer
+FSharedBuffer IoBufferToSharedBuffer(FIoBuffer& InBuffer)
+{
+	InBuffer.EnsureOwned();
+	const uint64 DataSize = InBuffer.DataSize();
+	uint8* DataPtr = InBuffer.Release().ValueOrDie();
+	return FSharedBuffer{ FSharedBuffer::TakeOwnership(DataPtr, DataSize, FMemory::Free) };
+};
+
 FCbObjectId ToObjectId(const FIoChunkId& ChunkId)
 {
 	return FCbObjectId(MakeMemoryView(ChunkId.GetData(), ChunkId.GetSize()));
@@ -289,11 +300,11 @@ FZenStoreWriter::FZenStoreWriter(
 						const int64	PkgDiskSize				= PackageObj["disksize"].AsUInt64();
 						FPackageStoreEntryResource Entry	= FPackageStoreEntryResource::FromCbObject(OplogObj["packagestoreentry"].AsObject());
 						const FName PackageName				= Entry.PackageName;
-
+						FIoHash DependenciesHash			= OplogObj["dependencies"].AsHash();
 						const int32 Index					= PackageStoreEntries.Num();
 
 						PackageStoreEntries.Add(MoveTemp(Entry));
-						CookedPackagesInfo.Add(FCookedPackageInfo { PackageName, IoHashToMD5(PkgHash), PkgGuid, PkgDiskSize });
+						CookedPackagesInfo.Add(FCookedPackageInfo { PackageName, IoHashToMD5(PkgHash), PkgGuid, PkgDiskSize, DependenciesHash });
 						PackageNameToIndex.Add(PackageName, Index);
 					}
 				}
@@ -582,25 +593,26 @@ void FZenStoreWriter::CommitPackage(const FCommitPackageInfo& Info)
 		{
 			check(PackageState->PackageData.IsValid);
 
-			// Note that this is destructive - we yank out the buffer memory from the 
-			// IoBuffer into the FSharedBuffer
-			auto IoBufferToSharedBuffer = [](FIoBuffer& InBuffer) -> FSharedBuffer {
-				InBuffer.EnsureOwned();
-				const uint64 DataSize = InBuffer.DataSize();
-				uint8* DataPtr = InBuffer.Release().ValueOrDie();
-				return FSharedBuffer{ FSharedBuffer::TakeOwnership(DataPtr, DataSize, FMemory::Free) };
-			};
-
 			FPackageDataEntry& PkgData = PackageState->PackageData;
 			
 			const int64 PkgDiskSize = PkgData.Payload.DataSize();
 			
 			FCbAttachment PkgDataAttachment(IoBufferToSharedBuffer(PkgData.Payload));
 			Pkg.AddAttachment(PkgDataAttachment);
-			
+
+			TOptional<FCbAttachment> DependenciesAttachment;
+			FIoHash DependenciesHash;
+			if (Info.TargetDomainDependencies)
+			{
+				DependenciesAttachment.Emplace(Info.TargetDomainDependencies);
+				Pkg.AddAttachment(*DependenciesAttachment);
+				DependenciesHash = DependenciesAttachment->GetHash();
+			}
+
 			CommitEventArgs.EntryIndex = PackageStoreEntries.Num();
 			PackageStoreEntries.Add(PkgData.PackageStoreEntry);
-			CookedPackagesInfo.Add(FCookedPackageInfo { Info.PackageName, IoHashToMD5(PkgDataAttachment.GetHash()), Info.PackageGuid, PkgDiskSize });
+			CookedPackagesInfo.Add(FCookedPackageInfo { Info.PackageName, IoHashToMD5(PkgDataAttachment.GetHash()), Info.PackageGuid, PkgDiskSize, DependenciesHash });
+			PackageNameToIndex.Add(Info.PackageName, CookedPackagesInfo.Num()-1);
 
 			FCbWriter PackageObj;
 			PackageObj.BeginObject();
@@ -662,11 +674,9 @@ void FZenStoreWriter::CommitPackage(const FCommitPackageInfo& Info)
 				PackageObj.EndArray();
 			}
 
-			if (Info.TargetDomainDependencies)
+			if (DependenciesAttachment)
 			{
-				FCbAttachment DependencyAttachment(Info.TargetDomainDependencies);
-				Pkg.AddAttachment(DependencyAttachment);
-				PackageObj << "dependencies" << DependencyAttachment;
+				PackageObj << "dependencies" << *DependenciesAttachment;
 			}
 
 			PackageObj.EndObject();
@@ -712,6 +722,30 @@ void FZenStoreWriter::GetCookedPackages(TArray<FCookedPackageInfo>& OutCookedPac
 	OutCookedPackages.Append(CookedPackagesInfo);
 }
 
+FCbObject FZenStoreWriter::GetTargetDomainDependencies(FName PackageName)
+{
+	const int32* Idx = PackageNameToIndex.Find(PackageName);
+	if (!Idx)
+	{
+		return FCbObject();
+	}
+
+	const IPackageStoreWriter::FCookedPackageInfo& Info = CookedPackagesInfo[*Idx];
+	TIoStatusOr<FIoBuffer> BufferResult = HttpClient->ReadOpLogAttachment(WriteToString<48>(Info.TargetDomainDependencies));
+	if (!BufferResult.IsOk())
+	{
+		return FCbObject();
+	}
+	FIoBuffer Buffer = BufferResult.ValueOrDie();
+	if (Buffer.DataSize() == 0)
+	{
+		return FCbObject();
+	}
+
+	FSharedBuffer SharedBuffer = IoBufferToSharedBuffer(Buffer);
+	return FCbObject(SharedBuffer);
+}
+
 void FZenStoreWriter::RemoveCookedPackages(TArrayView<const FName> PackageNamesToRemove)
 {
 	TSet<int32> PackageIndicesToKeep;
diff --git a/Engine/Source/Developer/IoStoreUtilities/Public/FilePackageStoreWriter.h b/Engine/Source/Developer/IoStoreUtilities/Public/FilePackageStoreWriter.h
--- a/Engine/Source/Developer/IoStoreUtilities/Public/FilePackageStoreWriter.h
+++ b/Engine/Source/Developer/IoStoreUtilities/Public/FilePackageStoreWriter.h
@@ -98,6 +98,9 @@ public:
 	virtual void GetCookedPackages(TArray<FCookedPackageInfo>& OutCookedPackages)
 	{ }
 
+	virtual FCbObject GetTargetDomainDependencies(FName PackageName) override
+	{ return FCbObject(); }
+
 	virtual void RemoveCookedPackages(TArrayView<const FName> PackageNamesToRemove)
 	{ }
 
diff --git a/Engine/Source/Developer/IoStoreUtilities/Public/ZenStoreWriter.h b/Engine/Source/Developer/IoStoreUtilities/Public/ZenStoreWriter.h
--- a/Engine/Source/Developer/IoStoreUtilities/Public/ZenStoreWriter.h
+++ b/Engine/Source/Developer/IoStoreUtilities/Public/ZenStoreWriter.h
@@ -57,7 +57,8 @@ public:
 	IOSTOREUTILITIES_API void WriteIoStorePackageData(const FPackageInfo& Info, const FIoBuffer& PackageData, const FPackageStoreEntryResource& PackageStoreEntry, const TArray<FFileRegion>& FileRegions);
 
 	IOSTOREUTILITIES_API virtual void GetCookedPackages(TArray<FCookedPackageInfo>& OutCookedPackages) override;
-	IOSTOREUTILITIES_API virtual void RemoveCookedPackages(TArrayView<const FName> PackageNamesToRemove);
+	IOSTOREUTILITIES_API virtual FCbObject GetTargetDomainDependencies(FName PackageName) override;
+	IOSTOREUTILITIES_API virtual void RemoveCookedPackages(TArrayView<const FName> PackageNamesToRemove) override;
 
 private:
 	void CreateProjectMetaData(FCbPackage& Pkg, FCbWriter& PackageObj, bool bGenerateContainerHeader);
diff --git a/Engine/Source/Editor/UnrealEd/Private/EditorDomain/EditorDomainUtils.cpp b/Engine/Source/Editor/UnrealEd/Private/EditorDomain/EditorDomainUtils.cpp
--- a/Engine/Source/Editor/UnrealEd/Private/EditorDomain/EditorDomainUtils.cpp
+++ b/Engine/Source/Editor/UnrealEd/Private/EditorDomain/EditorDomainUtils.cpp
@@ -518,6 +518,11 @@ public:
 	{
 		checkNoEntry();
 	}
+	virtual FCbObject GetTargetDomainDependencies(FName InPackageName) override
+	{
+		checkNoEntry();
+		return FCbObject();
+	}
 
 	// FMemoryPackageStoreWriter interface
 	uint64 GetHeaderSize() const { return HeaderSize; }
@@ -555,7 +560,8 @@ void FMemoryPackageStoreWriter::BeginPackage(const FBeginPackageInfo& Info)
 	SetPackageName(Info.PackageName);
 }
 
-auto IoBufferToSharedBuffer = [](const FIoBuffer& InBuffer) -> FSharedBuffer {
+FSharedBuffer IoBufferToSharedBuffer(const FIoBuffer& InBuffer)
+{
 	InBuffer.EnsureOwned();
 	const uint64 DataSize = InBuffer.DataSize();
 	FIoBuffer MutableBuffer(InBuffer);
diff --git a/Engine/Source/Runtime/Core/Public/IO/PackageStoreWriter.h b/Engine/Source/Runtime/Core/Public/IO/PackageStoreWriter.h
--- a/Engine/Source/Runtime/Core/Public/IO/PackageStoreWriter.h
+++ b/Engine/Source/Runtime/Core/Public/IO/PackageStoreWriter.h
@@ -2,6 +2,7 @@
 
 #pragma once
 
+#include "Containers/StringView.h"
 #include "IO/IoDispatcher.h"
 #include "Misc/SecureHash.h"
 #include "Serialization/CompactBinary.h"
@@ -163,6 +164,7 @@ public:
 		FMD5Hash Hash;
 		FGuid PackageGuid;
 		int64 DiskSize = -1;
+		FIoHash TargetDomainDependencies;
 	};
 
 	/**
@@ -170,6 +172,12 @@ public:
 	 */
 	virtual void GetCookedPackages(TArray<FCookedPackageInfo>& OutCookedPackages) = 0;
 
+	/**
+	 * Returns the TargetDomainDependencies that were previously commited for the given PackageName.
+	 * Returns an empty object if not found.
+	 */
+	virtual FCbObject GetTargetDomainDependencies(FName PackageName) = 0;
+
 	/**
 	 * Remove cooked package(s) that has been modified since the last cook.
 	 */
