ZenStoreHttpClient: Fix missing function definition in the non-windows version of the class.
#rnx

diff --git a/Engine/Source/Developer/IoStoreUtilities/Private/ZenStoreHttpClient.cpp b/Engine/Source/Developer/IoStoreUtilities/Private/ZenStoreHttpClient.cpp
--- a/Engine/Source/Developer/IoStoreUtilities/Private/ZenStoreHttpClient.cpp
+++ b/Engine/Source/Developer/IoStoreUtilities/Private/ZenStoreHttpClient.cpp
@@ -1375,6 +1375,16 @@ TIoStatusOr<FIoBuffer> FZenStoreHttpClient::ReadChunk(const FIoChunkId& Id, uint
 	return FIoBuffer();
 }
 
+TIoStatusOr<FIoBuffer> FZenStoreHttpClient::ReadOpLogAttachment(FStringView Id, uint64 Offset, uint64 Size)
+{
+	return FIoBuffer();
+}
+
+TIoStatusOr<FIoBuffer> FZenStoreHttpClient::ReadOpLogUri(FStringBuilderBase& ChunkUri, uint64 Offset, uint64 Size)
+{
+	return FIoBuffer();
+}
+
 void FZenStoreHttpClient::StartBuildPass()
 {
 }
