log warning instead of check for packages not in clusters

diff --git a/Engine/Source/Developer/IoStoreUtilities/Private/IoStoreUtilities.cpp b/Engine/Source/Developer/IoStoreUtilities/Private/IoStoreUtilities.cpp
--- a/Engine/Source/Developer/IoStoreUtilities/Private/IoStoreUtilities.cpp
+++ b/Engine/Source/Developer/IoStoreUtilities/Private/IoStoreUtilities.cpp
@@ -871,7 +871,7 @@ static void AssignPackagesDiskOrder(
 	}
 	UE_LOG(LogIoStore, Display, TEXT("Ordered %d packages using fallback bundle order"), AssignedPackages.Num() - LastAssignedCount);
 
-	check(AssignedPackages.Num() == Packages.Num());
+	check(AssignedPackages.Num() <= Packages.Num());
 	
 	if (bClusterByOrderFilePriority)
 	{
@@ -895,6 +895,19 @@ static void AssignPackagesDiskOrder(
 		}
 		delete Cluster;
 	}
+	
+	if (AssignedPackages.Num() != Packages.Num())
+	{
+		UE_LOG(LogIoStore, Warning, TEXT(" %d/%d packages not assigned order"), Packages.Num()-AssignedPackages.Num(), Packages.Num());
+	
+		for (FLegacyCookedPackage* Package : Packages)
+		{
+			if (!AssignedPackages.Contains(Package))
+			{
+				Package->DiskLayoutOrder = LayoutIndex++;
+			}
+		}
+	}
 }
 
 static void CreateDiskLayout(
