PVS Studio analysis fix

#rnx

diff --git a/Engine/Source/Developer/IoStoreUtilities/Private/PackageStoreOptimizer.cpp b/Engine/Source/Developer/IoStoreUtilities/Private/PackageStoreOptimizer.cpp
--- a/Engine/Source/Developer/IoStoreUtilities/Private/PackageStoreOptimizer.cpp
+++ b/Engine/Source/Developer/IoStoreUtilities/Private/PackageStoreOptimizer.cpp
@@ -815,11 +815,12 @@ void FPackageStoreOptimizer::OptimizeExportBundles(const TMap<FPackageId, FPacka
 	TArray<FPackageStorePackage::FExportBundleGraphNode*> LoadOrder = SortExportBundleGraphNodesInLoadOrder(Packages, Edges);
 
 	FPackageStorePackage* PreviousPackage = nullptr;
-	for (FPackageStorePackage::FExportBundleGraphNode* Node : LoadOrder)
+	for (const FPackageStorePackage::FExportBundleGraphNode* Node : LoadOrder)
 	{
+		check(Node);
 		FPackageStorePackage* Package = Node->Package;
 		check(Package);
-		if (PreviousPackage != Package || !Package->CurrentBundle)
+		if (Package->CurrentBundle == nullptr || Package != PreviousPackage)
 		{
 			if (Package->GraphData.ExportBundles.IsEmpty())
 			{
