Refactored the UE-side Zen project store interface so that we now use a common HTTP interface implementation across cache / project store classes. No functional impact, tested with ShooterGame

diff --git a/Engine/Source/Developer/DerivedDataCache/Internal/ZenServerHttp.h b/Engine/Source/Developer/DerivedDataCache/Internal/ZenServerHttp.h
--- a/Engine/Source/Developer/DerivedDataCache/Internal/ZenServerHttp.h
+++ b/Engine/Source/Developer/DerivedDataCache/Internal/ZenServerHttp.h
@@ -38,13 +38,13 @@ enum class EContentType
 class FZenHttpRequest
 {
 public:
-	DERIVEDDATACACHE_API FZenHttpRequest(const TCHAR* InDomain, bool bInLogErrors);
+	DERIVEDDATACACHE_API FZenHttpRequest(FStringView InDomain, bool bInLogErrors);
 	DERIVEDDATACACHE_API ~FZenHttpRequest();
 
 	/**
-		* Resets all options on the request except those that should always be set.
-		*/
-	void Reset();
+	  * Resets all options on the request except those that should always be set.
+	  */
+	DERIVEDDATACACHE_API void Reset();
 
 	/** Returns the HTTP response code.*/
 	inline const int GetResponseCode() const
@@ -102,14 +102,17 @@ public:
 		* @param Uri Url to use.
 		* @return Result of the request
 		*/
-	DERIVEDDATACACHE_API Result PerformBlockingHead(const TCHAR* Uri);
+	DERIVEDDATACACHE_API Result PerformBlockingHead(FStringView Uri);
 
 	/**
 		* Query an url using the request. Queries can use either "Head" or "Delete" verbs.
 		* @param Uri Url to use.
 		* @return Result of the request
 		*/
-	DERIVEDDATACACHE_API Result PerformBlockingDelete(const TCHAR* Uri);
+	DERIVEDDATACACHE_API Result PerformBlockingDelete(FStringView Uri);
+
+	DERIVEDDATACACHE_API Result PerformBlockingPost(FStringView Uri, FCbObjectView Obj);
+	DERIVEDDATACACHE_API Result PerformBlockingPost(FStringView Uri, FMemoryView Payload);
 
 	/**
 		* Returns the response buffer as a string. Note that is the request is performed
@@ -120,6 +123,17 @@ public:
 		return GetAnsiBufferAsString(ResponseBuffer);
 	}
 
+	/**
+	  * Returns the response buffer. Note that is the request is performed
+	  * with an external buffer as target buffer this will be empty.
+	  */
+	const TArray64<uint8>& GetResponseBuffer() const
+	{
+		return ResponseBuffer;
+	}
+
+	DERIVEDDATACACHE_API bool GetHeader(const ANSICHAR* Header, FString& OutValue) const;
+
 private:
 	void* /* CURL */		Curl = nullptr;
 	long /* CURLCode */		CurlResult;
@@ -175,7 +189,7 @@ private:
   */
 struct FZenHttpRequestPool
 {
-	DERIVEDDATACACHE_API explicit FZenHttpRequestPool(const TCHAR* InServiceUrl);
+	DERIVEDDATACACHE_API explicit FZenHttpRequestPool(FStringView InServiceUrl, uint32 PoolEntryCount = 16);
 	DERIVEDDATACACHE_API ~FZenHttpRequestPool();
 
 	/** Block until a request is free. Once a request has been returned it is
@@ -192,11 +206,11 @@ struct FZenHttpRequestPool
 private:
 	struct FEntry
 	{
-		TAtomic<uint8> Usage;
+		std::atomic<uint8> IsAllocated;
 		FZenHttpRequest* Request;
 	};
 
-	TStaticArray<FEntry, 16> Pool;
+	TArray<FEntry> Pool;
 };
 
 /**
diff --git a/Engine/Source/Developer/DerivedDataCache/Private/ZenDerivedDataBackend.h b/Engine/Source/Developer/DerivedDataCache/Private/ZenDerivedDataBackend.h
--- a/Engine/Source/Developer/DerivedDataCache/Private/ZenDerivedDataBackend.h
+++ b/Engine/Source/Developer/DerivedDataCache/Private/ZenDerivedDataBackend.h
@@ -6,7 +6,8 @@
 
 #include "ZenServerInterface.h"
 
-// Macro for whether to enable the S3 backend. libcurl is not currently available on Mac.
+// Macro for whether to enable the Zen DDC backend. libcurl is not currently available on Mac.
+
 #if UE_WITH_ZEN
 #	define WITH_ZEN_DDC_BACKEND 1
 #else
diff --git a/Engine/Source/Developer/DerivedDataCache/Private/ZenServerHttp.cpp b/Engine/Source/Developer/DerivedDataCache/Private/ZenServerHttp.cpp
--- a/Engine/Source/Developer/DerivedDataCache/Private/ZenServerHttp.cpp
+++ b/Engine/Source/Developer/DerivedDataCache/Private/ZenServerHttp.cpp
@@ -23,6 +23,7 @@
 #include "Serialization/CompactBinaryPackage.h"
 #include "Serialization/CompactBinaryValidation.h"
 #include "Serialization/LargeMemoryReader.h"
+#include "Serialization/LargeMemoryWriter.h"
 
 namespace UE::Zen {
 
@@ -39,7 +40,7 @@ namespace UE::Zen {
 		static size_t StaticSeekFn(void* UserData, curl_off_t Offset, int Origin);
 	};
 
-	FZenHttpRequest::FZenHttpRequest(const TCHAR* InDomain, bool bInLogErrors)
+	FZenHttpRequest::FZenHttpRequest(FStringView InDomain, bool bInLogErrors)
 		: bLogErrors(bInLogErrors)
 		, Domain(InDomain)
 	{
@@ -123,6 +124,33 @@ namespace UE::Zen {
 		return PerformBlocking(Uri, RequestVerb::Put, ContentLength);
 	}
 
+	FZenHttpRequest::Result FZenHttpRequest::PerformBlockingPost(FStringView Uri, FCbObjectView Obj)
+	{
+		FLargeMemoryWriter Out;
+		Obj.CopyTo(Out);
+
+		return PerformBlockingPost(Uri, Out.GetView());
+	}
+
+	FZenHttpRequest::Result FZenHttpRequest::PerformBlockingPost(FStringView Uri, FMemoryView Payload)
+	{
+		uint64 ContentLength = 0u;
+
+		curl_easy_setopt(Curl, CURLOPT_POST, 1L);
+		curl_easy_setopt(Curl, CURLOPT_INFILESIZE, Payload.GetSize());
+		curl_easy_setopt(Curl, CURLOPT_READDATA, this);
+		curl_easy_setopt(Curl, CURLOPT_READFUNCTION, &FZenHttpRequest::FStatics::StaticReadFn);
+
+		// TODO: proper Content-Type: header
+
+		ContentLength = Payload.GetSize();
+
+		FCompositeBuffer Buffer(FSharedBuffer::MakeView(FMemoryView{ reinterpret_cast<const uint8*>(Payload.GetData()), Payload.GetSize() }));
+		ReadDataView = &Buffer;
+
+		return PerformBlocking(Uri, RequestVerb::Post, ContentLength);
+	}
+
 	FZenHttpRequest::Result FZenHttpRequest::PerformBlockingDownload(FStringView Uri, TArray64<uint8>* Buffer)
 	{
 		curl_easy_setopt(Curl, CURLOPT_HTTPGET, 1L);
@@ -154,14 +182,14 @@ namespace UE::Zen {
 		return LocalResult;
 	}
 
-	FZenHttpRequest::Result FZenHttpRequest::PerformBlockingHead(const TCHAR* Uri)
+	FZenHttpRequest::Result FZenHttpRequest::PerformBlockingHead(FStringView Uri)
 	{
 		curl_easy_setopt(Curl, CURLOPT_NOBODY, 1L);
 
 		return PerformBlocking(Uri, RequestVerb::Head, 0u);
 	}
 
-	FZenHttpRequest::Result FZenHttpRequest::PerformBlockingDelete(const TCHAR* Uri)
+	FZenHttpRequest::Result FZenHttpRequest::PerformBlockingDelete(const FStringView Uri)
 	{
 		curl_easy_setopt(Curl, CURLOPT_POST, 1L);
 		curl_easy_setopt(Curl, CURLOPT_CUSTOMREQUEST, "DELETE");
@@ -176,7 +204,7 @@ namespace UE::Zen {
 			nullptr
 		};
 
-		TRACE_CPUPROFILER_EVENT_SCOPE(ZenDDC_CurlPerform);
+		TRACE_CPUPROFILER_EVENT_SCOPE(ZenHttp_CurlPerform);
 
 		// Setup request options
 		FString Url = FString::Printf(TEXT("%s/%s"), *Domain, *FString(Uri));
@@ -224,6 +252,30 @@ namespace UE::Zen {
 		return CurlResult == CURLE_OK ? Result::Success : Result::Failed;
 	}
 
+	/**
+	 * Attempts to find the header from the response. Returns false if header is not present.
+	 */
+	bool FZenHttpRequest::GetHeader(const ANSICHAR* Header, FString& OutValue) const
+	{
+		check(CurlResult != CURL_LAST);  // Cannot query headers before request is sent
+
+		const ANSICHAR* HeadersBuffer = (const ANSICHAR*)ResponseHeader.GetData();
+		size_t HeaderLen = strlen(Header);
+
+		// Find the header key in the (ANSI) response buffer. If not found we can exist immediately
+		if (const ANSICHAR* Found = strstr(HeadersBuffer, Header))
+		{
+			const ANSICHAR* Linebreak = strchr(Found, '\r');
+			const ANSICHAR* ValueStart = Found + HeaderLen + 2; //colon and space
+			const size_t ValueSize = Linebreak - ValueStart;
+			FUTF8ToTCHAR TCHARData(ValueStart, ValueSize);
+			OutValue = FString(TCHARData.Length(), TCHARData.Get());
+			return true;
+		}
+		return false;
+	}
+
+
 	void FZenHttpRequest::LogResult(long InResult, const TCHAR* Uri, RequestVerb Verb) const
 	{
 		CURLcode Result = (CURLcode) InResult;
@@ -233,17 +285,17 @@ namespace UE::Zen {
 			const TCHAR* VerbStr = nullptr;
 			FString AdditionalInfo;
 
-			const bool Is400 = (ResponseCode >= 400) && (ResponseCode <= 499);
+			const bool Is404 = (ResponseCode == 404);
 			const bool Is200 = (ResponseCode >= 200) && (ResponseCode <= 299);
 
 			switch (Verb)
 			{
 			case RequestVerb::Head:
-				bSuccess = Is400 || Is200;
+				bSuccess = Is200 || Is404;
 				VerbStr = TEXT("querying");
 				break;
 			case RequestVerb::Get:
-				bSuccess = Is400 || Is200;
+				bSuccess = Is200 || Is404;
 				VerbStr = TEXT("fetching");
 				AdditionalInfo = FString::Printf(TEXT("Received: %d bytes."), BytesReceived);
 				break;
@@ -257,7 +309,7 @@ namespace UE::Zen {
 				VerbStr = TEXT("posting");
 				break;
 			case RequestVerb::Delete:
-				bSuccess = Is200;
+				bSuccess = Is200 || Is404;
 				VerbStr = TEXT("deleting");
 				break;
 			}
@@ -267,7 +319,7 @@ namespace UE::Zen {
 				UE_LOG(
 					LogDerivedDataCache,
 					Verbose,
-					TEXT("Finished %s HTTP cache entry (response %d) from %s. %s"),
+					TEXT("Finished %s zen data (response %d) from %s. %s"),
 					VerbStr,
 					ResponseCode,
 					Uri,
@@ -283,7 +335,7 @@ namespace UE::Zen {
 				UE_LOG(
 					LogDerivedDataCache,
 					Error,
-					TEXT("Failed %s HTTP cache entry (response %d) from %s. Response: %s"),
+					TEXT("Failed %s zen data (response %d) from %s. Response: %s"),
 					VerbStr,
 					ResponseCode,
 					Uri,
@@ -449,11 +501,13 @@ namespace UE::Zen {
 
 	//////////////////////////////////////////////////////////////////////////
 
-	FZenHttpRequestPool::FZenHttpRequestPool(const TCHAR* InServiceUrl)
+	FZenHttpRequestPool::FZenHttpRequestPool(FStringView InServiceUrl, uint32 PoolEntryCount)
 	{
+		Pool.SetNum(PoolEntryCount);
+
 		for (uint8 i = 0; i < Pool.Num(); ++i)
 		{
-			Pool[i].Usage = 0u;
+			Pool[i].IsAllocated = 0u;
 			Pool[i].Request = new FZenHttpRequest(InServiceUrl, true);
 		}
 	}
@@ -463,7 +517,8 @@ namespace UE::Zen {
 		for (uint8 i = 0; i < Pool.Num(); ++i)
 		{
 			// No requests should be in use by now.
-			check(Pool[i].Usage.Load(EMemoryOrder::Relaxed) == 0u);
+			check(Pool[i].IsAllocated.load(std::memory_order_relaxed) == 0u);
+
 			delete Pool[i].Request;
 		}
 	}
@@ -477,20 +532,23 @@ namespace UE::Zen {
 	  */
 	FZenHttpRequest* FZenHttpRequestPool::WaitForFreeRequest()
 	{
-		TRACE_CPUPROFILER_EVENT_SCOPE(ZenDDC_WaitForConnPool);
+		TRACE_CPUPROFILER_EVENT_SCOPE(ZenHttp_WaitForConnPool);
+
 		while (true)
 		{
 			for (uint8 i = 0; i < Pool.Num(); ++i)
 			{
-				if (!Pool[i].Usage.Load(EMemoryOrder::Relaxed))
+				if (!Pool[i].IsAllocated.load(std::memory_order_relaxed))
 				{
 					uint8 Expected = 0u;
-					if (Pool[i].Usage.CompareExchange(Expected, 1u))
+					if (Pool[i].IsAllocated.compare_exchange_strong(Expected, 1u))
 					{
 						return Pool[i].Request;
 					}
 				}
 			}
+
+			// This should use a better mechanism like condition variables / events
 			FPlatformProcess::Sleep(UE_ZENDDC_BACKEND_WAIT_INTERVAL);
 		}
 	}
@@ -507,7 +565,7 @@ namespace UE::Zen {
 			{
 				Request->Reset();
 				uint8 Expected = 1u;
-				Pool[i].Usage.CompareExchange(Expected, 0u);
+				Pool[i].IsAllocated.compare_exchange_strong(Expected, 0u);
 				return;
 			}
 		}
diff --git a/Engine/Source/Developer/IoStoreUtilities/IoStoreUtilities.Build.cs b/Engine/Source/Developer/IoStoreUtilities/IoStoreUtilities.Build.cs
--- a/Engine/Source/Developer/IoStoreUtilities/IoStoreUtilities.Build.cs
+++ b/Engine/Source/Developer/IoStoreUtilities/IoStoreUtilities.Build.cs
@@ -15,7 +15,8 @@ public class IoStoreUtilities : ModuleRules
             "CoreUObject",
             "Projects",
             "AssetRegistry",
-            "RenderCore",
+			"DerivedDataCache",
+			"RenderCore",
 			"libcurl",
         });
         PrivateDependencyModuleNames.Add("PakFile");
diff --git a/Engine/Source/Developer/IoStoreUtilities/Private/ZenStoreHttpClient.cpp b/Engine/Source/Developer/IoStoreUtilities/Private/ZenStoreHttpClient.cpp
--- a/Engine/Source/Developer/IoStoreUtilities/Private/ZenStoreHttpClient.cpp
+++ b/Engine/Source/Developer/IoStoreUtilities/Private/ZenStoreHttpClient.cpp
@@ -22,792 +22,11 @@
 #include "Serialization/BufferArchive.h"
 #include "Serialization/CompactBinaryWriter.h"
 #include "Serialization/LargeMemoryWriter.h"
+#include "ZenServerHttp.h"
 
 DEFINE_LOG_CATEGORY_STATIC(LogZenStore, Log, All);
 
-namespace UE { namespace Zen {
-
-#define	UE_ZENHTTP_BACKEND_WAIT_INTERVAL			0.01f
-#define	UE_ZENHTTP_HTTP_REQUEST_TIMEOUT_SECONDS		30L
-#define UE_ZENHTTP_HTTP_REQUEST_TIMOUT_ENABLED		1
-#define UE_ZENHTTP_HTTP_DEBUG						0
-#define UE_ZENHTTP_MAX_FAILED_LOGIN_ATTEMPTS		16
-#define UE_ZENHTTP_MAX_ATTEMPTS						4
-#define UE_ZENHTTP_MAX_BUFFER_RESERVE				104857600u
-
-/**
- * Minimal HTTP request type wrapping CURL without the need for managers. This request
- * is written to allow reuse of request objects, in order to allow connections to be reused.
- *
- * CURL has a global library initialization (curl_global_init). We rely on this happening in
- * the Online/HTTP library which is a dependency on this module.
- */
-class FRequest
-{
-public:
-	/**
-	 * Supported request verbs
-	 */
-	enum RequestVerb
-	{
-		Get,
-		Put,
-		Post,
-		Delete,
-		Head
-	};
-
-	/**
-	 * Convenience result type interpreted from HTTP response code.
-	 */
-	enum Result
-	{
-		Success,
-		Failed,
-		FailedTimeout
-	};
-
-	FRequest(FStringView InDomain, bool bInLogErrors)
-		: bLogErrors(bInLogErrors)
-		, Domain(InDomain)
-	{
-		Curl = curl_easy_init();
-		Reset();
-	}
-
-	~FRequest()
-	{
-		curl_easy_cleanup(Curl);
-	}
-
-	/**
-	 * Resets all options on the request except those that should always be set.
-	 */
-	void Reset()
-	{
-		Headers.Reset();
-		ResponseHeader.Reset();
-		ResponseBuffer.Reset();
-		ResponseCode = 0;
-		ReadDataView = TConstArrayView64<uint8>();
-		WriteDataBufferPtr = nullptr;
-		WriteHeaderBufferPtr = nullptr;
-		BytesSent = 0;
-		BytesReceived = 0;
-		CurlResult = CURL_LAST;
-
-		curl_easy_reset(Curl);
-
-		// Options that are always set for all connections.
-#if UE_ZENHTTP_HTTP_REQUEST_TIMOUT_ENABLED
-		curl_easy_setopt(Curl, CURLOPT_CONNECTTIMEOUT, UE_ZENHTTP_HTTP_REQUEST_TIMEOUT_SECONDS);
-#endif
-		curl_easy_setopt(Curl, CURLOPT_FOLLOWLOCATION, 1L);
-		curl_easy_setopt(Curl, CURLOPT_NOSIGNAL, 1L);
-		// Response functions
-		curl_easy_setopt(Curl, CURLOPT_HEADERDATA, this);
-		curl_easy_setopt(Curl, CURLOPT_HEADERFUNCTION, &FRequest::StaticWriteHeaderFn);
-		curl_easy_setopt(Curl, CURLOPT_WRITEDATA, this);
-		curl_easy_setopt(Curl, CURLOPT_WRITEFUNCTION, StaticWriteBodyFn);
-		// Rewind method, handle special error case where request need to rewind data stream
-		curl_easy_setopt(Curl, CURLOPT_SEEKFUNCTION, StaticSeekFn);
-		curl_easy_setopt(Curl, CURLOPT_SEEKDATA, this);
-		// Debug hooks
-#if UE_ZENHTTP_HTTP_DEBUG
-		curl_easy_setopt(Curl, CURLOPT_DEBUGDATA, this);
-		curl_easy_setopt(Curl, CURLOPT_DEBUGFUNCTION, StaticDebugCallback);
-		curl_easy_setopt(Curl, CURLOPT_VERBOSE, 1L);
-#endif
-	}
-
-	/** Returns the HTTP response code.*/
-	const int64 GetResponseCode() const
-	{
-		return ResponseCode;
-	}
-
-	/** Returns the number of bytes received this request (headers withstanding). */
-	const size_t GetBytesReceived() const
-	{
-		return BytesReceived;
-	}
-
-	/** Returns the number of bytes sent during this request (headers withstanding). */
-	const size_t  GetBytesSent() const
-	{
-		return BytesSent;
-	}
-
-	template<RequestVerb V>
-	Result PerformBlockingUpload(FStringView Uri, FCbObjectView Obj)
-	{
-		FLargeMemoryWriter Out;
-		Obj.CopyTo(Out);
-
-		TConstArrayView64<uint8> Payload { Out.GetData(), Out.TotalSize()};
-		return PerformBlockingUpload<V>(Uri, Payload);
-	}
-
-	Result PerformBlockingPost(FStringView Uri, FCbObjectView Obj)
-	{
-		FLargeMemoryWriter Out;
-		Obj.CopyTo(Out);
-
-		return PerformBlockingPost(Uri, Out.GetView());
-	}
-
-	Result PerformBlockingPost(FStringView Uri, FMemoryView Payload)
-	{
-		uint64 ContentLength = 0u;
-
-		curl_easy_setopt(Curl, CURLOPT_POST, 1L);
-		curl_easy_setopt(Curl, CURLOPT_INFILESIZE, Payload.GetSize());
-		curl_easy_setopt(Curl, CURLOPT_READDATA, this);
-		curl_easy_setopt(Curl, CURLOPT_READFUNCTION, StaticReadFn);
-
-		// TODO: proper Content-Type: header
-
-		ContentLength = Payload.GetSize();
-		ReadDataView = { reinterpret_cast<const uint8*>(Payload.GetData()), (int64)Payload.GetSize() };
-
-		return PerformBlocking(Uri, Post, ContentLength);
-	}
-
-	/**
-	 * Upload buffer using the request, using either "Put" or "Post" verbs.
-	 * @param Uri Url to use.
-	 * @param Buffer Data to upload
-	 * @return Result of the request
-	 */
-	template<RequestVerb V>
-	Result PerformBlockingUpload(FStringView Uri, TConstArrayView64<uint8> Buffer)
-	{
-		static_assert(V == Put || V == Post, "Upload should use either Put or Post verbs.");
-
-		uint64 ContentLength = 0u;
-
-		if (V == Put)
-		{
-			curl_easy_setopt(Curl, CURLOPT_UPLOAD, 1L);
-			curl_easy_setopt(Curl, CURLOPT_INFILESIZE, Buffer.Num());
-			curl_easy_setopt(Curl, CURLOPT_READDATA, this);
-			curl_easy_setopt(Curl, CURLOPT_READFUNCTION, StaticReadFn);
-			Headers.Add(FString(TEXT("Content-Type: application/octet-stream")));
-			ContentLength = Buffer.Num();
-			ReadDataView = Buffer;
-		}
-		else if (V == Post)
-		{
-			curl_easy_setopt(Curl, CURLOPT_POST, 1L);
-			curl_easy_setopt(Curl, CURLOPT_INFILESIZE, Buffer.Num());
-			curl_easy_setopt(Curl, CURLOPT_READDATA, this);
-			curl_easy_setopt(Curl, CURLOPT_READFUNCTION, StaticReadFn);
-			
-			//
-			// TODO: proper Content-Type: header
-			//Headers.Add(V == Post ? FString(TEXT("Content-Type: application/x-www-form-urlencoded")) : FString(TEXT("Content-Type: application/json")));
-
-			ContentLength = Buffer.Num();
-			ReadDataView = Buffer;
-		}
-
-		return PerformBlocking(Uri, V, ContentLength);
-	}
-
-	/**
-	 * Download an url into a buffer using the request.
-	 * @param Uri Url to use.
-	 * @param Buffer Optional buffer where data should be downloaded to. If empty downloaded data will
-	 * be stored in an internal buffer and accessed GetResponse* methods.
-	 * @return Result of the request
-	 */
-	Result PerformBlockingDownload(FStringView Uri, TArray64<uint8>* Buffer)
-	{
-		curl_easy_setopt(Curl, CURLOPT_HTTPGET, 1L);
-		WriteDataBufferPtr = Buffer;
-
-		return PerformBlocking(Uri, Get, 0u);
-	}
-
-	/**
-	 * Query an url using the request. Queries can use either "Head" or "Delete" verbs.
-	 * @param Uri Url to use.
-	 * @return Result of the request
-	 */
-	template<RequestVerb V>
-	Result PerformBlockingQuery(FStringView Uri)
-	{
-		static_assert(V == Head || V == Delete, "Queries should use either Head or Delete verbs.");
-
-		if (V == Delete)
-		{
-			curl_easy_setopt(Curl, CURLOPT_CUSTOMREQUEST, "DELETE");
-		}
-		else if (V == Head)
-		{
-			curl_easy_setopt(Curl, CURLOPT_NOBODY, 1L);
-		}
-
-		return PerformBlocking(Uri, V, 0u);
-	}
-
-	/**
-	 * Set a header to send with the request.
-	 */
-	void SetHeader(const TCHAR* Header, const TCHAR* Value)
-	{
-		check(CurlResult == CURL_LAST); // Cannot set header after request is sent
-		Headers.Add(FString::Printf(TEXT("%s: %s"), Header, Value));
-	}
-
-	/**
-	 * Attempts to find the header from the response. Returns false if header is not present.
-	 */
-	bool GetHeader(const ANSICHAR* Header, FString& OutValue) const
-	{
-		check(CurlResult != CURL_LAST);  // Cannot query headers before request is sent
-
-		const ANSICHAR* HeadersBuffer = (const ANSICHAR*)ResponseHeader.GetData();
-		size_t HeaderLen = strlen(Header);
-
-		// Find the header key in the (ANSI) response buffer. If not found we can exist immediately
-		if (const ANSICHAR* Found = strstr(HeadersBuffer, Header))
-		{
-			const ANSICHAR* Linebreak = strchr(Found, '\r');
-			const ANSICHAR* ValueStart = Found + HeaderLen + 2; //colon and space
-			const size_t ValueSize = Linebreak - ValueStart;
-			FUTF8ToTCHAR TCHARData(ValueStart, ValueSize);
-			OutValue = FString(TCHARData.Length(), TCHARData.Get());
-			return true;
-		}
-		return false;
-	}
-
-	/**
-	 * Returns the response buffer. Note that is the request is performed
-	 * with an external buffer as target buffer this string will be empty.
-	 */
-	const TArray64<uint8>& GetResponseBuffer() const
-	{
-		return ResponseBuffer;
-	}
-
-	/**
-	 * Returns the response buffer as a string. Note that is the request is performed
-	 * with an external buffer as target buffer this string will be empty.
-	 */
-	FString GetResponseAsString() const
-	{
-		return GetAnsiBufferAsString(ResponseBuffer);
-	}
-
-	/**
-	 * Returns the response header as a string.
-	 */
-	FString GetResponseHeaderAsString()
-	{
-		return GetAnsiBufferAsString(ResponseHeader);
-	}
-
-	/** Will return true if the response code is considered a success */
-	static bool IsSuccessResponse(long ResponseCode)
-	{
-		// We consider anything in the 1XX or 2XX range a success
-		return ResponseCode >= 100 && ResponseCode < 300;
-	}
-
-private:
-
-	CURL* Curl;
-	CURLcode				CurlResult;
-	long					ResponseCode;
-	size_t					BytesSent;
-	size_t					BytesReceived;
-	bool					bLogErrors;
-
-	TConstArrayView64<uint8>	ReadDataView;
-	TArray64<uint8>*			WriteDataBufferPtr;
-	TArray64<uint8>*			WriteHeaderBufferPtr;
-
-	TArray64<uint8>			ResponseHeader;
-	TArray64<uint8>			ResponseBuffer;
-	TArray<FString>			Headers;
-	FString					Domain;
-
-	/**
-	 * Performs the request, blocking until finished.
-	 * @param Uri Address on the domain to query
-	 * @param Verb HTTP verb to use
-	 * @param Buffer Optional buffer to directly receive the result of the request.
-	 * If unset the response body will be stored in the request.
-	 */
-	Result PerformBlocking(FStringView Uri, RequestVerb Verb, uint64 ContentLength)
-	{
-		static const char* CommonHeaders[] = {
-			"User-Agent: UE",
-			nullptr
-		};
-
-		TRACE_CPUPROFILER_EVENT_SCOPE(ZenStore_CurlPerform);
-
-		// Setup request options
-		TStringBuilder<128> Url;
-		Url << Domain << Uri;
-		
-		curl_easy_setopt(Curl, CURLOPT_URL, TCHAR_TO_ANSI(*Url));
-
-		// Setup response header buffer. If caller has not setup a response data buffer, use interal.
-		WriteHeaderBufferPtr = &ResponseHeader;
-		if (WriteDataBufferPtr == nullptr)
-		{
-			WriteDataBufferPtr = &ResponseBuffer;
-		}
-
-		// Content-Length should always be set
-		Headers.Add(FString::Printf(TEXT("Content-Length: %llu"), ContentLength));
-
-		// Build headers list
-		curl_slist* CurlHeaders = nullptr;
-		// Add common headers
-		for (uint8 i = 0; CommonHeaders[i] != nullptr; ++i)
-		{
-			CurlHeaders = curl_slist_append(CurlHeaders, CommonHeaders[i]);
-		}
-		// Setup added headers
-		for (const FString& Header : Headers)
-		{
-			CurlHeaders = curl_slist_append(CurlHeaders, TCHAR_TO_ANSI(*Header));
-		}
-		curl_easy_setopt(Curl, CURLOPT_HTTPHEADER, CurlHeaders);
-
-		// Shots fired!
-		CurlResult = curl_easy_perform(Curl);
-
-		// Get response code
-		bool bRedirected = false;
-		if (CURLE_OK == curl_easy_getinfo(Curl, CURLINFO_RESPONSE_CODE, &ResponseCode))
-		{
-			bRedirected = (ResponseCode >= 300 && ResponseCode < 400);
-		}
-
-		LogResult(CurlResult, Uri, Verb);
-
-		// Clean up
-		curl_slist_free_all(CurlHeaders);
-
-		return CurlResult == CURLE_OK ? Success : Failed;
-	}
-
-	void LogResult(CURLcode Result, FStringView Uri, RequestVerb Verb) const
-	{
-		if (Result == CURLE_OK)
-		{
-			bool bSuccess = false;
-			const TCHAR* VerbStr = nullptr;
-			FString AdditionalInfo;
-
-			switch (Verb)
-			{
-			case Head:
-				bSuccess = (ResponseCode == 400 || IsSuccessResponse(ResponseCode));
-				VerbStr = TEXT("querying");
-				break;
-			case Get:
-				bSuccess = (ResponseCode == 400 || IsSuccessResponse(ResponseCode));
-				VerbStr = TEXT("fetching");
-				AdditionalInfo = FString::Printf(TEXT("Received: %d bytes."), BytesReceived);
-				break;
-			case Put:
-				bSuccess = IsSuccessResponse(ResponseCode);
-				VerbStr = TEXT("updating");
-				AdditionalInfo = FString::Printf(TEXT("Sent: %d bytes."), BytesSent);
-				break;
-			case Post:
-				bSuccess = IsSuccessResponse(ResponseCode);
-				VerbStr = TEXT("posting");
-				break;
-			case Delete:
-				bSuccess = IsSuccessResponse(ResponseCode);
-				VerbStr = TEXT("deleting");
-				break;
-			}
-
-			if (bSuccess)
-			{
-				UE_LOG(
-					LogZenStore,
-					Verbose,
-					TEXT("Finished %s Zen store HTTP operation (response %d) from %s. %s"),
-					VerbStr,
-					ResponseCode,
-					*FString(Uri),
-					*AdditionalInfo
-				);
-			}
-			else if (bLogErrors)
-			{
-				// Print the response body if we got one, otherwise print header.
-				FString Response = GetAnsiBufferAsString(ResponseBuffer.Num() > 0 ? ResponseBuffer : ResponseHeader);
-				Response.ReplaceCharInline('\n', ' ');
-				Response.ReplaceCharInline('\r', ' ');
-				// Dont log access denied as error, since tokens can expire mid session
-				if (ResponseCode == 401)
-				{
-					UE_LOG(
-						LogZenStore,
-						Verbose,
-						TEXT("Failed %s Zen store HTTP operation (response %d) from %s. Response: %s"),
-						VerbStr,
-						ResponseCode,
-						*FString(Uri),
-						*Response
-					);
-				}
-				else if (ResponseCode == 404)
-				{
-					// 404 is not an error, it is an expected and valid result in most cases
-				}
-				else
-				{
-					UE_LOG(
-						LogZenStore,
-						Display,
-						TEXT("Failed %s Zen store HTTP operation (response %d) from %s. Response: %s"),
-						VerbStr,
-						ResponseCode,
-						*FString(Uri),
-						*Response
-					);
-				}
-			}
-		}
-		else if (bLogErrors)
-		{
-			UE_LOG(
-				LogZenStore,
-				Display,
-				TEXT("Error while connecting to %s: %s"),
-				*Domain,
-				ANSI_TO_TCHAR(curl_easy_strerror(Result))
-			);
-		}
-	}
-
-	FString GetAnsiBufferAsString(const TArray64<uint8>& Buffer) const
-	{
-		// Content is NOT null-terminated; we need to specify lengths here
-		FUTF8ToTCHAR TCHARData(reinterpret_cast<const ANSICHAR*>(Buffer.GetData()), Buffer.Num());
-		return FString(TCHARData.Length(), TCHARData.Get());
-	}
-
-	static size_t StaticDebugCallback(CURL* Handle, curl_infotype DebugInfoType, char* DebugInfo, size_t DebugInfoSize, void* UserData)
-	{
-		FRequest* Request = static_cast<FRequest*>(UserData);
-
-		switch (DebugInfoType)
-		{
-		case CURLINFO_TEXT:
-		{
-			// Truncate at 1023 characters. This is just an arbitrary number based on a buffer size seen in
-			// the libcurl code.
-			DebugInfoSize = FMath::Min(DebugInfoSize, (size_t)1023);
-
-			// Calculate the actual length of the string due to incorrect use of snprintf() in lib/vtls/openssl.c.
-			char* FoundNulPtr = (char*)memchr(DebugInfo, 0, DebugInfoSize);
-			int CalculatedSize = FoundNulPtr != nullptr ? FoundNulPtr - DebugInfo : DebugInfoSize;
-
-			auto ConvertedString = StringCast<TCHAR>(static_cast<const ANSICHAR*>(DebugInfo), CalculatedSize);
-			FString DebugText(ConvertedString.Length(), ConvertedString.Get());
-			DebugText.ReplaceInline(TEXT("\n"), TEXT(""), ESearchCase::CaseSensitive);
-			DebugText.ReplaceInline(TEXT("\r"), TEXT(""), ESearchCase::CaseSensitive);
-			UE_LOG(LogZenStore, VeryVerbose, TEXT("%p: '%s'"), Request, *DebugText);
-		}
-		break;
-
-		case CURLINFO_HEADER_IN:
-			UE_LOG(LogZenStore, VeryVerbose, TEXT("%p: Received header (%d bytes)"), Request, DebugInfoSize);
-			break;
-
-		case CURLINFO_DATA_IN:
-			UE_LOG(LogZenStore, VeryVerbose, TEXT("%p: Received data (%d bytes)"), Request, DebugInfoSize);
-			break;
-
-		case CURLINFO_DATA_OUT:
-			UE_LOG(LogZenStore, VeryVerbose, TEXT("%p: Sent data (%d bytes)"), Request, DebugInfoSize);
-			break;
-
-		case CURLINFO_SSL_DATA_IN:
-			UE_LOG(LogZenStore, VeryVerbose, TEXT("%p: Received SSL data (%d bytes)"), Request, DebugInfoSize);
-			break;
-
-		case CURLINFO_SSL_DATA_OUT:
-			UE_LOG(LogZenStore, VeryVerbose, TEXT("%p: Sent SSL data (%d bytes)"), Request, DebugInfoSize);
-			break;
-		}
-
-		return 0;
-	}
-
-	static size_t StaticReadFn(void* Ptr, size_t SizeInBlocks, size_t BlockSizeInBytes, void* UserData)
-	{
-		FRequest* Request = static_cast<FRequest*>(UserData);
-		TConstArrayView64<uint8>& ReadDataView = Request->ReadDataView;
-
-		const size_t Offset = Request->BytesSent;
-		const size_t ReadSize = FMath::Min((size_t)ReadDataView.Num() - Offset, SizeInBlocks * BlockSizeInBytes);
-		check((size_t)ReadDataView.Num() >= Offset + ReadSize);
-
-		FMemory::Memcpy(Ptr, ReadDataView.GetData() + Offset, ReadSize);
-		Request->BytesSent += ReadSize;
-		return ReadSize;
-
-		return 0;
-	}
-
-	static size_t StaticWriteHeaderFn(void* Ptr, size_t SizeInBlocks, size_t BlockSizeInBytes, void* UserData)
-	{
-		FRequest* Request = static_cast<FRequest*>(UserData);
-		const size_t WriteSize = SizeInBlocks * BlockSizeInBytes;
-		TArray64<uint8>* WriteHeaderBufferPtr = Request->WriteHeaderBufferPtr;
-		if (WriteHeaderBufferPtr && WriteSize > 0)
-		{
-			const size_t CurrentBufferLength = WriteHeaderBufferPtr->Num();
-			if (CurrentBufferLength > 0)
-			{
-				// Remove the previous zero termination
-				(*WriteHeaderBufferPtr)[CurrentBufferLength - 1] = ' ';
-			}
-
-			// Write the header
-			WriteHeaderBufferPtr->Append((const uint8*)Ptr, WriteSize + 1);
-			(*WriteHeaderBufferPtr)[WriteHeaderBufferPtr->Num() - 1] = 0; // Zero terminate string
-			return WriteSize;
-		}
-		return 0;
-	}
-
-	static size_t StaticWriteBodyFn(void* Ptr, size_t SizeInBlocks, size_t BlockSizeInBytes, void* UserData)
-	{
-		FRequest* Request = static_cast<FRequest*>(UserData);
-		const size_t WriteSize = SizeInBlocks * BlockSizeInBytes;
-		TArray64<uint8>* WriteDataBufferPtr = Request->WriteDataBufferPtr;
-
-		if (WriteDataBufferPtr && WriteSize > 0)
-		{
-			// If this is the first part of the body being received, try to reserve 
-			// memory if content length is defined in the header.
-			if (Request->BytesReceived == 0 && Request->WriteHeaderBufferPtr)
-			{
-				static const ANSICHAR* ContentLengthHeaderStr = "Content-Length: ";
-				const ANSICHAR* Header = (const ANSICHAR*)Request->WriteHeaderBufferPtr->GetData();
-
-				if (const ANSICHAR* ContentLengthHeader = FCStringAnsi::Strstr(Header, ContentLengthHeaderStr))
-				{
-					size_t ContentLength = (size_t)FCStringAnsi::Atoi64(ContentLengthHeader + strlen(ContentLengthHeaderStr));
-					if (ContentLength > 0u && ContentLength < UE_ZENHTTP_MAX_BUFFER_RESERVE)
-					{
-						WriteDataBufferPtr->Reserve(ContentLength);
-					}
-				}
-			}
-
-			// Write to the target buffer
-			WriteDataBufferPtr->Append((const uint8*)Ptr, WriteSize);
-			Request->BytesReceived += WriteSize;
-			return WriteSize;
-		}
-
-		return 0;
-	}
-
-	static size_t StaticSeekFn(void* UserData, curl_off_t Offset, int Origin)
-	{
-		FRequest* Request = static_cast<FRequest*>(UserData);
-		size_t NewPosition = 0;
-
-		switch (Origin)
-		{
-		case SEEK_SET: NewPosition = Offset; break;
-		case SEEK_CUR: NewPosition = Request->BytesSent + Offset; break;
-		case SEEK_END: NewPosition = Request->ReadDataView.Num() + Offset; break;
-		}
-
-		// Make sure we don't seek outside of the buffer
-		if (NewPosition < 0 || NewPosition >= (size_t)Request->ReadDataView.Num())
-		{
-			return CURL_SEEKFUNC_FAIL;
-		}
-
-		// Update the used offset
-		Request->BytesSent = NewPosition;
-		return CURL_SEEKFUNC_OK;
-	}
-
-};
-
-//----------------------------------------------------------------------------------------------------------
-// Request pool
-//----------------------------------------------------------------------------------------------------------
-
-/**
- * Pool that manages a fixed set of requests. Users are required to release requests that have been
- * acquired. Usable with \ref FScopedRequestPtr which handles this automatically.
- */
-struct FRequestPool
-{
-	FRequestPool(FStringView InServiceUrl, uint32 PoolSize)
-	{
-		Pool.AddUninitialized(PoolSize);
-		for (uint8 i = 0; i < Pool.Num(); ++i)
-		{
-			Pool[i].Usage = 0u;
-			Pool[i].Request = new FRequest(InServiceUrl, true);
-		}
-	}
-
-	~FRequestPool()
-	{
-		for (uint8 i = 0; i < Pool.Num(); ++i)
-		{
-			// No requests should be in use by now.
-			check(Pool[i].Usage.load(std::memory_order_acquire) == 0u);
-			delete Pool[i].Request;
-		}
-	}
-
-	/**
-	 * Attempts to get a request is free. Once a request has been returned it is
-	 * "owned by the caller and need to release it to the pool when work has been completed.
-	 * @return Usable request instance if one is available, otherwise null.
-	 */
-	FRequest* GetFreeRequest()
-	{
-		for (uint8 i = 0; i < Pool.Num(); ++i)
-		{
-			if (!Pool[i].Usage.load(std::memory_order_relaxed))
-			{
-				uint8 Expected = 0u;
-				if (Pool[i].Usage.compare_exchange_strong(Expected, 1u))
-				{
-					Pool[i].Request->Reset();
-					return Pool[i].Request;
-				}
-			}
-		}
-		return nullptr;
-	}
-
-	/**
-	 * Block until a request is free. Once a request has been returned it is
-	 * "owned by the caller and need to release it to the pool when work has been completed.
-	 * @return Usable request instance.
-	 */
-	FRequest* WaitForFreeRequest()
-	{
-		TRACE_CPUPROFILER_EVENT_SCOPE(ZenStore_WaitForConnPool);
-		FRequest* Request = nullptr;
-		while (true)
-		{
-			Request = GetFreeRequest();
-			if (Request != nullptr)
-				break;
-			FPlatformProcess::Sleep(UE_ZENHTTP_BACKEND_WAIT_INTERVAL);
-		}
-		return Request;
-	}
-
-	/**
-	 * Release request to the pool.
-	 * @param Request Request that should be freed. Note that any buffer owned by the request can now be reset.
-	 */
-	void ReleaseRequestToPool(FRequest* Request)
-	{
-		for (uint8 i = 0; i < Pool.Num(); ++i)
-		{
-			if (Pool[i].Request == Request)
-			{
-
-				Pool[i].Usage--;
-				return;
-			}
-		}
-		check(false);
-	}
-
-	/**
-	 * While holding a request, make it shared across many users.
-	 */
-	void MakeRequestShared(FRequest* Request, uint8 Users)
-	{
-		check(Users != 0);
-		for (uint8 i = 0; i < Pool.Num(); ++i)
-		{
-			if (Pool[i].Request == Request)
-			{
-
-				Pool[i].Usage = Users;
-				return;
-			}
-		}
-		check(false);
-	}
-
-private:
-
-	struct FEntry
-	{
-		std::atomic<uint8> Usage;
-		FRequest* Request;
-	};
-
-	TArray<FEntry> Pool;
-
-	FRequestPool() = delete;
-};
-
-//----------------------------------------------------------------------------------------------------------
-// FScopedRequestPtr
-//----------------------------------------------------------------------------------------------------------
-
-/**
- * Utility class to manage requesting and releasing requests from the \ref FRequestPool.
- */
-struct FScopedRequestPtr
-{
-public:
-	FScopedRequestPtr(FRequestPool* InPool)
-		: Request(InPool->WaitForFreeRequest())
-		, Pool(InPool)
-	{}
-
-	~FScopedRequestPtr()
-	{
-		Pool->ReleaseRequestToPool(Request);
-	}
-
-	bool IsValid() const
-	{
-		return Request != nullptr;
-	}
-
-	FRequest* Get() const
-	{
-		check(IsValid());
-		return Request;
-	}
-
-	FRequest* operator->()
-	{
-		check(IsValid());
-		return Request;
-	}
-
-private:
-	FRequest* Request;
-	FRequestPool* Pool;
-};
-
-} /* Zen */ 
+namespace UE { 
 
 FZenStoreHttpClient::FZenStoreHttpClient(const FStringView InHostName, uint16 InPort)
 : HostName(InHostName)
@@ -819,7 +38,7 @@ FZenStoreHttpClient::FZenStoreHttpClient(const FStringView InHostName, uint16 In
 	Uri.AppendAnsi(":");
 	Uri << InPort;
 
-	RequestPool = MakeUnique<Zen::FRequestPool>(Uri, 32);
+	RequestPool = MakeUnique<Zen::FZenHttpRequestPool>(Uri, 32);
 }
 
 FZenStoreHttpClient::~FZenStoreHttpClient()
@@ -841,17 +60,17 @@ FZenStoreHttpClient::Initialize(FStringView InProjectId,
 	// Establish project
 
 	{
-		UE::Zen::FScopedRequestPtr Request(RequestPool.Get());
+		UE::Zen::FZenScopedRequestPtr Request(RequestPool.Get());
 
 		TStringBuilder<128> ProjectUri;
 		ProjectUri << "/prj/" << InProjectId;
 		TArray64<uint8> GetBuffer;
-		UE::Zen::FRequest::Result Res = Request->PerformBlockingDownload(ProjectUri, &GetBuffer);
+		UE::Zen::FZenHttpRequest::Result Res = Request->PerformBlockingDownload(ProjectUri, &GetBuffer);
 
 		// TODO: how to handle failure here? This is probably the most likely point of failure
 		// if the service is not up or not responding
 
-		if (Res == Zen::FRequest::Success && Request->GetResponseCode() == 200)
+		if (Res == Zen::FZenHttpRequest::Result::Success && Request->GetResponseCode() == 200)
 		{
 			UE_LOG(LogZenStore, Display, TEXT("Zen project '%s' already exists"), *FString(InProjectId));
 		}
@@ -869,7 +88,7 @@ FZenStoreHttpClient::Initialize(FStringView InProjectId,
 
 			Res = Request->PerformBlockingPost(ProjectUri, ProjInfo.Save().AsObject());
 
-			if (Res != Zen::FRequest::Success)
+			if (Res != Zen::FZenHttpRequest::Result::Success)
 			{
 				UE_LOG(LogZenStore, Error, TEXT("Zen project '%s' creation FAILED"), *FString(InProjectId));
 
@@ -889,7 +108,7 @@ FZenStoreHttpClient::Initialize(FStringView InProjectId,
 	// Establish oplog
 
 	{
-		UE::Zen::FScopedRequestPtr Request(RequestPool.Get());
+		UE::Zen::FZenScopedRequestPtr Request(RequestPool.Get());
 
 		TStringBuilder<128> OplogUri;
 		OplogUri << "/prj/" << InProjectId << "/oplog/" << InOplogId;
@@ -899,15 +118,15 @@ FZenStoreHttpClient::Initialize(FStringView InProjectId,
 		if (IsCleanBuild)
 		{
 			UE_LOG(LogZenStore, Display, TEXT("Deleting oplog '%s'/'%s' if it exists"), *FString(InProjectId), *FString(InOplogId));
-			Request->PerformBlockingQuery<Zen::FRequest::Delete>(OplogUri);
+			Request->PerformBlockingDelete(OplogUri);
 			Request->Reset();
 		}
 
 		TArray64<uint8> GetBuffer;
-		UE::Zen::FRequest::Result Res = Request->PerformBlockingDownload(OplogUri, &GetBuffer);
+		UE::Zen::FZenHttpRequest::Result Res = Request->PerformBlockingDownload(OplogUri, &GetBuffer);
 		FCbObjectView OplogInfo;
 
-		if (Res == Zen::FRequest::Success && Request->GetResponseCode() == 200)
+		if (Res == Zen::FZenHttpRequest::Result::Success && Request->GetResponseCode() == 200)
 		{
 			UE_LOG(LogZenStore, Display, TEXT("Zen oplog '%s'/'%s' already exists"), *FString(InProjectId), *FString(InOplogId));
 
@@ -915,12 +134,12 @@ FZenStoreHttpClient::Initialize(FStringView InProjectId,
 		}
 		else
 		{
-			TConstArrayView64<uint8> Payload;
+			FMemoryView Payload;
 
 			Request->Reset();
-			Res = Request->PerformBlockingUpload<Zen::FRequest::Post>(OplogUri, Payload);
+			Res = Request->PerformBlockingPost(OplogUri, Payload);
 
-			if (Res != Zen::FRequest::Success)
+			if (Res != Zen::FZenHttpRequest::Result::Success)
 			{
 				UE_LOG(LogZenStore, Error, TEXT("Zen oplog '%s'/'%s' creation FAILED"), *FString(InProjectId), *FString(InOplogId));
 
@@ -941,7 +160,7 @@ FZenStoreHttpClient::Initialize(FStringView InProjectId,
 			Request->Reset();
 			Res = Request->PerformBlockingDownload(OplogUri, &GetBuffer);
 
-			if (Res == Zen::FRequest::Success && Request->GetResponseCode() == 200)
+			if (Res == Zen::FZenHttpRequest::Result::Success && Request->GetResponseCode() == 200)
 			{
 				OplogInfo = FCbObjectView(GetBuffer.GetData());
 			}
@@ -972,17 +191,17 @@ void FZenStoreHttpClient::InitializeReadOnly(FStringView InProjectId, FStringVie
 	// Establish project
 
 	{
-		UE::Zen::FScopedRequestPtr Request(RequestPool.Get());
+		UE::Zen::FZenScopedRequestPtr Request(RequestPool.Get());
 
 		TStringBuilder<128> ProjectUri;
 		ProjectUri << "/prj/" << InProjectId;
 		TArray64<uint8> GetBuffer;
-		UE::Zen::FRequest::Result Res = Request->PerformBlockingDownload(ProjectUri, &GetBuffer);
+		UE::Zen::FZenHttpRequest::Result Res = Request->PerformBlockingDownload(ProjectUri, &GetBuffer);
 
 		// TODO: how to handle failure here? This is probably the most likely point of failure
 		// if the service is not up or not responding
 
-		if (Res != Zen::FRequest::Success || Request->GetResponseCode() != 200)
+		if (Res != Zen::FZenHttpRequest::Result::Success || Request->GetResponseCode() != 200)
 		{
 			UE_LOG(LogZenStore, Fatal, TEXT("Zen project '%s' not found"), *FString(InProjectId));
 		}
@@ -991,7 +210,7 @@ void FZenStoreHttpClient::InitializeReadOnly(FStringView InProjectId, FStringVie
 	// Establish oplog
 
 	{
-		UE::Zen::FScopedRequestPtr Request(RequestPool.Get());
+		UE::Zen::FZenScopedRequestPtr Request(RequestPool.Get());
 
 		TStringBuilder<128> OplogUri;
 		OplogUri << "/prj/" << InProjectId << "/oplog/" << InOplogId;
@@ -999,9 +218,9 @@ void FZenStoreHttpClient::InitializeReadOnly(FStringView InProjectId, FStringVie
 		OplogPath = OplogUri;
 
 		TArray64<uint8> GetBuffer;
-		UE::Zen::FRequest::Result Res = Request->PerformBlockingDownload(OplogUri, &GetBuffer);
+		UE::Zen::FZenHttpRequest::Result Res = Request->PerformBlockingDownload(OplogUri, &GetBuffer);
 
-		if (Res != Zen::FRequest::Success || Request->GetResponseCode() != 200)
+		if (Res != Zen::FZenHttpRequest::Result::Success || Request->GetResponseCode() != 200)
 		{
 			UE_LOG(LogZenStore, Fatal, TEXT("Zen oplog '%s'/'%s' not found"), *FString(InProjectId), *FString(InOplogId));
 		}
@@ -1055,17 +274,17 @@ TFuture<TIoStatusOr<uint64>> FZenStoreHttpClient::AppendOp(FCbPackage OpEntry)
 
 				FCbFieldIterator Prep = Writer.Save();
 
-				UE::Zen::FScopedRequestPtr Request(RequestPool.Get());
+				UE::Zen::FZenScopedRequestPtr Request(RequestPool.Get());
 
 				bool IsOk = false;
 				
-				const Zen::FRequest::Result Res = Request->PerformBlockingUpload<Zen::FRequest::Post>(OplogPrepNewEntryPath, Prep.AsObjectView());
+				const Zen::FZenHttpRequest::Result Res = Request->PerformBlockingPost(OplogPrepNewEntryPath, Prep.AsObjectView());
 
-				if (Res == Zen::FRequest::Success)
+				if (Res == Zen::FZenHttpRequest::Result::Success)
 				{
 					FCbObjectView NeedObject;
 
-					if (Res == Zen::FRequest::Success && Request->GetResponseCode() == 200)
+					if (Res == Zen::FZenHttpRequest::Result::Success && Request->GetResponseCode() == 200)
 					{
 						NeedObject = FCbObjectView(Request->GetResponseBuffer().GetData());
 
@@ -1151,7 +370,7 @@ TFuture<TIoStatusOr<uint64>> FZenStoreHttpClient::AppendOp(FCbPackage OpEntry)
 
 		UE_LOG(LogZenStore, Verbose, TEXT("Package size: %" UINT64_FMT), SerializedPackage.TotalSize());
 
-		UE::Zen::FScopedRequestPtr Request(RequestPool.Get());
+		UE::Zen::FZenScopedRequestPtr Request(RequestPool.Get());
 
 		TStringBuilder<64> NewOpPostUri;
 		NewOpPostUri << OplogNewEntryPath;
@@ -1161,7 +380,7 @@ TFuture<TIoStatusOr<uint64>> FZenStoreHttpClient::AppendOp(FCbPackage OpEntry)
 			NewOpPostUri << "?salt=" << Salt;
 		}
 
-		if (UE::Zen::FRequest::Success == Request->PerformBlockingPost(NewOpPostUri, SerializedPackage.GetView()))
+		if (UE::Zen::FZenHttpRequest::Result::Success == Request->PerformBlockingPost(NewOpPostUri, SerializedPackage.GetView()))
 		{
 			return TIoStatusOr<uint64>(SerializedPackage.TotalSize());
 		}
@@ -1178,13 +397,13 @@ TIoStatusOr<uint64> FZenStoreHttpClient::GetChunkSize(const FIoChunkId& Id)
 
 	check(bAllowRead);
 
-	UE::Zen::FScopedRequestPtr Request(RequestPool.Get());
+	UE::Zen::FZenScopedRequestPtr Request(RequestPool.Get());
 	TArray64<uint8> GetBuffer;
 	TStringBuilder<128> ChunkUri;
 	ChunkUri << OplogPath << '/' << Id;
-	UE::Zen::FRequest::Result Res = Request->PerformBlockingQuery<Zen::FRequest::Head>(ChunkUri);
+	UE::Zen::FZenHttpRequest::Result Res = Request->PerformBlockingHead(ChunkUri);
 	FString ContentLengthStr;
-	if (Res == Zen::FRequest::Success && Request->GetResponseCode() == 200 && Request->GetHeader("Content-Length", ContentLengthStr))
+	if (Res == Zen::FZenHttpRequest::Result::Success && Request->GetResponseCode() == 200 && Request->GetHeader("Content-Length", ContentLengthStr))
 	{
 		return FCStringWide::Atoi64(*ContentLengthStr);
 	}
@@ -1210,7 +429,7 @@ TIoStatusOr<FIoBuffer> FZenStoreHttpClient::ReadOpLogUri(FStringBuilderBase& Chu
 {
 	check(bAllowRead);
 
-	UE::Zen::FScopedRequestPtr Request(RequestPool.Get());
+	UE::Zen::FZenScopedRequestPtr Request(RequestPool.Get());
 	TArray64<uint8> GetBuffer;
 
 	bool bHaveQuery = false;
@@ -1241,9 +460,9 @@ TIoStatusOr<FIoBuffer> FZenStoreHttpClient::ReadOpLogUri(FStringBuilderBase& Chu
 	}
 
 
-	UE::Zen::FRequest::Result Res = Request->PerformBlockingDownload(ChunkUri, &GetBuffer);
+	UE::Zen::FZenHttpRequest::Result Res = Request->PerformBlockingDownload(ChunkUri, &GetBuffer);
 
-	if (Res == Zen::FRequest::Success && Request->GetResponseCode() == 200)
+	if (Res == Zen::FZenHttpRequest::Result::Success && Request->GetResponseCode() == 200)
 	{
 		return FIoBuffer(FIoBuffer::Clone, GetBuffer.GetData(), GetBuffer.Num());
 	}
@@ -1254,15 +473,15 @@ TFuture<TIoStatusOr<FCbObject>> FZenStoreHttpClient::GetOplog()
 {
 	return Async(EAsyncExecution::LargeThreadPool, [this]
 	{
-		UE::Zen::FScopedRequestPtr Request(RequestPool.Get());
+		UE::Zen::FZenScopedRequestPtr Request(RequestPool.Get());
 		
 		TStringBuilder<128> Uri;
 		Uri << OplogPath << "/entries";
 
 		TArray64<uint8> GetBuffer;
-		UE::Zen::FRequest::Result Res = Request->PerformBlockingDownload(Uri, &GetBuffer);
+		UE::Zen::FZenHttpRequest::Result Res = Request->PerformBlockingDownload(Uri, &GetBuffer);
 
-		if (Res == Zen::FRequest::Success && Request->GetResponseCode() == 200)
+		if (Res == Zen::FZenHttpRequest::Result::Success && Request->GetResponseCode() == 200)
 		{
 			FCbObjectView Response(GetBuffer.GetData());
 			return TIoStatusOr<FCbObject>(FCbObject::Clone(Response));
@@ -1278,15 +497,15 @@ TFuture<TIoStatusOr<FCbObject>> FZenStoreHttpClient::GetFiles()
 {
 	return Async(EAsyncExecution::LargeThreadPool, [this]
 	{
-		UE::Zen::FScopedRequestPtr Request(RequestPool.Get());
+		UE::Zen::FZenScopedRequestPtr Request(RequestPool.Get());
 		
 		TStringBuilder<128> Uri;
 		Uri << OplogPath << "/files";
 
 		TArray64<uint8> GetBuffer;
-		UE::Zen::FRequest::Result Res = Request->PerformBlockingDownload(Uri, &GetBuffer);
+		UE::Zen::FZenHttpRequest::Result Res = Request->PerformBlockingDownload(Uri, &GetBuffer);
 
-		if (Res == Zen::FRequest::Success && Request->GetResponseCode() == 200)
+		if (Res == Zen::FZenHttpRequest::Result::Success && Request->GetResponseCode() == 200)
 		{
 			FCbObjectView Response(GetBuffer.GetData());
 			return TIoStatusOr<FCbObject>(FCbObject::Clone(Response));
@@ -1318,13 +537,13 @@ FZenStoreHttpClient::EndBuildPass(FCbPackage OpEntry)
 
 	UE_LOG(LogZenStore, Verbose, TEXT("Package size: %lld"), SerializedPackage.TotalSize());
 
-	UE::Zen::FScopedRequestPtr Request(RequestPool.Get());
+	UE::Zen::FZenScopedRequestPtr Request(RequestPool.Get());
 
-	TConstArrayView64<uint8> Payload { SerializedPackage.GetData(), SerializedPackage.TotalSize()};
+	FMemoryView Payload { SerializedPackage.GetData(), (uint64)SerializedPackage.TotalSize()};
 	
-	if (UE::Zen::FRequest::Success == Request->PerformBlockingUpload<Zen::FRequest::Post>(OplogNewEntryPath, Payload))
+	if (UE::Zen::FZenHttpRequest::Result::Success == Request->PerformBlockingPost(OplogNewEntryPath, Payload))
 	{
-		return static_cast<uint64>(Payload.Num());
+		return static_cast<uint64>(Payload.GetSize());
 	}
 	else
 	{
@@ -1338,7 +557,7 @@ FZenStoreHttpClient::EndBuildPass(FCbPackage OpEntry)
 
 namespace UE {
 namespace Zen {
-	struct FRequestPool
+	struct FZenHttpRequestPool
 	{
 	};
 }
diff --git a/Engine/Source/Developer/IoStoreUtilities/Private/ZenStoreHttpClient.h b/Engine/Source/Developer/IoStoreUtilities/Private/ZenStoreHttpClient.h
--- a/Engine/Source/Developer/IoStoreUtilities/Private/ZenStoreHttpClient.h
+++ b/Engine/Source/Developer/IoStoreUtilities/Private/ZenStoreHttpClient.h
@@ -12,7 +12,7 @@ class FCbObject;
 
 namespace UE {
 	namespace Zen {
-		struct FRequestPool;
+		struct FZenHttpRequestPool;
 	}
 
 /**
@@ -50,7 +50,7 @@ private:
 
 	FString	HostName;
 	uint16 Port;
-	TUniquePtr<Zen::FRequestPool> RequestPool;
+	TUniquePtr<Zen::FZenHttpRequestPool> RequestPool;
 	FString OplogPath;
 	FString OplogNewEntryPath;
 	FString OplogPrepNewEntryPath;
