Removed functionality to cook to an IoStore container file, it has been superseded by cooking to Zen.
#rnx

diff --git a/Engine/Source/Developer/IoStoreUtilities/Private/FilePackageStoreWriter.cpp b/Engine/Source/Developer/IoStoreUtilities/Private/FilePackageStoreWriter.cpp
deleted file mode 100644
diff --git a/Engine/Source/Developer/IoStoreUtilities/Private/IoStoreUtilities.cpp b/Engine/Source/Developer/IoStoreUtilities/Private/IoStoreUtilities.cpp
--- a/Engine/Source/Developer/IoStoreUtilities/Private/IoStoreUtilities.cpp
+++ b/Engine/Source/Developer/IoStoreUtilities/Private/IoStoreUtilities.cpp
@@ -3,7 +3,6 @@
 #include "IoStoreUtilities.h"
 
 #if !UE_BUILD_SHIPPING
-#include "FilePackageStoreWriter.h"
 #include "HAL/FileManager.h"
 #include "HAL/PlatformFileManager.h"
 #include "Hash/CityHash.h"
@@ -376,15 +375,9 @@ public:
 		}
 		else
 		{
-			TUniquePtr<FIoStoreReaderDataSource> IoStoreReaderDataSource = MakeUnique<FIoStoreReaderDataSource>();
-			Status = IoStoreReaderDataSource->Initialize(*(CookedDirectory / TEXT("global")));
-			if (!Status.IsOk())
-			{
-				return Status;
-			}
-			DataSource = MoveTemp(IoStoreReaderDataSource);
+			return FIoStatusBuilder(EIoErrorCode::NotFound) << TEXT("Missing ZenServerInfo in package store manifest");
 		}
-
+		
 		FIoContainerId ContainerId = FIoContainerId::FromName(TEXT("global"));
 
 		TIoStatusOr<FIoBuffer> HeaderBuffer = DataSource->ReadChunk(CreateIoChunkId(ContainerId.Value(), 0, EIoChunkType::ContainerHeader), FIoReadOptions());
@@ -533,35 +526,6 @@ private:
 		virtual TIoStatusOr<FIoBuffer> ReadChunk(const FIoChunkId& ChunkId, const FIoReadOptions& ReadOptions) = 0;
 	};
 
-	class FIoStoreReaderDataSource
-		: public IDataSource
-	{
-	public:
-		FIoStatus Initialize(const TCHAR* ContainerPath)
-		{
-			using FEncryptionKeys = TMap<FGuid, FAES::FAESKey>;
-			return IoStoreReader.Initialize(ContainerPath, FEncryptionKeys());
-		}
-
-		virtual TIoStatusOr<uint64> GetChunkSize(const FIoChunkId& ChunkId) override
-		{
-			TIoStatusOr<FIoStoreTocChunkInfo> ChunkInfo = IoStoreReader.GetChunkInfo(ChunkId);
-			if (!ChunkInfo.IsOk())
-			{
-				return ChunkInfo.Status();
-			}
-			return ChunkInfo.ValueOrDie().Size;
-		}
-
-		virtual TIoStatusOr<FIoBuffer> ReadChunk(const FIoChunkId& ChunkId, const FIoReadOptions& ReadOptions) override
-		{
-			return IoStoreReader.Read(ChunkId, ReadOptions);
-		}
-
-	private:
-		FIoStoreReader IoStoreReader;
-	};
-
 	class FZenStoreDataSource
 		: public IDataSource
 	{
diff --git a/Engine/Source/Developer/IoStoreUtilities/Private/PackageStoreManifest.cpp b/Engine/Source/Developer/IoStoreUtilities/Private/PackageStoreManifest.cpp
new file mode 100644
--- /dev/null
+++ b/Engine/Source/Developer/IoStoreUtilities/Private/PackageStoreManifest.cpp
@@ -0,0 +1,229 @@
+// Copyright Epic Games, Inc. All Rights Reserved.
+
+#include "PackageStoreManifest.h"
+#include "Serialization/JsonWriter.h"
+#include "Serialization/JsonReader.h"
+#include "Serialization/JsonSerializer.h"
+#include "Dom/JsonObject.h"
+#include "Misc/FileHelper.h"
+#include "Misc/ScopeLock.h"
+
+void FPackageStoreManifest::BeginPackage(FName PackageName)
+{
+	FScopeLock Lock(&CriticalSection);
+	FPackageInfo& PackageInfo = PackageInfoByNameMap.FindOrAdd(PackageName);
+	PackageInfo.PackageName = PackageName;
+	if (PackageInfo.PackageChunkId.IsValid())
+	{
+		FileNameByChunkIdMap.Remove(PackageInfo.PackageChunkId);
+		PackageInfo.PackageChunkId = FIoChunkId();
+	}
+	for (const FIoChunkId& BulkDataChunkId : PackageInfo.BulkDataChunkIds)
+	{
+		FileNameByChunkIdMap.Remove(BulkDataChunkId);
+	}
+	PackageInfo.BulkDataChunkIds.Empty();
+}
+
+void FPackageStoreManifest::AddPackageData(FName PackageName, const FString& FileName, const FIoChunkId& ChunkId)
+{
+	FScopeLock Lock(&CriticalSection);
+	FPackageInfo* PackageInfo = PackageInfoByNameMap.Find(PackageName);
+	check(PackageInfo);
+	PackageInfo->PackageChunkId = ChunkId;
+	if (!FileName.IsEmpty())
+	{
+		FileNameByChunkIdMap.Add(ChunkId, FileName);
+	}
+}
+
+void FPackageStoreManifest::AddBulkData(FName PackageName, const FString& FileName, const FIoChunkId& ChunkId)
+{
+	FScopeLock Lock(&CriticalSection);
+	FPackageInfo* PackageInfo = PackageInfoByNameMap.Find(PackageName);
+	check(PackageInfo);
+	PackageInfo->BulkDataChunkIds.Add(ChunkId);
+	if (!FileName.IsEmpty())
+	{
+		FileNameByChunkIdMap.Add(ChunkId, FileName);
+	}
+}
+
+FIoStatus FPackageStoreManifest::Save(const TCHAR* Filename) const
+{
+	FScopeLock Lock(&CriticalSection);
+	TStringBuilder<64> ChunkIdStringBuilder;
+	auto ChunkIdToString = [&ChunkIdStringBuilder](const FIoChunkId& ChunkId)
+	{
+		ChunkIdStringBuilder.Reset();
+		ChunkIdStringBuilder << ChunkId;
+		return *ChunkIdStringBuilder;
+	};
+
+	FString JsonTcharText;
+	TSharedRef<TJsonWriter<TCHAR, TPrettyJsonPrintPolicy<TCHAR>>> Writer = TJsonWriterFactory<TCHAR, TPrettyJsonPrintPolicy<TCHAR>>::Create(&JsonTcharText);
+	Writer->WriteObjectStart();
+
+	if (ZenServerInfo)
+	{
+		Writer->WriteObjectStart(TEXT("ZenServer"));
+		Writer->WriteValue(TEXT("HostName"), ZenServerInfo->HostName);
+		Writer->WriteValue(TEXT("Port"), ZenServerInfo->Port);
+		Writer->WriteValue(TEXT("ProjectId"), ZenServerInfo->ProjectId);
+		Writer->WriteValue(TEXT("OplogId"), ZenServerInfo->OplogId);
+		Writer->WriteObjectEnd();
+	}
+	
+	Writer->WriteArrayStart(TEXT("Files"));
+	for (const auto& KV : FileNameByChunkIdMap)
+	{
+		Writer->WriteObjectStart();
+		Writer->WriteValue(TEXT("Path"), KV.Value);
+		Writer->WriteValue(TEXT("ChunkId"), ChunkIdToString(KV.Key));
+		Writer->WriteObjectEnd();
+	}
+	Writer->WriteArrayEnd();
+
+	Writer->WriteArrayStart(TEXT("Packages"));
+	for (const auto& KV : PackageInfoByNameMap)
+	{
+		const FPackageInfo& PackageInfo = KV.Value;
+		Writer->WriteObjectStart();
+		Writer->WriteValue(TEXT("Name"), PackageInfo.PackageName.ToString());
+		Writer->WriteValue(TEXT("PackageChunkId"), ChunkIdToString(PackageInfo.PackageChunkId));
+		if (!PackageInfo.BulkDataChunkIds.IsEmpty())
+		{
+			Writer->WriteArrayStart(TEXT("BulkDataChunkIds"));
+			for (const FIoChunkId& ChunkId : PackageInfo.BulkDataChunkIds)
+			{
+				Writer->WriteValue(ChunkIdToString(ChunkId));
+			}
+			Writer->WriteArrayEnd();
+		}
+		Writer->WriteObjectEnd();
+	}
+	Writer->WriteArrayEnd();
+
+	Writer->WriteObjectEnd();
+	Writer->Close();
+
+	if (!FFileHelper::SaveStringToFile(JsonTcharText, Filename))
+	{
+		return FIoStatus(EIoErrorCode::FileOpenFailed);
+	}
+
+	return FIoStatus::Ok;
+}
+
+FIoStatus FPackageStoreManifest::Load(const TCHAR* Filename)
+{
+	FScopeLock Lock(&CriticalSection);
+	PackageInfoByNameMap.Empty();
+	FileNameByChunkIdMap.Empty();
+
+	auto ChunkIdFromString = [](const FString& ChunkIdString)
+	{
+		FStringView ChunkIdStringView(*ChunkIdString, 24);
+		uint8 Data[12];
+		UE::String::HexToBytes(ChunkIdStringView, Data);
+		FIoChunkId ChunkId;
+		ChunkId.Set(Data, 12);
+		return ChunkId;
+	};
+
+	FString JsonText;
+	if (!FFileHelper::LoadFileToString(JsonText, Filename))
+	{
+		return FIoStatus(EIoErrorCode::FileOpenFailed);
+	}
+
+	TSharedPtr<FJsonObject> JsonObject;
+	TSharedRef<TJsonReader<TCHAR>> Reader = TJsonReaderFactory<TCHAR>::Create(JsonText);
+	if (!FJsonSerializer::Deserialize(Reader, JsonObject) || !JsonObject.IsValid())
+	{
+		return FIoStatus(EIoErrorCode::Unknown);
+	}
+
+	TSharedPtr<FJsonValue> ZenServerValue = JsonObject->Values.FindRef(TEXT("ZenServer"));
+	if (ZenServerValue)
+	{
+		ZenServerInfo = MakeUnique<FZenServerInfo>();
+		TSharedPtr<FJsonObject> ZenServerObject = ZenServerValue->AsObject();
+		ZenServerInfo->HostName = ZenServerObject->Values.FindRef(TEXT("HostName"))->AsString();
+		ZenServerInfo->Port = uint16(ZenServerObject->Values.FindRef(TEXT("Port"))->AsNumber());
+		ZenServerInfo->ProjectId = ZenServerObject->Values.FindRef(TEXT("ProjectId"))->AsString();
+		ZenServerInfo->OplogId = ZenServerObject->Values.FindRef(TEXT("OplogId"))->AsString();
+	}
+
+	TSharedPtr<FJsonValue> FilesArrayValue = JsonObject->Values.FindRef(TEXT("Files"));
+	TArray<TSharedPtr<FJsonValue>> FilesArray = FilesArrayValue->AsArray();
+	FileNameByChunkIdMap.Reserve(FilesArray.Num());
+	for (const TSharedPtr<FJsonValue>& FileValue : FilesArray)
+	{
+		TSharedPtr<FJsonObject> FileObject = FileValue->AsObject();
+		FIoChunkId ChunkId = ChunkIdFromString(FileObject->Values.FindRef(TEXT("ChunkId"))->AsString());
+		FileNameByChunkIdMap.Add(ChunkId, FileObject->Values.FindRef(TEXT("Path"))->AsString());
+	}
+
+	TArray<TSharedPtr<FJsonValue>> PackagesArray = JsonObject->Values.FindRef(TEXT("Packages"))->AsArray();
+	PackageInfoByNameMap.Reserve(PackagesArray.Num());
+	for (const TSharedPtr<FJsonValue>& PackageValue : PackagesArray)
+	{
+		TSharedPtr<FJsonObject> PackageObject = PackageValue->AsObject();
+		FName PackageName = FName(PackageObject->Values.FindRef(TEXT("Name"))->AsString());
+
+		FPackageInfo& PackageInfo = PackageInfoByNameMap.FindOrAdd(PackageName);
+		PackageInfo.PackageName = PackageName;
+		PackageInfo.PackageChunkId = ChunkIdFromString(PackageObject->Values.FindRef(TEXT("PackageChunkId"))->AsString());
+
+		TSharedPtr<FJsonValue> BulkDataChunkIdsValue = PackageObject->Values.FindRef(TEXT("BulkDataChunkIds"));
+		if (BulkDataChunkIdsValue.IsValid())
+		{
+			TArray<TSharedPtr<FJsonValue>> BulkDataChunkIdsArray = BulkDataChunkIdsValue->AsArray();
+			PackageInfo.BulkDataChunkIds.Reserve(BulkDataChunkIdsArray.Num());
+			for (const TSharedPtr<FJsonValue>& BulkDataChunkIdValue : BulkDataChunkIdsArray)
+			{
+				PackageInfo.BulkDataChunkIds.Add(ChunkIdFromString(BulkDataChunkIdValue->AsString()));
+			}
+		}
+	}
+
+	return FIoStatus::Ok;
+}
+
+TArray<FPackageStoreManifest::FFileInfo> FPackageStoreManifest::GetFiles() const
+{
+	FScopeLock Lock(&CriticalSection);
+	TArray<FFileInfo> Files;
+	Files.Reserve(FileNameByChunkIdMap.Num());
+	for (const auto& KV : FileNameByChunkIdMap)
+	{
+		Files.Add({ KV.Value, KV.Key });
+	}
+	return Files;
+}
+
+TArray<FPackageStoreManifest::FPackageInfo> FPackageStoreManifest::GetPackages() const
+{
+	FScopeLock Lock(&CriticalSection);
+	TArray<FPackageInfo> Packages;
+	PackageInfoByNameMap.GenerateValueArray(Packages);
+	return Packages;
+}
+
+FPackageStoreManifest::FZenServerInfo& FPackageStoreManifest::EditZenServerInfo()
+{
+	FScopeLock Lock(&CriticalSection);
+	if (!ZenServerInfo)
+	{
+		ZenServerInfo = MakeUnique<FZenServerInfo>();
+	}
+	return *ZenServerInfo;
+}
+
+const FPackageStoreManifest::FZenServerInfo* FPackageStoreManifest::ReadZenServerInfo() const
+{
+	FScopeLock Lock(&CriticalSection);
+	return ZenServerInfo.Get();
+}
+
diff --git a/Engine/Source/Developer/IoStoreUtilities/Public/FilePackageStoreWriter.h b/Engine/Source/Developer/IoStoreUtilities/Public/FilePackageStoreWriter.h
deleted file mode 100644
diff --git a/Engine/Source/Developer/IoStoreUtilities/Public/PackageStoreManifest.h b/Engine/Source/Developer/IoStoreUtilities/Public/PackageStoreManifest.h
new file mode 100644
--- /dev/null
+++ b/Engine/Source/Developer/IoStoreUtilities/Public/PackageStoreManifest.h
@@ -0,0 +1,52 @@
+// Copyright Epic Games, Inc. All Rights Reserved.
+
+#pragma once
+
+#include "IO/IoDispatcher.h"
+
+class FPackageStoreManifest
+{
+public:
+	struct FFileInfo
+	{
+		FString FileName;
+		FIoChunkId ChunkId;
+	};
+
+	struct FPackageInfo
+	{
+		FName PackageName;
+		FIoChunkId PackageChunkId;
+		TArray<FIoChunkId> BulkDataChunkIds;
+	};
+	
+	struct FZenServerInfo
+	{
+		FString HostName;
+		uint16 Port;
+		FString ProjectId;
+		FString OplogId;
+	};
+
+	IOSTOREUTILITIES_API FPackageStoreManifest() = default;
+	IOSTOREUTILITIES_API ~FPackageStoreManifest() = default;
+
+	IOSTOREUTILITIES_API void BeginPackage(FName PackageName);
+	IOSTOREUTILITIES_API void AddPackageData(FName PackageName, const FString& FileName, const FIoChunkId& ChunkId);
+	IOSTOREUTILITIES_API void AddBulkData(FName PackageName, const FString& FileName, const FIoChunkId& ChunkId);
+
+	IOSTOREUTILITIES_API FIoStatus Save(const TCHAR* Filename) const;
+	IOSTOREUTILITIES_API FIoStatus Load(const TCHAR* Filename);
+
+	IOSTOREUTILITIES_API TArray<FFileInfo> GetFiles() const;
+	IOSTOREUTILITIES_API TArray<FPackageInfo> GetPackages() const;
+
+	IOSTOREUTILITIES_API FZenServerInfo& EditZenServerInfo();
+	IOSTOREUTILITIES_API const FZenServerInfo* ReadZenServerInfo() const;
+
+private:
+	mutable FCriticalSection CriticalSection;
+	TMap<FName, FPackageInfo> PackageInfoByNameMap;
+	TMap<FIoChunkId, FString> FileNameByChunkIdMap;
+	TUniquePtr<FZenServerInfo> ZenServerInfo;
+};
diff --git a/Engine/Source/Developer/IoStoreUtilities/Public/ZenStoreWriter.h b/Engine/Source/Developer/IoStoreUtilities/Public/ZenStoreWriter.h
--- a/Engine/Source/Developer/IoStoreUtilities/Public/ZenStoreWriter.h
+++ b/Engine/Source/Developer/IoStoreUtilities/Public/ZenStoreWriter.h
@@ -8,7 +8,8 @@
 #include "Serialization/CompactBinary.h"
 #include "Serialization/PackageWriter.h"
 #include "IO/PackageStore.h"
-#include "FilePackageStoreWriter.h"
+#include "PackageStoreWriter.h"
+#include "PackageStoreManifest.h"
 
 namespace UE {
 	class FZenStoreHttpClient;
diff --git a/Engine/Source/Editor/UnrealEd/Classes/CookOnTheSide/CookOnTheFlyServer.h b/Engine/Source/Editor/UnrealEd/Classes/CookOnTheSide/CookOnTheFlyServer.h
--- a/Engine/Source/Editor/UnrealEd/Classes/CookOnTheSide/CookOnTheFlyServer.h
+++ b/Engine/Source/Editor/UnrealEd/Classes/CookOnTheSide/CookOnTheFlyServer.h
@@ -74,7 +74,6 @@ enum class ECookByTheBookOptions
 	SkipSoftReferences =				0x00001000, // Don't follow soft references when cooking. Usually not viable for a real cook and the results probably wont load properly, but can be useful for debugging
 	SkipHardReferences =				0x00002000, // Don't follow hard references when cooking. Not viable for a real cook, only useful for debugging
 	FullLoadAndSave =					0x00004000, // Load all packages into memory and save them all at once in one tick for speed reasons. This requires a lot of RAM for large games.
-	IoStore =							0x00008000, // Cook for IoStore
 	CookAgainstFixedBase =				0x00010000, // If cooking DLC, assume that the base content can not be modified. 
 	DlcLoadMainAssetRegistry =			0x00020000, // If cooking DLC, populate the main game asset registry
 	ZenStore =							0x00040000, // Store cooked data in Zen Store
@@ -574,7 +573,6 @@ public:
 
 	bool IsUsingShaderCodeLibrary() const;
 
-	bool IsUsingIoStore() const;
 	bool IsUsingZenStore() const;
 
 	/**
diff --git a/Engine/Source/Editor/UnrealEd/Private/Commandlets/CookCommandlet.cpp b/Engine/Source/Editor/UnrealEd/Private/Commandlets/CookCommandlet.cpp
--- a/Engine/Source/Editor/UnrealEd/Private/Commandlets/CookCommandlet.cpp
+++ b/Engine/Source/Editor/UnrealEd/Private/Commandlets/CookCommandlet.cpp
@@ -896,8 +896,7 @@ bool UCookCommandlet::CookByTheBook( const TArray<ITargetPlatform*>& Platforms)
 	CookOptions |= Switches.Contains(TEXT("MAPSONLY")) ? ECookByTheBookOptions::MapsOnly : ECookByTheBookOptions::None;
 	CookOptions |= Switches.Contains(TEXT("NODEV")) ? ECookByTheBookOptions::NoDevContent : ECookByTheBookOptions::None;
 	CookOptions |= Switches.Contains(TEXT("FullLoadAndSave")) ? ECookByTheBookOptions::FullLoadAndSave : ECookByTheBookOptions::None;
-	CookOptions |= Switches.Contains(TEXT("IoStore")) ? ECookByTheBookOptions::IoStore : ECookByTheBookOptions::None;
-	CookOptions |= Switches.Contains(TEXT("ZenStore")) ? (ECookByTheBookOptions::ZenStore | ECookByTheBookOptions::IoStore) : ECookByTheBookOptions::None;
+	CookOptions |= Switches.Contains(TEXT("ZenStore")) ? ECookByTheBookOptions::ZenStore : ECookByTheBookOptions::None;
 	CookOptions |= Switches.Contains(TEXT("NoGameAlwaysCook")) ? ECookByTheBookOptions::NoGameAlwaysCookPackages : ECookByTheBookOptions::None;
 	CookOptions |= Switches.Contains(TEXT("DisableUnsolicitedPackages")) ? (ECookByTheBookOptions::SkipHardReferences | ECookByTheBookOptions::SkipSoftReferences) : ECookByTheBookOptions::None;
 	CookOptions |= Switches.Contains(TEXT("NoDefaultMaps")) ? ECookByTheBookOptions::NoDefaultMaps : ECookByTheBookOptions::None;
diff --git a/Engine/Source/Editor/UnrealEd/Private/CookOnTheFlyServer.cpp b/Engine/Source/Editor/UnrealEd/Private/CookOnTheFlyServer.cpp
--- a/Engine/Source/Editor/UnrealEd/Private/CookOnTheFlyServer.cpp
+++ b/Engine/Source/Editor/UnrealEd/Private/CookOnTheFlyServer.cpp
@@ -43,7 +43,6 @@
 #include "Engine/TextureLODSettings.h"
 #include "Engine/WorldComposition.h"
 #include "EngineGlobals.h"
-#include "FilePackageStoreWriter.h"
 #include "FileServerMessages.h"
 #include "GameDelegates.h"
 #include "GlobalShader.h"
@@ -1099,11 +1098,6 @@ bool UCookOnTheFlyServer::IsUsingShaderCodeLibrary() const
 	return IsCookByTheBookMode() && AllowShaderCompiling();
 }
 
-bool UCookOnTheFlyServer::IsUsingIoStore() const
-{
-	return (IsCookByTheBookMode() && CookByTheBookOptions->bIoStore) || (IsCookOnTheFlyMode() && CookOnTheFlyOptions.bZenStore);
-}
-
 bool UCookOnTheFlyServer::IsUsingZenStore() const
 {
 	return (IsCookByTheBookMode() && CookByTheBookOptions->bZenStore) || (IsCookOnTheFlyMode() && CookOnTheFlyOptions.bZenStore);
@@ -7455,11 +7449,6 @@ UE::Cook::FCookSavePackageContext* UCookOnTheFlyServer::CreateSaveContext(const
 		PackageWriter = new FZenStoreWriter(ResolvedRootPath, MetadataDirectoryPath, TargetPlatform);
 		WriterDebugName = TEXT("ZenStore");
 	}
-	else if (IsUsingIoStore())
-	{
-		PackageWriter = new FFilePackageStoreWriter(ResolvedRootPath, MetadataDirectoryPath, TargetPlatform);
-		WriterDebugName = TEXT("IoStoreFileWriter");
-	}
 	else
 	{
 		PackageWriter = new FLooseCookedPackageWriter(ResolvedRootPath, MetadataDirectoryPath, TargetPlatform,
@@ -7604,7 +7593,6 @@ void UCookOnTheFlyServer::StartCookByTheBook( const FCookByTheBookStartupOptions
 	CookByTheBookOptions->bSkipHardReferences = !!(CookOptions & ECookByTheBookOptions::SkipHardReferences);
 	CookByTheBookOptions->bSkipSoftReferences = !!(CookOptions & ECookByTheBookOptions::SkipSoftReferences);
 	CookByTheBookOptions->bFullLoadAndSave = !!(CookOptions & ECookByTheBookOptions::FullLoadAndSave);
-	CookByTheBookOptions->bIoStore = !!(CookOptions & ECookByTheBookOptions::IoStore);
 	CookByTheBookOptions->bZenStore = !!(CookOptions & ECookByTheBookOptions::ZenStore);
 	CookByTheBookOptions->bCookAgainstFixedBase = !!(CookOptions & ECookByTheBookOptions::CookAgainstFixedBase);
 	CookByTheBookOptions->bDlcLoadMainAssetRegistry = !!(CookOptions & ECookByTheBookOptions::DlcLoadMainAssetRegistry);
diff --git a/Engine/Source/Editor/UnrealEd/Private/Cooker/CookTypes.h b/Engine/Source/Editor/UnrealEd/Private/Cooker/CookTypes.h
--- a/Engine/Source/Editor/UnrealEd/Private/Cooker/CookTypes.h
+++ b/Engine/Source/Editor/UnrealEd/Private/Cooker/CookTypes.h
@@ -223,7 +223,6 @@ public:
 	bool							bSkipHardReferences = false;
 	bool							bSkipSoftReferences = false;
 	bool							bFullLoadAndSave = false;
-	bool							bIoStore = false;
 	bool							bZenStore = false;
 	bool							bCookAgainstFixedBase = false;
 	bool							bDlcLoadMainAssetRegistry = false;
diff --git a/Engine/Source/Runtime/SandboxFile/Private/IPlatformFileSandboxWrapper.cpp b/Engine/Source/Runtime/SandboxFile/Private/IPlatformFileSandboxWrapper.cpp
--- a/Engine/Source/Runtime/SandboxFile/Private/IPlatformFileSandboxWrapper.cpp
+++ b/Engine/Source/Runtime/SandboxFile/Private/IPlatformFileSandboxWrapper.cpp
@@ -9,7 +9,6 @@
 #include "Modules/ModuleManager.h"
 #include "HAL/IPlatformFileModule.h"
 #include "Templates/UniquePtr.h"
-#include "IO/IoDispatcherBackend.h"
 
 DEFINE_LOG_CATEGORY(SandboxFile);
 
@@ -184,31 +183,6 @@ bool FSandboxPlatformFile::Initialize(IPlatformFile* Inner, const TCHAR* CmdLine
 		}
 	}
 
-	FString GlobalUTocPath = FString::Printf(TEXT("%sglobal.utoc"), *SandboxDirectory);
-	if (FileExists(*GlobalUTocPath))
-	{
-		FIoStatus IoDispatcherInitStatus = FIoDispatcher::Initialize();
-		if (IoDispatcherInitStatus.IsOk())
-		{
-			FIoDispatcher& IoDispatcher = FIoDispatcher::Get();
-			TSharedRef<IIoDispatcherFileBackend> IoDispatcherFileBackend = CreateIoDispatcherFileBackend();
-			IoDispatcher.Mount(IoDispatcherFileBackend);
-			TIoStatusOr<FIoContainerId> IoDispatcherMountStatus = IoDispatcherFileBackend->Mount(*FPaths::ChangeExtension(GlobalUTocPath, TEXT("")), 0, FGuid(), FAES::FAESKey());
-			if (IoDispatcherMountStatus.IsOk())
-			{
-				UE_LOG(LogInit, Display, TEXT("Initialized I/O dispatcher"));
-			}
-			else
-			{
-				UE_LOG(LogInit, Error, TEXT("Failed to mount I/O dispatcher container: '%s'"), *IoDispatcherMountStatus.Status().ToString());
-			}
-		}
-		else
-		{
-			UE_LOG(LogInit, Error, TEXT("Failed to initialize I/O dispatcher: '%s'"), *IoDispatcherInitStatus.ToString());
-		}
-	}
-
 	return !!LowerLevel;
 }
 
