IoStore: Fix log spam caused by localized packages without a corresponding source package
#rnx

diff --git a/Engine/Source/Developer/IoStoreUtilities/Private/PackageStoreOptimizer.cpp b/Engine/Source/Developer/IoStoreUtilities/Private/PackageStoreOptimizer.cpp
--- a/Engine/Source/Developer/IoStoreUtilities/Private/PackageStoreOptimizer.cpp
+++ b/Engine/Source/Developer/IoStoreUtilities/Private/PackageStoreOptimizer.cpp
@@ -1106,16 +1106,6 @@ bool FPackageStoreOptimizer::VerifyRedirect(
 			}
 			return true;
 		}
-
-		if (TargetPackage.Region.Len() > 0)
-		{
-			// no update or verification required
-			UE_LOG(LogPackageStoreOptimizer, Verbose,
-				TEXT("For culture '%s': Localized package '%s' (0x%llX) is unique and does not override a source package."),
-				*TargetPackage.Region,
-				*TargetPackage.Name.ToString(),
-				TargetPackage.Id.ValueForDebugging());
-		}
 		return false;
 	}
 
@@ -1265,12 +1255,24 @@ void FPackageStoreOptimizer::ProcessRedirects(const TMap<FPackageId, FPackageSto
 		}
 		else
 		{
-			UE_LOG(LogPackageStoreOptimizer, Display,
-				TEXT("Skipping package redirect from '%s' (0x%llX) to '%s' (0x%llX) due to mismatching public exports."),
-				*Package->SourceName.ToString(),
-				SourcePackageId.ValueForDebugging(),
-				*Package->Name.ToString(),
-				Package->Id.ValueForDebugging());
+			if (Package->Region.Len() > 0 && !SourcePackage)
+			{
+				// no update or verification required
+				UE_LOG(LogPackageStoreOptimizer, Verbose,
+					TEXT("For culture '%s': Localized package '%s' (0x%llX) is unique and does not override a source package."),
+					*Package->Region,
+					*Package->Name.ToString(),
+					Package->Id.ValueForDebugging());
+			}
+			else
+			{
+				UE_LOG(LogPackageStoreOptimizer, Display,
+					TEXT("Skipping package redirect from '%s' (0x%llX) to '%s' (0x%llX) due to mismatching public exports."),
+					*Package->SourceName.ToString(),
+					SourcePackageId.ValueForDebugging(),
+					*Package->Name.ToString(),
+					Package->Id.ValueForDebugging());
+			}
 		}
 	}
 }
