ZenStoreWriter - Make CommitPackage take an array of attachments; remove separate variables for TargetDomainDependencies and send them in as an attachment.
ZenStoreHttpClient - Add IsConnected.
#rnx

diff --git a/Engine/Source/Developer/IoStoreUtilities/Internal/ZenStoreHttpClient.h b/Engine/Source/Developer/IoStoreUtilities/Internal/ZenStoreHttpClient.h
new file mode 100644
--- /dev/null
+++ b/Engine/Source/Developer/IoStoreUtilities/Internal/ZenStoreHttpClient.h
@@ -0,0 +1,74 @@
+// Copyright Epic Games, Inc. All Rights Reserved.
+
+#pragma once
+
+#include "Async/Future.h"
+#include "Containers/Array.h"
+#include "Containers/StringView.h"
+#include "HAL/CriticalSection.h"
+#include "HAL/Platform.h"
+#include "IO/IoDispatcher.h"
+#include "Misc/StringBuilder.h"
+#include "Serialization/CompactBinaryPackage.h"
+#include "Templates/UniquePtr.h"
+
+class FCbPackage;
+class FCbObject;
+
+namespace UE {
+	namespace Zen {
+		struct FZenHttpRequestPool;
+	}
+
+/**
+ * HTTP protocol implementation of Zen Store client interface
+ */
+class IOSTOREUTILITIES_API FZenStoreHttpClient
+{
+public:
+	FZenStoreHttpClient(const FStringView InHostName, uint16 InPort);
+	~FZenStoreHttpClient();
+
+	void Initialize(FStringView InProjectId, 
+					FStringView InOplogId, 
+					FStringView ServerRoot, 
+					FStringView EngineRoot, 
+					FStringView ProjectRoot);
+	void EstablishWritableOpLog(FStringView InProjectId, FStringView InOplogId, bool bFullBuild);
+
+	void InitializeReadOnly(FStringView InProjectId, FStringView InOplogId);
+
+	bool IsConnected() const;
+
+	void StartBuildPass();
+	TIoStatusOr<uint64> EndBuildPass(FCbPackage OpEntry);
+
+	TFuture<TIoStatusOr<uint64>> AppendOp(FCbPackage OpEntry);
+
+	TIoStatusOr<uint64> GetChunkSize(const FIoChunkId& Id);
+	TIoStatusOr<FIoBuffer> ReadChunk(const FIoChunkId& Id, uint64 Offset = 0, uint64 Size = ~0ull);
+	TIoStatusOr<FIoBuffer> ReadOpLogAttachment(FStringView Id, uint64 Offset = 0, uint64 Size = ~0ull);
+
+	TFuture<TIoStatusOr<FCbObject>> GetOplog();
+	TFuture<TIoStatusOr<FCbObject>> GetFiles();
+
+	static const UTF8CHAR* FindOrAddAttachmentId(FUtf8StringView AttachmentText);
+	static const UTF8CHAR* FindAttachmentId(FUtf8StringView AttachmentText);
+
+private:
+	TIoStatusOr<FIoBuffer> ReadOpLogUri(FStringBuilderBase& ChunkUri, uint64 Offset, uint64 Size);
+
+	FString	HostName;
+	uint16 Port;
+	TUniquePtr<Zen::FZenHttpRequestPool> RequestPool;
+	FString OplogPath;
+	FString OplogNewEntryPath;
+	FString OplogPrepNewEntryPath;
+	FString TempDirPath;
+	uint64 StandaloneThresholdBytes = 1 * 1024 * 1024;
+	bool bAllowRead = false;
+	bool bAllowEdit = false;
+	bool bConnectionSucceeded = false;
+};
+
+}
diff --git a/Engine/Source/Developer/IoStoreUtilities/Private/ZenStoreHttpClient.cpp b/Engine/Source/Developer/IoStoreUtilities/Private/ZenStoreHttpClient.cpp
--- a/Engine/Source/Developer/IoStoreUtilities/Private/ZenStoreHttpClient.cpp
+++ b/Engine/Source/Developer/IoStoreUtilities/Private/ZenStoreHttpClient.cpp
@@ -2,16 +2,12 @@
 
 #include "ZenStoreHttpClient.h"
 
-#if PLATFORM_WINDOWS
+#include "Algo/BinarySearch.h"
+#include "Algo/IsSorted.h"
+#include "Algo/Sort.h"
+#include "Misc/ScopeRWLock.h"
 
-#if PLATFORM_WINDOWS || PLATFORM_HOLOLENS
-#include "Windows/WindowsHWrapper.h"
-#include "Windows/AllowWindowsPlatformTypes.h"
-#endif
-#include "curl/curl.h"
-#if PLATFORM_WINDOWS || PLATFORM_HOLOLENS
-#include "Windows/HideWindowsPlatformTypes.h"
-#endif
+#if PLATFORM_WINDOWS
 
 #include "Async/Async.h"
 #include "HAL/PlatformFileManager.h"
@@ -73,6 +69,7 @@ FZenStoreHttpClient::Initialize(FStringView InProjectId,
 		if (Res == Zen::FZenHttpRequest::Result::Success && Request->GetResponseCode() == 200)
 		{
 			UE_LOG(LogZenStore, Display, TEXT("Zen project '%s' already exists"), *FString(InProjectId));
+			bConnectionSucceeded = true;
 		}
 		else
 		{
@@ -91,16 +88,19 @@ FZenStoreHttpClient::Initialize(FStringView InProjectId,
 			if (Res != Zen::FZenHttpRequest::Result::Success)
 			{
 				UE_LOG(LogZenStore, Error, TEXT("Zen project '%s' creation FAILED"), *FString(InProjectId));
+				bConnectionSucceeded = false;
 
 				// TODO: how to recover / handle this?
 			}
 			else if (Request->GetResponseCode() == 201)
 			{
 				UE_LOG(LogZenStore, Display, TEXT("Zen project '%s' created"), *FString(InProjectId));
+				bConnectionSucceeded = true;
 			}
 			else
 			{
 				UE_LOG(LogZenStore, Warning, TEXT("Zen project '%s' creation returned success but not HTTP 201"), *FString(InProjectId));
+				bConnectionSucceeded = true;
 			}
 		}
 	}
@@ -115,6 +115,8 @@ FZenStoreHttpClient::Initialize(FStringView InProjectId,
 
 void FZenStoreHttpClient::EstablishWritableOpLog(FStringView InProjectId, FStringView InOplogId, bool bFullBuild)
 {
+	check(IsConnected());
+
 	UE::Zen::FZenScopedRequestPtr Request(RequestPool.Get());
 
 	if (bFullBuild)
@@ -191,10 +193,16 @@ void FZenStoreHttpClient::InitializeReadOnly(FStringView InProjectId, FStringVie
 		if (Res != Zen::FZenHttpRequest::Result::Success || Request->GetResponseCode() != 200)
 		{
 			UE_LOG(LogZenStore, Fatal, TEXT("Zen project '%s' not found"), *FString(InProjectId));
+			bConnectionSucceeded = false;
+		}
+		else
+		{
+			bConnectionSucceeded = true;
 		}
 	}
 
 	// Establish oplog
+	check(IsConnected());
 
 	{
 		UE::Zen::FZenScopedRequestPtr Request(RequestPool.Get());
@@ -618,3 +626,82 @@ TFuture<TIoStatusOr<FCbObject>> FZenStoreHttpClient::GetFiles()
 }
 
 #endif // PLATFORM_WINDOWS
+
+namespace UE
+{
+
+bool FZenStoreHttpClient::IsConnected() const
+{
+	return bConnectionSucceeded;
+}
+
+namespace Zen::Private
+{
+
+TArray<TUniquePtr<UTF8CHAR[]>> AttachmentIds;
+FRWLock AttachmentIdsLock;
+
+}
+
+const UTF8CHAR* FZenStoreHttpClient::FindOrAddAttachmentId(FUtf8StringView AttachmentText)
+{
+	FRWScopeLock AttachmentIdScopeLock(Zen::Private::AttachmentIdsLock, SLT_ReadOnly);
+	bool bReadOnly = true;
+
+	for (;;)
+	{
+		int32 Index = Algo::LowerBound(Zen::Private::AttachmentIds, AttachmentText,
+			[](const TUniquePtr<UTF8CHAR[]>& Existing, FUtf8StringView AttachmentText)
+			{
+				return FUtf8StringView(Existing.Get()).Compare(AttachmentText, ESearchCase::IgnoreCase) < 0;
+			});
+		UTF8CHAR* Existing = nullptr;
+		if (Index != Zen::Private::AttachmentIds.Num())
+		{
+			Existing = Zen::Private::AttachmentIds[Index].Get();
+		}
+ 		if (Existing == nullptr || !FUtf8StringView(Existing).Equals(AttachmentText, ESearchCase::IgnoreCase))
+		{
+			if (bReadOnly)
+			{
+				AttachmentIdScopeLock.ReleaseReadOnlyLockAndAcquireWriteLock_USE_WITH_CAUTION();
+				bReadOnly = false;
+				continue;
+			}
+
+			auto StrDupNew = [](const UTF8CHAR* Data, int32 Len)
+			{
+				UTF8CHAR* Dup = new UTF8CHAR[Len + 1];
+				FMemory::Memcpy(Dup, Data, Len * sizeof(UTF8CHAR));
+				Dup[Len] = UTF8CHAR('\0');
+				return Dup;
+			};
+			Existing = StrDupNew(AttachmentText.GetData(), AttachmentText.Len());
+			Zen::Private::AttachmentIds.Insert(TUniquePtr<UTF8CHAR[]>(Existing), Index);
+		}
+		return Existing;
+	}
+}
+
+const UTF8CHAR* FZenStoreHttpClient::FindAttachmentId(FUtf8StringView AttachmentText)
+{
+	FReadScopeLock AttachmentIdScopeLock(Zen::Private::AttachmentIdsLock);
+
+	int32 Index = Algo::LowerBound(Zen::Private::AttachmentIds, AttachmentText,
+		[](const TUniquePtr<UTF8CHAR[]>& Existing, FUtf8StringView AttachmentText)
+		{
+			return FUtf8StringView(Existing.Get()).Compare(AttachmentText, ESearchCase::IgnoreCase) < 0;
+		});
+	if (Index == Zen::Private::AttachmentIds.Num())
+	{
+		return nullptr;
+	}
+	const UTF8CHAR* Existing = Zen::Private::AttachmentIds[Index].Get();
+	if (!FUtf8StringView(Existing).Equals(AttachmentText, ESearchCase::IgnoreCase))
+	{
+		return nullptr;
+	}
+	return Existing;
+}
+
+}
diff --git a/Engine/Source/Developer/IoStoreUtilities/Private/ZenStoreHttpClient.h b/Engine/Source/Developer/IoStoreUtilities/Private/ZenStoreHttpClient.h
deleted file mode 100644
diff --git a/Engine/Source/Developer/IoStoreUtilities/Private/ZenStoreWriter.cpp b/Engine/Source/Developer/IoStoreUtilities/Private/ZenStoreWriter.cpp
--- a/Engine/Source/Developer/IoStoreUtilities/Private/ZenStoreWriter.cpp
+++ b/Engine/Source/Developer/IoStoreUtilities/Private/ZenStoreWriter.cpp
@@ -5,6 +5,9 @@
 #include "ZenFileSystemManifest.h"
 #include "PackageStoreOptimizer.h"
 
+#include "Algo/BinarySearch.h"
+#include "Algo/IsSorted.h"
+#include "Algo/Sort.h"
 #include "Async/Async.h"
 #include "Serialization/CompactBinaryPackage.h"
 #include "Serialization/CompactBinaryWriter.h"
@@ -224,6 +227,22 @@ private:
 	FHttpRequest* FirstFreeHttpRequest = nullptr;
 };
 
+TArray<const UTF8CHAR*> FZenStoreWriter::ReservedOplogKeys;
+
+void FZenStoreWriter::StaticInit()
+{
+	if (ReservedOplogKeys.Num() > 0)
+	{
+		return;
+	}
+
+	ReservedOplogKeys.Append({ UTF8TEXT("files"), UTF8TEXT("key"), UTF8TEXT("package"), UTF8TEXT("packagestoreentry") });
+	Algo::Sort(ReservedOplogKeys, [](const UTF8CHAR* A, const UTF8CHAR* B)
+		{
+			return FUtf8StringView(A).Compare(FUtf8StringView(B), ESearchCase::IgnoreCase) < 0;
+		});;
+}
+
 FZenStoreWriter::FZenStoreWriter(
 	const FString& InOutputPath, 
 	const FString& InMetadataDirectoryPath, 
@@ -236,6 +255,8 @@ FZenStoreWriter::FZenStoreWriter(
 	, CookMode(ICookedPackageWriter::FCookInfo::CookByTheBookMode)
 	, bInitialized(false)
 {
+	StaticInit();
+
 	FString HostName = TEXT("localhost");
 	uint16 Port = 1337;
 	FString ProjectId = FApp::GetProjectName();
@@ -258,7 +279,7 @@ FZenStoreWriter::FZenStoreWriter(
 							OplogId, 
 							AbsServerRoot, 
 							AbsEngineDir, 
-							AbsProjectDir );
+							AbsProjectDir);
 
 	PackageStoreOptimizer->Initialize(InTargetPlatform);
 
@@ -421,6 +442,8 @@ void FZenStoreWriter::BeginCook(const FCookInfo& Info)
 
 	if (!bInitialized)
 	{
+		UE_CLOG(!HttpClient->IsConnected(), LogZenStoreWriter, Fatal, TEXT("Failed to connect to ZenServer"));
+
 		FString ProjectId = FApp::GetProjectName();
 		FString OplogId = TargetPlatform.PlatformName();
 		HttpClient->EstablishWritableOpLog(ProjectId, OplogId, Info.bFullBuild);
@@ -453,12 +476,35 @@ void FZenStoreWriter::BeginCook(const FCookInfo& Info)
 								const int64	PkgDiskSize = PackageObj["disksize"].AsUInt64();
 								FPackageStoreEntryResource Entry = FPackageStoreEntryResource::FromCbObject(OplogObj["packagestoreentry"].AsObject());
 								const FName PackageName = Entry.PackageName;
-								FIoHash DependenciesHash = OplogObj["dependencies"].AsHash();
+
 								const int32 Index = PackageStoreEntries.Num();
 
 								PackageStoreEntries.Add(MoveTemp(Entry));
-								CookedPackagesInfo.Add(FCookedPackageInfo{ PackageName, IoHashToMD5(PkgHash), PkgGuid, PkgDiskSize, DependenciesHash });
+								FOplogCookInfo& CookInfo = CookedPackagesInfo.Add_GetRef(
+									FOplogCookInfo{
+										FCookedPackageInfo {PackageName, IoHashToMD5(PkgHash), PkgGuid, PkgDiskSize }
+									});
 								PackageNameToIndex.Add(PackageName, Index);
+
+								for (FCbFieldView Field : OplogObj)
+								{
+									FUtf8StringView FieldName = Field.GetName();
+									if (IsReservedOplogKey(FieldName))
+									{
+										continue;
+									}
+									if (Field.IsHash())
+									{
+										const UTF8CHAR* AttachmentId = UE::FZenStoreHttpClient::FindOrAddAttachmentId(FieldName);
+										CookInfo.Attachments.Add({ AttachmentId, Field.AsHash() });
+									}
+								}
+								CookInfo.Attachments.Shrink();
+								check(Algo::IsSorted(CookInfo.Attachments,
+									[](const FOplogCookInfo::FAttachment& A, const FOplogCookInfo::FAttachment& B)
+									{
+										return FUtf8StringView(A.Key).Compare(FUtf8StringView(B.Key), ESearchCase::IgnoreCase) < 0;
+									}));
 							}
 						}
 					}
@@ -584,6 +630,17 @@ void FZenStoreWriter::BeginPackage(const FBeginPackageInfo& Info)
 	PackageStoreManifest.BeginPackage(Info.PackageName);
 }
 
+bool FZenStoreWriter::IsReservedOplogKey(FUtf8StringView Key)
+{
+	int32 Index = Algo::LowerBound(ReservedOplogKeys, Key,
+		[](const UTF8CHAR* Existing, FUtf8StringView Key)
+		{
+			return FUtf8StringView(Existing).Compare(Key, ESearchCase::IgnoreCase) < 0;
+		});
+	return Index != ReservedOplogKeys.Num() &&
+		FUtf8StringView(ReservedOplogKeys[Index]).Equals(Key, ESearchCase::IgnoreCase);
+}
+
 void FZenStoreWriter::CommitPackage(const FCommitPackageInfo& Info)
 {
 	TRACE_CPUPROFILER_EVENT_SCOPE(FZenStoreWriter::CommitPackage);
@@ -615,18 +672,36 @@ void FZenStoreWriter::CommitPackage(const FCommitPackageInfo& Info)
 			FCbAttachment PkgDataAttachment(IoBufferToSharedBuffer(PkgData.Payload));
 			Pkg.AddAttachment(PkgDataAttachment);
 
-			TOptional<FCbAttachment> DependenciesAttachment;
-			FIoHash DependenciesHash;
-			if (Info.TargetDomainDependencies)
+			CommitEventArgs.EntryIndex = PackageStoreEntries.Num();
+			PackageStoreEntries.Add(PkgData.PackageStoreEntry);
+			FOplogCookInfo& CookInfo = CookedPackagesInfo.Add_GetRef(FOplogCookInfo{
+				FCookedPackageInfo { Info.PackageName, IoHashToMD5(PkgDataAttachment.GetHash()), Info.PackageGuid, PkgDiskSize } });
+			int32 NumAttachments = Info.Attachments.Num();
+			TArray<FCbAttachment, TInlineAllocator<2>> CbAttachments;
+			if (NumAttachments)
 			{
-				DependenciesAttachment.Emplace(Info.TargetDomainDependencies);
-				Pkg.AddAttachment(*DependenciesAttachment);
-				DependenciesHash = DependenciesAttachment->GetHash();
+				TArray<const FCommitAttachmentInfo*, TInlineAllocator<2>> SortedAttachments;
+				SortedAttachments.Reserve(NumAttachments);
+				for (const FCommitAttachmentInfo& AttachmentInfo : Info.Attachments)
+				{
+					SortedAttachments.Add(&AttachmentInfo);
+				}
+				SortedAttachments.Sort([](const FCommitAttachmentInfo& A, const FCommitAttachmentInfo& B)
+					{
+						return A.Key.Compare(B.Key, ESearchCase::IgnoreCase) < 0;
+					});
+				CbAttachments.Reserve(NumAttachments);
+				CookInfo.Attachments.Reserve(NumAttachments);
+				for (const FCommitAttachmentInfo* AttachmentInfo : SortedAttachments)
+				{
+					check(!IsReservedOplogKey(AttachmentInfo->Key));
+					const FCbAttachment& CbAttachment = CbAttachments.Emplace_GetRef(AttachmentInfo->Value);
+					Pkg.AddAttachment(CbAttachment);
+					CookInfo.Attachments.Add(FOplogCookInfo::FAttachment{
+						UE::FZenStoreHttpClient::FindOrAddAttachmentId(AttachmentInfo->Key), CbAttachment.GetHash() });
+				}
 			}
 
-			CommitEventArgs.EntryIndex = PackageStoreEntries.Num();
-			PackageStoreEntries.Add(PkgData.PackageStoreEntry);
-			CookedPackagesInfo.Add(FCookedPackageInfo { Info.PackageName, IoHashToMD5(PkgDataAttachment.GetHash()), Info.PackageGuid, PkgDiskSize, DependenciesHash });
 			PackageNameToIndex.Add(Info.PackageName, CookedPackagesInfo.Num()-1);
 
 			FCbWriter PackageObj;
@@ -689,9 +764,11 @@ void FZenStoreWriter::CommitPackage(const FCommitPackageInfo& Info)
 				PackageObj.EndArray();
 			}
 
-			if (DependenciesAttachment)
+			for (int32 Index = 0; Index < NumAttachments; ++Index)
 			{
-				PackageObj << "dependencies" << *DependenciesAttachment;
+				FCbAttachment& CbAttachment = CbAttachments[Index];
+				FOplogCookInfo::FAttachment& CookInfoAttachment = CookInfo.Attachments[Index];
+				PackageObj << CookInfoAttachment.Key << CbAttachment;
 			}
 
 			PackageObj.EndObject();
@@ -734,10 +811,14 @@ void FZenStoreWriter::Flush()
 
 void FZenStoreWriter::GetCookedPackages(TArray<FCookedPackageInfo>& OutCookedPackages) 
 {
-	OutCookedPackages.Append(CookedPackagesInfo);
+	OutCookedPackages.Reserve(OutCookedPackages.Num() + CookedPackagesInfo.Num());
+	for (FOplogCookInfo& CookedPackageInfo : CookedPackagesInfo)
+	{
+		OutCookedPackages.Add(CookedPackageInfo.CookInfo);
+	}
 }
 
-FCbObject FZenStoreWriter::GetTargetDomainDependencies(FName PackageName)
+FCbObject FZenStoreWriter::GetOplogAttachment(FName PackageName, FUtf8StringView AttachmentKey)
 {
 	const int32* Idx = PackageNameToIndex.Find(PackageName);
 	if (!Idx)
@@ -745,8 +826,29 @@ FCbObject FZenStoreWriter::GetTargetDomainDependencies(FName PackageName)
 		return FCbObject();
 	}
 
-	const ICookedPackageWriter::FCookedPackageInfo& Info = CookedPackagesInfo[*Idx];
-	TIoStatusOr<FIoBuffer> BufferResult = HttpClient->ReadOpLogAttachment(WriteToString<48>(Info.TargetDomainDependencies));
+	const UTF8CHAR* AttachmentId = UE::FZenStoreHttpClient::FindAttachmentId(AttachmentKey);
+	if (!AttachmentId)
+	{
+		return FCbObject();
+	}
+	FUtf8StringView AttachmentIdView(AttachmentId);
+
+	const FOplogCookInfo& CookInfo = CookedPackagesInfo[*Idx];
+	int32 AttachmentIndex = Algo::LowerBound(CookInfo.Attachments, AttachmentIdView,
+		[](const FOplogCookInfo::FAttachment& Existing, FUtf8StringView AttachmentIdView)
+		{
+			return FUtf8StringView(Existing.Key).Compare(AttachmentIdView, ESearchCase::IgnoreCase) < 0;
+		});
+	if (AttachmentIndex == CookInfo.Attachments.Num())
+	{
+		return FCbObject();
+	}
+	const FOplogCookInfo::FAttachment& Existing = CookInfo.Attachments[AttachmentIndex];
+	if (!FUtf8StringView(Existing.Key).Equals(AttachmentIdView, ESearchCase::IgnoreCase))
+	{
+		return FCbObject();
+	}
+	TIoStatusOr<FIoBuffer> BufferResult = HttpClient->ReadOpLogAttachment(WriteToString<48>(Existing.Hash));
 	if (!BufferResult.IsOk())
 	{
 		return FCbObject();
@@ -780,7 +882,7 @@ void FZenStoreWriter::RemoveCookedPackages(TArrayView<const FName> PackageNamesT
 	const int32 NumPackagesToKeep = PackageIndicesToKeep.Num();
 	
 	TArray<FPackageStoreEntryResource> PreviousPackageStoreEntries = MoveTemp(PackageStoreEntries);
-	TArray<FCookedPackageInfo> PreviousCookedPackageInfo = MoveTemp(CookedPackagesInfo);
+	TArray<FOplogCookInfo> PreviousCookedPackageInfo = MoveTemp(CookedPackagesInfo);
 	PackageNameToIndex.Empty();
 
 	if (NumPackagesToKeep > 0)
@@ -792,7 +894,7 @@ void FZenStoreWriter::RemoveCookedPackages(TArrayView<const FName> PackageNamesT
 		int32 EntryIndex = 0;
 		for (int32 Idx : PackageIndicesToKeep)
 		{
-			const FName PackageName = PreviousCookedPackageInfo[Idx].PackageName;
+			const FName PackageName = PreviousCookedPackageInfo[Idx].CookInfo.PackageName;
 
 			PackageStoreEntries.Add(MoveTemp(PreviousPackageStoreEntries[Idx]));
 			CookedPackagesInfo.Add(MoveTemp(PreviousCookedPackageInfo[Idx]));
diff --git a/Engine/Source/Developer/IoStoreUtilities/Public/ZenStoreWriter.h b/Engine/Source/Developer/IoStoreUtilities/Public/ZenStoreWriter.h
--- a/Engine/Source/Developer/IoStoreUtilities/Public/ZenStoreWriter.h
+++ b/Engine/Source/Developer/IoStoreUtilities/Public/ZenStoreWriter.h
@@ -2,14 +2,15 @@
 
 #pragma once
 
-#include "Serialization/AsyncLoading2.h"
-#include "IO/IoDispatcher.h"
 #include "Containers/Map.h"
-#include "Serialization/CompactBinary.h"
-#include "Serialization/PackageWriter.h"
+#include "Containers/StringView.h"
+#include "IO/IoDispatcher.h"
 #include "IO/PackageStore.h"
-#include "PackageStoreWriter.h"
 #include "PackageStoreManifest.h"
+#include "PackageStoreWriter.h"
+#include "Serialization/AsyncLoading2.h"
+#include "Serialization/CompactBinary.h"
+#include "Serialization/PackageWriter.h"
 
 namespace UE {
 	class FZenStoreHttpClient;
@@ -56,7 +57,7 @@ public:
 	IOSTOREUTILITIES_API void WriteIoStorePackageData(const FPackageInfo& Info, const FIoBuffer& PackageData, const FPackageStoreEntryResource& PackageStoreEntry, const TArray<FFileRegion>& FileRegions);
 
 	IOSTOREUTILITIES_API virtual void GetCookedPackages(TArray<FCookedPackageInfo>& OutCookedPackages) override;
-	IOSTOREUTILITIES_API virtual FCbObject GetTargetDomainDependencies(FName PackageName) override;
+	IOSTOREUTILITIES_API virtual FCbObject GetOplogAttachment(FName PackageName, FUtf8StringView AttachmentKey) override;
 	IOSTOREUTILITIES_API virtual void RemoveCookedPackages(TArrayView<const FName> PackageNamesToRemove) override;
 	IOSTOREUTILITIES_API virtual void RemoveCookedPackages() override;
 
@@ -103,6 +104,18 @@ private:
 		double TotalRequestTime = 0.0;
 	};
 
+	struct FOplogCookInfo
+	{
+		struct FAttachment
+		{
+			const UTF8CHAR* Key;
+			FIoHash Hash;
+		};
+
+		ICookedPackageWriter::FCookedPackageInfo CookInfo;
+		TArray<FAttachment> Attachments;
+	};
+
 	FRWLock								PackagesLock;
 	TMap<FName,FPendingPackageState>	PendingPackages;
 	TUniquePtr<UE::FZenStoreHttpClient>	HttpClient;
@@ -115,7 +128,7 @@ private:
 	FPackageStoreManifest				PackageStoreManifest;
 	TUniquePtr<FPackageStoreOptimizer>	PackageStoreOptimizer;
 	TArray<FPackageStoreEntryResource>	PackageStoreEntries;
-	TArray<FCookedPackageInfo>			CookedPackagesInfo;
+	TArray<FOplogCookInfo>				CookedPackagesInfo;
 	TMap<FName, int32>					PackageNameToIndex;
 
 	TUniquePtr<FZenFileSystemManifest>	ZenFileSystemManifest;
@@ -131,4 +144,8 @@ private:
 	FZenStats							ZenStats;
 
 	bool								bInitialized;
+
+	static void StaticInit();
+	static bool IsReservedOplogKey(FUtf8StringView Key);
+	static TArray<const UTF8CHAR*>		ReservedOplogKeys;
 };
diff --git a/Engine/Source/Editor/UnrealEd/Private/CookOnTheFlyServer.cpp b/Engine/Source/Editor/UnrealEd/Private/CookOnTheFlyServer.cpp
--- a/Engine/Source/Editor/UnrealEd/Private/CookOnTheFlyServer.cpp
+++ b/Engine/Source/Editor/UnrealEd/Private/CookOnTheFlyServer.cpp
@@ -4477,8 +4477,8 @@ void FSaveCookedPackageContext::FinishPlatform()
 		check(AssetPackageData);
 
 		FCbObject TargetDomainDependencies;
-		UE_SCOPED_HIERARCHICAL_COOKTIMER(TargetDomainDependencies);
-		TargetDomainDependencies = UE::TargetDomain::CollectDependenciesObject(Package, TargetPlatform, nullptr /* ErrorMessage */);
+			UE_SCOPED_HIERARCHICAL_COOKTIMER(TargetDomainDependencies);
+			TargetDomainDependencies = UE::TargetDomain::CollectDependenciesObject(Package, TargetPlatform, nullptr /* ErrorMessage */);
 
 		ICookedPackageWriter::FCommitPackageInfo Info;
 		Info.bSucceeded = bSuccessful;
@@ -4486,7 +4486,7 @@ void FSaveCookedPackageContext::FinishPlatform()
 		PRAGMA_DISABLE_DEPRECATION_WARNINGS
 		Info.PackageGuid = AssetPackageData->PackageGuid;
 		PRAGMA_ENABLE_DEPRECATION_WARNINGS
-		Info.TargetDomainDependencies = TargetDomainDependencies;
+		Info.Attachments.Add({ "Dependencies", TargetDomainDependencies });
 
 		PackageWriter->CommitPackage(Info);
 	}
diff --git a/Engine/Source/Editor/UnrealEd/Private/Cooker/LooseCookedPackageWriter.cpp b/Engine/Source/Editor/UnrealEd/Private/Cooker/LooseCookedPackageWriter.cpp
--- a/Engine/Source/Editor/UnrealEd/Private/Cooker/LooseCookedPackageWriter.cpp
+++ b/Engine/Source/Editor/UnrealEd/Private/Cooker/LooseCookedPackageWriter.cpp
@@ -161,17 +161,18 @@ void FLooseCookedPackageWriter::GetCookedPackages(TArray<FCookedPackageInfo>& Ou
 			if (bExistsOnDisk)
 			{
 				PRAGMA_DISABLE_DEPRECATION_WARNINGS;
-				OutCookedPackages.Add(FCookedPackageInfo{ PackageName, PackageData->CookedHash,PackageData->PackageGuid,
-					PackageData->DiskSize, FIoHash() /** TargetDomainDependencies are not implemented by FLooseCookedPackageWriter */ });
+				OutCookedPackages.Add(FCookedPackageInfo{ PackageName, PackageData->CookedHash,
+					PackageData->PackageGuid, PackageData->DiskSize });
+				// Oplog attachments are not implemented by FLooseCookedPackageWriter
 				PRAGMA_ENABLE_DEPRECATION_WARNINGS;
 			}
 		}
 	}
 }
 
-FCbObject FLooseCookedPackageWriter::GetTargetDomainDependencies(FName PackageName)
+FCbObject FLooseCookedPackageWriter::GetOplogAttachment(FName PackageName, FUtf8StringView AttachmentKey)
 {
-	/** TargetDomainDependencies are not implemented by FLooseCookedPackageWriter */
+	/** Oplog attachments are not implemented by FLooseCookedPackageWriter */
 	return FCbObject();
 }
 
diff --git a/Engine/Source/Editor/UnrealEd/Private/Cooker/LooseCookedPackageWriter.h b/Engine/Source/Editor/UnrealEd/Private/Cooker/LooseCookedPackageWriter.h
--- a/Engine/Source/Editor/UnrealEd/Private/Cooker/LooseCookedPackageWriter.h
+++ b/Engine/Source/Editor/UnrealEd/Private/Cooker/LooseCookedPackageWriter.h
@@ -42,7 +42,7 @@ public:
 	void EndCook() override;
 	void Flush() override;
 	void GetCookedPackages(TArray<FCookedPackageInfo>& OutCookedPackages) override;
-	FCbObject GetTargetDomainDependencies(FName PackageName) override;
+	FCbObject GetOplogAttachment(FName PackageName, FUtf8StringView AttachmentKey) override;
 	void RemoveCookedPackages(TArrayView<const FName> PackageNamesToRemove) override;
 	void RemoveCookedPackages() override;
 	FAssetRegistryState* ReleasePreviousAssetRegistry() override;
diff --git a/Engine/Source/Editor/UnrealEd/Private/TargetDomain/TargetDomainUtils.cpp b/Engine/Source/Editor/UnrealEd/Private/TargetDomain/TargetDomainUtils.cpp
--- a/Engine/Source/Editor/UnrealEd/Private/TargetDomain/TargetDomainUtils.cpp
+++ b/Engine/Source/Editor/UnrealEd/Private/TargetDomain/TargetDomainUtils.cpp
@@ -183,16 +183,16 @@ FCbObject CollectDependenciesObject(UPackage* Package, const ITargetPlatform* Ta
 bool TryFetchKeyAndDependencies(ICookedPackageWriter* PackageWriter, FName PackageName, const ITargetPlatform* TargetPlatform,
 	FIoHash* OutHash, TArray<FName>* OutBuildDependencies, TArray<FName>* OutRuntimeOnlyDependencies, FString* OutErrorMessage)
 {
-	FCbObject DependenciesObj = PackageWriter->GetTargetDomainDependencies(PackageName);
+	FCbObject DependenciesObj = PackageWriter->GetOplogAttachment(PackageName, "Dependencies");
 	FIoHash StoredKey = DependenciesObj["targetdomainkey"].AsHash();
 	if (StoredKey.IsZero())
-	{
-		if (OutErrorMessage)
 		{
+			if (OutErrorMessage)
+			{
 			*OutErrorMessage = TEXT("Dependencies not in oplog.");
+			}
+			return false;
 		}
-		return false;
-	}
 
 	TArray<FName> BuildDependencies;
 	if (OutBuildDependencies)
@@ -200,7 +200,7 @@ bool TryFetchKeyAndDependencies(ICookedPackageWriter* PackageWriter, FName Packa
 		OutBuildDependencies->Reset();
 	}
 	else
-	{
+		{
 		OutBuildDependencies = &BuildDependencies;
 	}
 
@@ -229,14 +229,14 @@ bool TryFetchKeyAndDependencies(ICookedPackageWriter* PackageWriter, FName Packa
 
 	if (OutRuntimeOnlyDependencies)
 	{
-		for (FCbFieldView DepObj : DependenciesObj["runtimeonlydependencies"])
+	for (FCbFieldView DepObj : DependenciesObj["runtimeonlydependencies"])
+	{
+		if (FString DependencyName(DepObj.AsString()); !DependencyName.IsEmpty())
 		{
-			if (FString DependencyName(DepObj.AsString()); !DependencyName.IsEmpty())
-			{
 				OutRuntimeOnlyDependencies->Add(FName(*DependencyName));
-			}
 		}
 	}
+	}
 
 	return true;
 }
diff --git a/Engine/Source/Runtime/Core/Public/Serialization/PackageWriter.h b/Engine/Source/Runtime/Core/Public/Serialization/PackageWriter.h
--- a/Engine/Source/Runtime/Core/Public/Serialization/PackageWriter.h
+++ b/Engine/Source/Runtime/Core/Public/Serialization/PackageWriter.h
@@ -53,11 +53,16 @@ public:
 	  */
 	virtual void BeginPackage(const FBeginPackageInfo& Info) = 0;
 
+	struct FCommitAttachmentInfo
+	{
+		FUtf8StringView Key;
+		FCbObject Value;
+	};
 	struct FCommitPackageInfo
 	{
 		FName PackageName;
 		FGuid PackageGuid;
-		FCbObject TargetDomainDependencies;
+		TArray<FCommitAttachmentInfo> Attachments;
 		bool bSucceeded = false;
 	};
 
@@ -183,7 +188,6 @@ public:
 		FMD5Hash Hash;
 		FGuid PackageGuid;
 		int64 DiskSize = -1;
-		FIoHash TargetDomainDependencies;
 	};
 
 	/**
@@ -192,10 +196,10 @@ public:
 	virtual void GetCookedPackages(TArray<FCookedPackageInfo>& OutCookedPackages) = 0;
 
 	/**
-	 * Returns the TargetDomainDependencies that were previously commited for the given PackageName.
+	 * Returns an Attachment that was previously commited for the given PackageName.
 	 * Returns an empty object if not found.
 	 */
-	virtual FCbObject GetTargetDomainDependencies(FName PackageName) = 0;
+	virtual FCbObject GetOplogAttachment(FName PackageName, FUtf8StringView AttachmentKey) = 0;
 
 	/**
 	 * Remove the given cooked package(s) from storage; they have been modified since the last cook.
