ZenStoreWriter: Make oplog keys case insensitive, to remove the possiblity of two oplog entries for the same packagename.
#rnx

diff --git a/Engine/Source/Developer/IoStoreUtilities/Private/ZenStoreWriter.cpp b/Engine/Source/Developer/IoStoreUtilities/Private/ZenStoreWriter.cpp
--- a/Engine/Source/Developer/IoStoreUtilities/Private/ZenStoreWriter.cpp
+++ b/Engine/Source/Developer/IoStoreUtilities/Private/ZenStoreWriter.cpp
@@ -706,7 +706,9 @@ void FZenStoreWriter::CommitPackage(const FCommitPackageInfo& Info)
 
 			FCbWriter PackageObj;
 			PackageObj.BeginObject();
-			PackageObj << "key" << Info.PackageName.ToString();
+			FString PackageNameKey = Info.PackageName.ToString();
+			PackageNameKey.ToLowerInline();
+			PackageObj << "key" << PackageNameKey;
 
 			// NOTE: The package GUID and disk size are used for legacy iterative cooks when comparing asset registry package data
 			PackageObj.BeginObject("package");
