IoStore: Fixed issue that caused plugin remap redirects to be treated as localization redirects
#rnx

diff --git a/Engine/Source/Developer/IoStoreUtilities/Private/PackageStoreOptimizer.cpp b/Engine/Source/Developer/IoStoreUtilities/Private/PackageStoreOptimizer.cpp
--- a/Engine/Source/Developer/IoStoreUtilities/Private/PackageStoreOptimizer.cpp
+++ b/Engine/Source/Developer/IoStoreUtilities/Private/PackageStoreOptimizer.cpp
@@ -1737,7 +1737,7 @@ FContainerHeader FPackageStoreOptimizer::CreateContainerHeader(const FIoContaine
 		{
 			RedirectsNameMapBuilder.MarkNameAsReferenced(Entry->GetSourcePackageName());
 			FMappedName MappedSourcePackageName = RedirectsNameMapBuilder.MapName(Entry->GetSourcePackageName());
-			if (Entry->Region.IsValid())
+			if (!Entry->Region.IsNone())
 			{
 				Header.CulturePackageMap.FindOrAdd(Entry->Region.ToString()).Add({ Entry->GetSourcePackageId(), Entry->GetPackageId(), MappedSourcePackageName });
 			}
