IoStore: Changes to allow IoStore to run as part of UnrealPak instead of as a cmdlet
#rnx

diff --git a/Engine/Source/Developer/IoStoreUtilities/Internal/PackageStoreOptimizer.h b/Engine/Source/Developer/IoStoreUtilities/Internal/PackageStoreOptimizer.h
new file mode 100644
--- /dev/null
+++ b/Engine/Source/Developer/IoStoreUtilities/Internal/PackageStoreOptimizer.h
@@ -0,0 +1,402 @@
+// Copyright Epic Games, Inc. All Rights Reserved.
+
+#pragma once
+
+#include "Serialization/AsyncLoading2.h"
+#include "IO/PackageStore.h"
+#include "IO/IoDispatcher.h"
+#include "UObject/PackageFileSummary.h"
+#include "UObject/UObjectMarks.h"
+
+class FBufferWriter;
+
+class FPackageStoreNameMapBuilder
+{
+public:
+	void SetNameMapType(FMappedName::EType InNameMapType)
+	{
+		NameMapType = InNameMapType;
+	}
+
+	void AddName(const FName& Name)
+	{
+		const FNameEntryId ComparisonIndex = Name.GetComparisonIndex();
+		const FNameEntryId DisplayIndex = Name.GetDisplayIndex();
+		NameMap.Add(DisplayIndex);
+		int32 Index = NameMap.Num();
+		NameIndices.Add(ComparisonIndex, Index);
+	}
+
+	void MarkNamesAsReferenced(const TArray<FName>& Names, TArray<int32>& OutNameIndices)
+	{
+		for (const FName& Name : Names)
+		{
+			const FNameEntryId ComparisonIndex = Name.GetComparisonIndex();
+			const FNameEntryId DisplayIndex = Name.GetDisplayIndex();
+			int32& Index = NameIndices.FindOrAdd(ComparisonIndex);
+			if (Index == 0)
+			{
+				NameMap.Add(DisplayIndex);
+				Index = NameMap.Num();
+			}
+
+			OutNameIndices.Add(Index - 1);
+		}
+	}
+
+	void MarkNameAsReferenced(const FName& Name)
+	{
+		const FNameEntryId ComparisonIndex = Name.GetComparisonIndex();
+		const FNameEntryId DisplayIndex = Name.GetDisplayIndex();
+		int32& Index = NameIndices.FindOrAdd(ComparisonIndex);
+		if (Index == 0)
+		{
+			NameMap.Add(DisplayIndex);
+			Index = NameMap.Num();
+		}
+	}
+
+	FMappedName MapName(const FName& Name) const
+	{
+		const FNameEntryId Id = Name.GetComparisonIndex();
+		const int32* Index = NameIndices.Find(Id);
+		check(Index);
+		return FMappedName::Create(*Index - 1, Name.GetNumber(), NameMapType);
+	}
+
+	const TArray<FNameEntryId>& GetNameMap() const
+	{
+		return NameMap;
+	}
+
+	void Empty()
+	{
+		NameIndices.Empty();
+		NameMap.Empty();
+	}
+
+private:
+	TMap<FNameEntryId, int32> NameIndices;
+	TArray<FNameEntryId> NameMap;
+	FMappedName::EType NameMapType = FMappedName::EType::Package;
+};
+
+class FPackageStorePackage
+{
+public:
+	FPackageId GetId() const
+	{
+		return Id;
+	}
+
+	uint32 GetLoadOrder() const
+	{
+		return LoadOrder;
+	}
+
+	uint64 GetHeaderSize() const
+	{
+		return HeaderSize;
+	}
+
+	uint64 GetImportMapSize() const
+	{
+		return ImportMapSize;
+	}
+
+	uint64 GetExportMapSize() const
+	{
+		return ExportMapSize;
+	}
+
+	uint64 GetExportBundleCount() const
+	{
+		return GraphData.ExportBundles.Num();
+	}
+
+	uint64 GetExportBundleEntriesSize() const
+	{
+		return ExportBundleEntriesSize;
+	}
+	
+	uint64 GetGraphDataSize() const
+	{
+		return GraphDataSize;
+	}
+
+	uint64 GetNameCount() const
+	{
+		return NameMapBuilder.GetNameMap().Num();
+	}
+
+	uint64 GetNameMapSize() const
+	{
+		return NameMapSize;
+	}
+
+	const TArray<FPackageId>& GetImportedPackageIds() const
+	{
+		return ImportedPackageIds;
+	}
+
+	void RedirectFrom(FName SourcePackageName)
+	{
+		SourceName = SourcePackageName;
+	}
+
+	void AddShaderMapHash(const FSHAHash& ShaderMapHash)
+	{
+		ShaderMapHashes.Add(ShaderMapHash);
+	}
+
+	const TSet<FSHAHash>& GetShaderMapHashes() const
+	{
+		return ShaderMapHashes;
+	}
+
+private:
+	struct FExportGraphNode;
+
+	struct FExternalDependency
+	{
+		int32 ImportIndex = -1;
+		FExportBundleEntry::EExportCommandType ExportBundleCommandType;
+	};
+
+	struct FExportGraphNode
+	{
+		FExportBundleEntry BundleEntry;
+		TArray<FExportGraphNode*> InternalDependencies;
+		TArray<FExternalDependency> ExternalDependencies;
+		int32 ExportBundleIndex = -1;
+		int32 IncomingEdgeCount = 0;
+		bool bIsPublic = false;
+	};
+
+	struct FExport
+	{
+		FString FullName;
+		FName ObjectName;
+		uint32 ExportHash = 0;
+		FPackageObjectIndex OuterIndex;
+		FPackageObjectIndex ClassIndex;
+		FPackageObjectIndex SuperIndex;
+		FPackageObjectIndex TemplateIndex;
+		EObjectFlags ObjectFlags = RF_NoFlags;
+		uint64 CookedSerialOffset = uint64(-1);
+		uint64 SerialOffset = uint64(-1);
+		uint64 SerialSize = uint64(-1);
+		bool bNotForClient = false;
+		bool bNotForServer = false;
+		bool bIsPublic = false;
+		FExportGraphNode* Nodes[FExportBundleEntry::ExportCommandType_Count] = { nullptr };
+	};
+
+	struct FExportBundle
+	{
+		uint64 SerialOffset = uint64(-1);
+		TArray<FExportBundleEntry> Entries;
+	};
+
+	struct FUnresolvedImport
+	{
+		FString FullName;
+		FName FromPackageName;
+		int32 FromPackageNameLen = 0;
+		FPackageId FromPackageId;
+		bool bIsScriptImport = false;
+		bool bIsImportOfPackage = false;
+	};
+
+	struct FInternalArc
+	{
+		int32 FromExportBundleIndex;
+		int32 ToExportBundleIndex;
+
+		bool operator==(const FInternalArc& Other) const
+		{
+			return FromExportBundleIndex == Other.FromExportBundleIndex &&
+				ToExportBundleIndex == Other.ToExportBundleIndex;
+		}
+	};
+
+	struct FExternalArc
+	{
+		int32 FromImportIndex;
+		FExportBundleEntry::EExportCommandType FromCommandType;
+		int32 ToExportBundleIndex;
+
+		bool operator==(const FExternalArc& Other) const
+		{
+			return FromImportIndex == Other.FromImportIndex &&
+				FromCommandType == Other.FromCommandType &&
+				ToExportBundleIndex == Other.ToExportBundleIndex;
+		}
+	};
+
+	struct FGraphData
+	{
+		TArray<FExportBundle> ExportBundles;
+		TArray<FInternalArc> InternalArcs;
+		TMap<FPackageId, TArray<FExternalArc>> ExternalArcs;
+	};
+
+	struct FExportBundleGraphNode
+	{
+		FPackageStorePackage* Package = nullptr;
+		TArray<FExportGraphNode*> ExportGraphNodes;
+		int32 Index = -1;
+		int32 IncomingEdgeCount = 0;
+	};
+
+	FPackageId Id;
+	FName Name;
+	FName SourceName;
+	FString Region;
+
+	FPackageStoreNameMapBuilder NameMapBuilder;
+	TArray<FPackageObjectIndex> Imports;
+	TArray<FExport> Exports;
+	TArray<FExportGraphNode> ExportGraphNodes;
+	FGraphData GraphData;
+
+	TArray<FPackageId> ImportedPackageIds;
+	TSet<FSHAHash> ShaderMapHashes;
+
+	FIoBuffer HeaderBuffer;
+
+	uint32 PackageFlags = 0;
+	uint32 CookedHeaderSize = 0;
+	uint64 HeaderSize = 0;
+	uint64 ExportsSerialSize = 0;
+	uint64 ImportMapSize = 0;
+	uint64 ExportMapSize = 0;
+	uint64 ExportBundleEntriesSize = 0;
+	uint64 GraphDataSize = 0;
+	uint64 NameMapSize = 0;
+	uint32 LoadOrder = 0;
+	
+	TArray<FExportBundleGraphNode> ExportBundleGraphNodes;
+	FPackageStorePackage::FExportBundle* CurrentBundle = nullptr;
+	TArray<FExportBundleGraphNode*> NodesWithNoIncomingEdges;
+	bool bTemporaryMark = false;
+	bool bPermanentMark = false;
+
+	bool bIsRedirected = false;
+
+	friend class FPackageStoreOptimizer;
+};
+
+class FPackageStoreOptimizer
+{
+public:
+	uint64 GetTotalPackageCount() const
+	{
+		return TotalPackageCount;
+	}
+
+	uint64 GetTotalExportBundleCount() const
+	{
+		return TotalExportBundleCount;
+	}
+
+	uint64 GetTotalExportBundleEntryCount() const
+	{
+		return TotalExportBundleEntryCount;
+	}
+
+	uint64 GetTotalInternalBundleArcsCount() const
+	{
+		return TotalInternalBundleArcsCount;
+	}
+
+	uint64 GetTotalExternalBundleArcsCount() const
+	{
+		return TotalExternalBundleArcsCount;
+	}
+
+	uint64 GetTotalScriptObjectCount() const
+	{
+		return TotalScriptObjectCount;
+	}
+
+	IOSTOREUTILITIES_API void Initialize(const ITargetPlatform* TargetPlatform);
+	void Initialize(const FIoBuffer& ScriptObjectsBuffer);
+
+	FPackageStorePackage* CreateMissingPackage(const FName& Name) const;
+	FPackageStorePackage* CreatePackageFromCookedHeader(const FName& Name, const FIoBuffer& CookedHeaderBuffer) const;
+	FPackageStorePackage* CreatePackageFromPackageStoreHeader(const FName& Name, const FIoBuffer& Buffer, const FPackageStoreEntryResource& PackageStoreEntry) const;
+	void FinalizePackage(FPackageStorePackage* Package);
+	FIoBuffer CreatePackageBuffer(const FPackageStorePackage* Package, const FIoBuffer& CookedExportsBuffer, TArray<FFileRegion>* InOutFileRegions) const;
+	FPackageStoreEntryResource CreatePackageStoreEntry(const FPackageStorePackage* Package) const;
+	FContainerHeader CreateContainerHeader(const FIoContainerId& ContainerId, TArrayView<const FPackageStoreEntryResource> PackageStoreEntries) const;
+	IOSTOREUTILITIES_API FIoBuffer CreateScriptObjectsBuffer() const;
+	void LoadScriptObjectsBuffer(const FIoBuffer& ScriptObjectsBuffer);
+	void ProcessRedirects(const TMap<FPackageId, FPackageStorePackage*>& PackagesMap, bool bIsBuildingDLC) const;
+	void OptimizeExportBundles(const TMap<FPackageId, FPackageStorePackage*>& PackagesMap);
+
+private:
+	struct FScriptObjectData
+	{
+		FName ObjectName;
+		FString FullName;
+		FPackageObjectIndex GlobalIndex;
+		FPackageObjectIndex OuterIndex;
+		FPackageObjectIndex CDOClassIndex;
+	};
+
+	struct FCookedHeaderData
+	{
+		FPackageFileSummary Summary;
+		TArray<FName> SummaryNames;
+		TArray<FObjectImport> ObjectImports;
+		TArray<FObjectExport> ObjectExports;
+		TArray<FPackageIndex> PreloadDependencies;
+	};
+
+	struct FPackageStoreHeaderData
+	{
+		FPackageSummary Summary;
+		TArray<FPackageId> ImportedPackageIds;
+		TArray<FNameEntryId> NameMap;
+		TArray<FPackageObjectIndex> Imports;
+		TArray<FExportMapEntry> Exports;
+		TArray<FExportBundleHeader> ExportBundleHeaders;
+		TArray<FExportBundleEntry> ExportBundleEntries;
+		TArray<FPackageStorePackage::FInternalArc> InternalArcs;
+		TArray<FPackageStorePackage::FExternalArc> ExternalArcs;
+	};
+
+	using FExportGraphEdges = TMultiMap<FPackageStorePackage::FExportGraphNode*, FPackageStorePackage::FExportGraphNode*>;
+	using FExportBundleGraphEdges = TMultiMap<FPackageStorePackage::FExportBundleGraphNode*, FPackageStorePackage::FExportBundleGraphNode*>;
+
+	FCookedHeaderData LoadCookedHeader(const FIoBuffer& CookedHeaderBuffer) const;
+	FPackageStoreHeaderData LoadPackageStoreHeader(const FIoBuffer& PackageStoreHeaderBuffer, const FPackageStoreEntryResource& PackageStoreEntry) const;
+	void ResolveImport(FPackageStorePackage::FUnresolvedImport* Imports, const FObjectImport* ObjectImports, int32 LocalImportIndex) const;
+	void ProcessImports(const FCookedHeaderData& CookedHeaderData, FPackageStorePackage* Package) const;
+	void ProcessImports(const FPackageStoreHeaderData& PackageStoreHeaderData, FPackageStorePackage* Package) const;
+	void ResolveExport(FPackageStorePackage::FExport* Exports, const FObjectExport* ObjectExports, const int32 LocalExportIndex, const FName& PackageName) const;
+	void ResolveExport(FPackageStorePackage::FExport* Exports, const int32 LocalExportIndex, const FName& PackageName) const;
+	void ProcessExports(const FCookedHeaderData& CookedHeaderData, FPackageStorePackage* Package) const;
+	void ProcessExports(const FPackageStoreHeaderData& PackageStoreHeaderData, FPackageStorePackage* Package) const;
+	void ProcessPreloadDependencies(const FCookedHeaderData& CookedHeaderData, FPackageStorePackage* Package) const;
+	void ProcessPreloadDependencies(const FPackageStoreHeaderData& PackageStoreHeaderData, FPackageStorePackage* Package) const;
+	TArray<FPackageStorePackage*> SortPackagesInLoadOrder(const TMap<FPackageId, FPackageStorePackage*>& PackagesMap) const;
+	void SerializeGraphData(const TArray<FPackageId>& ImportedPackageIds, FPackageStorePackage::FGraphData& GraphData, FBufferWriter& GraphArchive) const;
+	TArray<FPackageStorePackage::FExportGraphNode*> SortExportGraphNodesInLoadOrder(FPackageStorePackage* Package, FExportGraphEdges& Edges) const;
+	TArray<FPackageStorePackage::FExportBundleGraphNode*> SortExportBundleGraphNodesInLoadOrder(const TArray<FPackageStorePackage*>& Packages, FExportBundleGraphEdges& Edges) const;
+	void CreateExportBundles(FPackageStorePackage* Package) const;
+	bool VerifyRedirect(const FPackageStorePackage* SourcePackage, FPackageStorePackage& TargetPackage, bool bIsBuildingDLC) const;
+	void FinalizePackageHeader(FPackageStorePackage* Package) const;
+	void FindScriptObjectsRecursive(FPackageObjectIndex OuterIndex, UObject* Object, EObjectMark ExcludedObjectMarks, const ITargetPlatform* TargetPlatform);
+	void FindScriptObjects(const ITargetPlatform* TargetPlatform);
+
+	TMap<FPackageObjectIndex, FScriptObjectData> ScriptObjectsMap;
+	uint64 TotalPackageCount = 0;
+	uint64 TotalExportBundleCount = 0;
+	uint64 TotalExportBundleEntryCount = 0;
+	uint64 TotalInternalBundleArcsCount = 0;
+	uint64 TotalExternalBundleArcsCount = 0;
+	uint64 TotalScriptObjectCount = 0;
+	uint32 NextLoadOrder = 0;
+};
diff --git a/Engine/Source/Developer/IoStoreUtilities/IoStoreUtilities.Build.cs b/Engine/Source/Developer/IoStoreUtilities/IoStoreUtilities.Build.cs
--- a/Engine/Source/Developer/IoStoreUtilities/IoStoreUtilities.Build.cs
+++ b/Engine/Source/Developer/IoStoreUtilities/IoStoreUtilities.Build.cs
@@ -13,7 +13,6 @@ public class IoStoreUtilities : ModuleRules
 		PrivateDependencyModuleNames.AddRange(new string[] {
             "Core",
             "CoreUObject",
-            "Projects",
             "AssetRegistry",
 			"DerivedDataCache",
 			"RenderCore",
diff --git a/Engine/Source/Developer/IoStoreUtilities/Private/IoStoreUtilities.cpp b/Engine/Source/Developer/IoStoreUtilities/Private/IoStoreUtilities.cpp
--- a/Engine/Source/Developer/IoStoreUtilities/Private/IoStoreUtilities.cpp
+++ b/Engine/Source/Developer/IoStoreUtilities/Private/IoStoreUtilities.cpp
@@ -6,7 +6,6 @@
 #include "HAL/FileManager.h"
 #include "HAL/PlatformFileManager.h"
 #include "Hash/CityHash.h"
-#include "Interfaces/IPluginManager.h"
 #include "Interfaces/ITargetPlatformManagerModule.h"
 #include "Interfaces/ITargetPlatform.h"
 #include "IO/IoDispatcher.h"
@@ -54,7 +53,7 @@
 #include "Algo/MaxElement.h"
 #include "Algo/StableSort.h"
 #include "PackageStoreOptimizer.h"
-#include "ShaderCodeLibrary.h"
+#include "ShaderCodeArchive.h"
 #include "ZenStoreHttpClient.h"
 #include "IPlatformFilePak.h"
 #include "ZenStoreWriter.h"
@@ -346,7 +345,6 @@ struct FContainerTargetFile
 	FLegacyCookedPackage* Package = nullptr;
 	FString NormalizedSourcePath;
 	TOptional<FIoBuffer> SourceBuffer;
-	FString TargetPath;
 	FString DestinationPath;
 	uint64 SourceSize = 0;
 	uint64 IdealOrder = 0;
@@ -361,13 +359,14 @@ struct FContainerTargetFile
 class FCookedPackageStore
 {
 public:
-	FCookedPackageStore()
+	FCookedPackageStore(const FString& CookedDir)
+		: PackageStoreManifest(CookedDir)
 	{
 	}
 
-	FIoStatus Load(const FString& CookedDirectory)
+	FIoStatus Load(const TCHAR* ManifestFilename)
 	{
-		FIoStatus Status = PackageStoreManifest.Load(*(CookedDirectory / FApp::GetProjectName() / TEXT("MetaData") / TEXT("packagestore.manifest")));
+		FIoStatus Status = PackageStoreManifest.Load(ManifestFilename);
 		if (!Status.IsOk())
 		{
 			return Status;
@@ -377,57 +376,65 @@ public:
 		{
 			DataSource = MakeUnique<FZenStoreDataSource>(*ZenServerInfo);
 		}
-		else
-		{
-			return FIoStatusBuilder(EIoErrorCode::NotFound) << TEXT("Missing ZenServerInfo in package store manifest");
-		}
 		
-		FIoContainerId ContainerId = FIoContainerId::FromName(TEXT("global"));
-
-		TIoStatusOr<FIoBuffer> HeaderBuffer = DataSource->ReadChunk(CreateIoChunkId(ContainerId.Value(), 0, EIoChunkType::ContainerHeader), FIoReadOptions());
-		if (!HeaderBuffer.IsOk())
+		if (DataSource)
 		{
-			return FIoStatusBuilder(EIoErrorCode::ReadError) << TEXT("Failed to read container header");
-		}
+			FIoContainerId ContainerId = FIoContainerId::FromName(TEXT("global"));
 
-		FMemoryReaderView Ar(MakeArrayView(HeaderBuffer.ValueOrDie().Data(), HeaderBuffer.ValueOrDie().DataSize()));
-		FContainerHeader ContainerHeader;
-		Ar << ContainerHeader;
+			TIoStatusOr<FIoBuffer> HeaderBuffer = DataSource->ReadChunk(CreateIoChunkId(ContainerId.Value(), 0, EIoChunkType::ContainerHeader), FIoReadOptions());
+			if (!HeaderBuffer.IsOk())
+			{
+				return FIoStatusBuilder(EIoErrorCode::ReadError) << TEXT("Failed to read container header");
+			}
 
-		PackageEntries.Reserve(ContainerHeader.PackageCount);
-		TArrayView<const FFilePackageStoreEntry> FilePackageEntries = MakeArrayView<const FFilePackageStoreEntry>(reinterpret_cast<const FFilePackageStoreEntry*>(ContainerHeader.StoreEntries.GetData()), ContainerHeader.PackageCount);
+			FMemoryReaderView Ar(MakeArrayView(HeaderBuffer.ValueOrDie().Data(), HeaderBuffer.ValueOrDie().DataSize()));
+			FContainerHeader ContainerHeader;
+			Ar << ContainerHeader;
 
-		int32 Idx = 0;
-		for (const FFilePackageStoreEntry& FilePackageEntry : FilePackageEntries)
-		{
-			FPackageStoreEntryResource& Entry = PackageEntries.AddDefaulted_GetRef();
-			
-			Entry.ExportInfo = FPackageStoreExportInfo
+			PackageEntries.Reserve(ContainerHeader.PackageCount);
+			TArrayView<const FFilePackageStoreEntry> FilePackageEntries = MakeArrayView<const FFilePackageStoreEntry>(reinterpret_cast<const FFilePackageStoreEntry*>(ContainerHeader.StoreEntries.GetData()), ContainerHeader.PackageCount);
+
+			int32 Idx = 0;
+			for (const FFilePackageStoreEntry& FilePackageEntry : FilePackageEntries)
 			{
-				FilePackageEntry.ExportCount,
-				FilePackageEntry.ExportBundleCount
-			};
-			Entry.ImportedPackageIds = MakeArrayView<const FPackageId>(FilePackageEntry.ImportedPackages.Data(), FilePackageEntry.ImportedPackages.Num());
+				FPackageStoreEntryResource& Entry = PackageEntries.AddDefaulted_GetRef();
 
-			PackageIdToEntry.Add(ContainerHeader.PackageIds[Idx++], &Entry);
-		}
+				Entry.ExportInfo = FPackageStoreExportInfo
+				{
+					FilePackageEntry.ExportCount,
+					FilePackageEntry.ExportBundleCount
+				};
+				Entry.ImportedPackageIds = MakeArrayView<const FPackageId>(FilePackageEntry.ImportedPackages.Data(), FilePackageEntry.ImportedPackages.Num());
 
-		for (const FPackageStoreManifest::FFileInfo& FileInfo : PackageStoreManifest.GetFiles())
-		{
-			FilenameToChunkIdMap.Add(FileInfo.FileName, FileInfo.ChunkId);
+				PackageIdToEntry.Add(ContainerHeader.PackageIds[Idx++], &Entry);
+			}
 		}
 
 		for (const FPackageStoreManifest::FPackageInfo& PackageInfo : PackageStoreManifest.GetPackages())
 		{
 			FName PackageName = FName(PackageInfo.PackageName);
 			PackageNameToPackageInfoMap.Add(PackageName, &PackageInfo);
-			ChunkIdToPackageNameMap.Add(PackageInfo.PackageChunkId, PackageName);
+			if (PackageInfo.PackageChunkId.IsValid())
+			{
+				ChunkIdToPackageNameMap.Add(PackageInfo.PackageChunkId, PackageName);
+			}
 			for (const FIoChunkId& ChunkId : PackageInfo.BulkDataChunkIds)
 			{
 				ChunkIdToPackageNameMap.Add(ChunkId, PackageName);
 			}
 		}
 
+		for (const FPackageStoreManifest::FFileInfo& FileInfo : PackageStoreManifest.GetFiles())
+		{
+			FilenameToChunkIdMap.Add(FileInfo.FileName, FileInfo.ChunkId);
+			FName* FindPackageName = ChunkIdToPackageNameMap.Find(FileInfo.ChunkId);
+			if (FindPackageName)
+			{
+				FString FilenameWithoutExtension = FPaths::ChangeExtension(FileInfo.FileName, FString());
+				FilenameToPackageNameMap.Add(MoveTemp(FilenameWithoutExtension), *FindPackageName);
+			}
+		}
+
 		return FIoStatus::Ok;
 	}
 
@@ -441,6 +448,12 @@ public:
 		return ChunkIdToPackageNameMap.FindRef(ChunkId);
 	}
 
+	FName GetPackageNameFromFileName(const FString& Filename) const
+	{
+		FString FilenameWithoutExtension = FPaths::ChangeExtension(Filename, FString());
+		return FilenameToPackageNameMap.FindRef(FilenameWithoutExtension);
+	}
+
 	const FPackageStoreManifest::FPackageInfo* GetPackageInfoFromPackageName(FName PackageName)
 	{
 		return PackageNameToPackageInfoMap.FindRef(PackageName);
@@ -451,11 +464,21 @@ public:
 		return PackageIdToEntry.FindRef(PackageId);
 	}
 
+	bool HasDataSource() const
+	{
+		return DataSource.IsValid();
+	}
+
 	TIoStatusOr<uint64> GetChunkSize(const FIoChunkId& ChunkId)
 	{
 		return DataSource->GetChunkSize(ChunkId);
 	}
 
+	TIoStatusOr<FIoBuffer> ReadChunk(const FIoChunkId& ChunkId)
+	{
+		return DataSource->ReadChunk(ChunkId, FIoReadOptions());
+	}
+
 	void ReadChunkAsync(const FIoChunkId& ChunkId, TFunction<void(TIoStatusOr<FIoBuffer>)>&& Callback)
 	{
 		class FReadChunkTask
@@ -560,6 +583,7 @@ private:
 	TMap<FPackageId, const FPackageStoreEntryResource*> PackageIdToEntry;
 	TMap<FString, FIoChunkId> FilenameToChunkIdMap;
 	TMap<FIoChunkId, FName> ChunkIdToPackageNameMap;
+	TMap<FString, FName> FilenameToPackageNameMap;
 	TMap<FName, const FPackageStoreManifest::FPackageInfo*> PackageNameToPackageInfoMap;
 };
 
@@ -581,9 +605,6 @@ struct FIoStoreArguments
 {
 	FString GlobalContainerPath;
 	FString CookedDir;
-	ITargetPlatform* TargetPlatform = nullptr;
-	FString MetaInputDir;
-	FString MetaOutputDir;
 	TArray<FContainerSourceSpec> Containers;
 	FCookedFileStatMap CookedFileStatMap;
 	TArray<FFileOrderMap> OrderMaps;
@@ -595,17 +616,13 @@ struct FIoStoreArguments
 	FAssetRegistryState ReleaseAssetRegistry;
 	FReleasedPackages ReleasedPackages;
 	TUniquePtr<FCookedPackageStore> PackageStore;
+	TUniquePtr<FIoBuffer> ScriptObjects;
 	bool bSign = false;
 	bool bRemapPluginContentToGame = false;
 	bool bCreateDirectoryIndex = true;
 	bool bClusterByOrderFilePriority = false;
 	bool bFileRegions = false;
 
-	bool ShouldCreateContainers() const
-	{
-		return GlobalContainerPath.Len() > 0 || DLCPluginPath.Len() > 0;
-	}
-
 	bool IsDLC() const
 	{
 		return DLCPluginPath.Len() > 0;
@@ -1029,7 +1046,9 @@ static void CreateDiskLayout(
 				return A->ChunkType < B->ChunkType;
 			}
 			if (A->ChunkType == EContainerChunkType::ShaderCodeLibrary)
-				return A->TargetPath < B->TargetPath;
+			{
+				return A->DestinationPath < B->DestinationPath;
+			}
 			if (A->Package != B->Package)
 			{
 				return A->Package->DiskLayoutOrder < B->Package->DiskLayoutOrder;
@@ -1187,21 +1206,19 @@ FLegacyCookedPackage* FindOrAddPackage(
 
 FLegacyCookedPackage* FindOrAddPackage(
 	const FIoStoreArguments& Arguments,
-	const TCHAR* RelativeFileName,
+	const TCHAR* FileName,
 	TArray<FLegacyCookedPackage*>& Packages,
 	FPackageNameMap& PackageNameMap,
 	FPackageIdMap& PackageIdMap)
 {
-	FString PackageName;
-	FString ErrorMessage;
-	if (!FPackageName::TryConvertFilenameToLongPackageName(RelativeFileName, PackageName, &ErrorMessage))
+	FName PackageName = Arguments.PackageStore->GetPackageNameFromFileName(FileName);
+	if (PackageName.IsNone())
 	{
-		UE_LOG(LogIoStore, Warning, TEXT("Failed to obtain package name from file name '%s'"), *ErrorMessage);
+		UE_LOG(LogIoStore, Fatal, TEXT("Failed to obtain package name from file name '%s'"), FileName);
 		return nullptr;
 	}
 
-	FName PackageFName = *PackageName;
-	return FindOrAddPackage(Arguments, PackageFName, Packages, PackageNameMap, PackageIdMap);
+	return FindOrAddPackage(Arguments, PackageName, Packages, PackageNameMap, PackageIdMap);
 }
 
 static void ParsePackageAssetsFromFiles(TArray<FLegacyCookedPackage*>& Packages, const FPackageStoreOptimizer& PackageStoreOptimizer)
@@ -1349,6 +1366,139 @@ TArray<TUniquePtr<FIoStoreReader>> CreatePatchSourceReaders(const TArray<FString
 	return Readers;
 }
 
+bool LoadShaderAssetInfo(const FString& Filename, TMap<FSHAHash, TSet<FName>>& OutShaderCodeToAssets)
+{
+	FString JsonText; 
+	if (!FFileHelper::LoadFileToString(JsonText, *Filename))
+	{
+		return false;
+	}
+
+	TSharedPtr<FJsonObject> JsonObject;
+	TSharedRef<TJsonReader<TCHAR>> Reader = TJsonReaderFactory<TCHAR>::Create(JsonText);
+
+	if (!FJsonSerializer::Deserialize(Reader, JsonObject) || !JsonObject.IsValid())
+	{
+		return false;
+	}
+
+	TSharedPtr<FJsonValue> AssetInfoArrayValue = JsonObject->Values.FindRef(TEXT("ShaderCodeToAssets"));
+	if (!AssetInfoArrayValue.IsValid())
+	{
+		return false;
+	}
+
+	TArray<TSharedPtr<FJsonValue>> AssetInfoArray = AssetInfoArrayValue->AsArray();
+	for (int32 IdxPair = 0, NumPairs = AssetInfoArray.Num(); IdxPair < NumPairs; ++IdxPair)
+	{
+		TSharedPtr<FJsonObject> Pair = AssetInfoArray[IdxPair]->AsObject();
+		if (Pair.IsValid())
+		{
+			TSharedPtr<FJsonValue> ShaderMapHashJson = Pair->Values.FindRef(TEXT("ShaderMapHash"));
+			if (!ShaderMapHashJson.IsValid())
+			{
+				continue;
+			}
+			FSHAHash ShaderMapHash;
+			ShaderMapHash.FromString(ShaderMapHashJson->AsString());
+
+			TSharedPtr<FJsonValue> AssetPathsArrayValue = Pair->Values.FindRef(TEXT("Assets"));
+			if (!AssetPathsArrayValue.IsValid())
+			{
+				continue;
+			}
+
+			TSet<FName>& Assets = OutShaderCodeToAssets.Add(ShaderMapHash);
+			TArray<TSharedPtr<FJsonValue>> AssetPathsArray = AssetPathsArrayValue->AsArray();
+			for (int32 IdxAsset = 0, NumAssets = AssetPathsArray.Num(); IdxAsset < NumAssets; ++IdxAsset)
+			{
+				Assets.Add(FName(*AssetPathsArray[IdxAsset]->AsString()));
+			}
+		}
+	}
+	return true;
+}
+
+bool ConvertToIoStoreShaderLibrary(
+	const TCHAR* FileName,
+	TTuple<FIoChunkId, FIoBuffer>& OutLibraryIoChunk,
+	TArray<TTuple<FIoChunkId, FIoBuffer, FSHAHash>>& OutCodeIoChunks,
+	TArray<TTuple<FSHAHash, TArray<FIoChunkId>>>& OutShaderMaps,
+	TArray<TTuple<FSHAHash, TSet<FName>>>& OutShaderMapAssetAssociations)
+{
+	TArray<FString> Components;
+	if (FPaths::GetBaseFilename(FileName).ParseIntoArray(Components, TEXT("-")) != 3)
+	{
+		UE_LOG(LogIoStore, Error, TEXT("Invalid shader code library file name '%s'."), FileName);
+		return false;
+	}
+	FString LibraryName = Components[1];
+	FName FormatName(Components[2]);
+
+	TUniquePtr<FArchive> LibraryAr(IFileManager::Get().CreateFileReader(FileName));
+	if (!LibraryAr)
+	{
+		UE_LOG(LogIoStore, Error, TEXT("Missing shader code library file '%s'."), FileName);
+		return false;
+	}
+
+	uint32 Version;
+	(*LibraryAr) << Version;
+
+	FSerializedShaderArchive SerializedShaders;
+	(*LibraryAr) << SerializedShaders;
+	int64 OffsetToShaderCode = LibraryAr->Tell();
+
+	FString AssetInfoFileName = FPaths::GetPath(FileName) / FString::Printf(TEXT("ShaderAssetInfo-%s-%s.assetinfo.json"), *LibraryName, *FormatName.ToString());
+	TMap<FSHAHash, FShaderMapAssetPaths> ShaderCodeToAssets;
+	if (!LoadShaderAssetInfo(AssetInfoFileName, ShaderCodeToAssets))
+	{
+		UE_LOG(LogIoStore, Error, TEXT("Failed loading asset asset info file '%s'"), *AssetInfoFileName);
+		return false;
+	}
+
+	FLargeMemoryWriter IoStoreLibraryAr(0, true);
+	FIoStoreShaderCodeArchive::SaveIoStoreShaderCodeArchive(SerializedShaders, IoStoreLibraryAr);
+	OutLibraryIoChunk.Key = FIoStoreShaderCodeArchive::GetShaderCodeArchiveChunkId(LibraryName, FormatName);
+	int64 TotalSize = IoStoreLibraryAr.TotalSize();
+	OutLibraryIoChunk.Value = FIoBuffer(FIoBuffer::AssumeOwnership, IoStoreLibraryAr.ReleaseOwnership(), TotalSize);
+
+	int32 ShaderCount = SerializedShaders.ShaderEntries.Num();
+	OutCodeIoChunks.Reserve(ShaderCount);
+	for (int32 ShaderIndex = 0; ShaderIndex < ShaderCount; ++ShaderIndex)
+	{
+		const FSHAHash& ShaderHash = SerializedShaders.ShaderHashes[ShaderIndex];
+		const FShaderCodeEntry& ShaderEntry = SerializedShaders.ShaderEntries[ShaderIndex];
+		TTuple<FIoChunkId, FIoBuffer, FSHAHash>& OutCodeIoChunk = OutCodeIoChunks.Emplace_GetRef(FIoStoreShaderCodeArchive::GetShaderCodeChunkId(ShaderHash), ShaderEntry.Size, ShaderHash);
+		LibraryAr->Seek(OffsetToShaderCode + ShaderEntry.Offset);
+		LibraryAr->Serialize(OutCodeIoChunk.Get<1>().Data(), ShaderEntry.Size);
+	}
+
+	int32 ShaderMapCount = SerializedShaders.ShaderMapHashes.Num();
+	OutShaderMaps.Reserve(ShaderMapCount);
+	for (int32 ShaderMapIndex = 0; ShaderMapIndex < ShaderMapCount; ++ShaderMapIndex)
+	{
+		TTuple<FSHAHash, TArray<FIoChunkId>>& OutShaderMap = OutShaderMaps.AddDefaulted_GetRef();
+		OutShaderMap.Key = SerializedShaders.ShaderMapHashes[ShaderMapIndex];
+		const FShaderMapEntry& ShaderMapEntry = SerializedShaders.ShaderMapEntries[ShaderMapIndex];
+		int32 LookupIndexEnd = ShaderMapEntry.ShaderIndicesOffset + ShaderMapEntry.NumShaders;
+		OutShaderMap.Value.Reserve(ShaderMapEntry.NumShaders);
+		for (int32 ShaderLookupIndex = ShaderMapEntry.ShaderIndicesOffset; ShaderLookupIndex < LookupIndexEnd; ++ShaderLookupIndex)
+		{
+			int32 ShaderIndex = SerializedShaders.ShaderIndices[ShaderLookupIndex];
+			OutShaderMap.Value.Add(FIoStoreShaderCodeArchive::GetShaderCodeChunkId(SerializedShaders.ShaderHashes[ShaderIndex]));
+		}
+	}
+
+	OutShaderMapAssetAssociations.Reserve(ShaderCodeToAssets.Num());
+	for (auto& KV : ShaderCodeToAssets)
+	{
+		OutShaderMapAssetAssociations.Emplace(KV.Key, MoveTemp(KV.Value));
+	}
+
+	return true;
+}
+
 void ProcessShaderLibraries(const FIoStoreArguments& Arguments, TArray<FContainerTargetSpec*>& ContainerTargets)
 {
 	IOSTORE_CPU_SCOPE(ProcessShaderLibraries);
@@ -1370,7 +1520,7 @@ void ProcessShaderLibraries(const FIoStoreArguments& Arguments, TArray<FContaine
 				TTuple<FIoChunkId, FIoBuffer> LibraryChunk;
 				TArray<TTuple<FIoChunkId, FIoBuffer, FSHAHash>> CodeChunks;
 				TArray<TTuple<FSHAHash, TSet<FName>>> ShaderMapAssetAssociations;
-				if (!FShaderLibraryCooker::ConvertToIoStoreShaderLibrary(*TargetFile.NormalizedSourcePath, LibraryChunk, CodeChunks, ShaderMaps, ShaderMapAssetAssociations))
+				if (!ConvertToIoStoreShaderLibrary(*TargetFile.NormalizedSourcePath, LibraryChunk, CodeChunks, ShaderMaps, ShaderMapAssetAssociations))
 				{
 					UE_LOG(LogIoStore, Warning, TEXT("Failed converting shader library '%s'"), *TargetFile.NormalizedSourcePath);
 					continue;
@@ -1490,46 +1640,11 @@ void InitializeContainerTargetsAndPackages(
 	TArray<FContainerTargetSpec*>& ContainerTargets,
 	FPackageStoreOptimizer& PackageStoreOptimizer)
 {
-	FString ProjectName = FApp::GetProjectName();
-	FString RelativeEnginePath = FPaths::GetRelativePathToRoot();
-	FString RelativeProjectPath = FPaths::ProjectDir();
-	int32 CookedEngineDirLen = Arguments.CookedDir.Len() + 1;;
-	int32 CookedProjectDirLen = CookedEngineDirLen + ProjectName.Len() + 1;
-
-	auto ConvertCookedPathToRelativePath = [
-		&ProjectName,
-		&CookedEngineDirLen,
-		&CookedProjectDirLen,
-		&RelativeEnginePath,
-		&RelativeProjectPath]
-		(const FString& CookedFile) -> FString
-	{
-		FString RelativeFileName;
-		const TCHAR* FileName = *CookedFile + CookedEngineDirLen;
-		if (FCString::Strncmp(FileName, *ProjectName, ProjectName.Len()))
-		{
-			int32 FileNameLen = CookedFile.Len() - CookedEngineDirLen;
-			RelativeFileName.Reserve(RelativeEnginePath.Len() + FileNameLen);
-			RelativeFileName = RelativeEnginePath;
-			RelativeFileName.AppendChars(*CookedFile + CookedEngineDirLen, FileNameLen);
-		}
-		else
-		{
-			FileName = *CookedFile + CookedProjectDirLen;
-			int32 FileNameLen = CookedFile.Len() - CookedProjectDirLen;
-			RelativeFileName.Reserve(RelativeProjectPath.Len() + FileNameLen);
-			RelativeFileName = RelativeProjectPath;
-			RelativeFileName.AppendChars(*CookedFile + CookedProjectDirLen, FileNameLen);
-		}
-		return RelativeFileName;
-	};
-
 	auto CreateTargetFileFromCookedFile = [
 		&Arguments,
-		&ConvertCookedPathToRelativePath,
 		&Packages,
 		&PackageNameMap,
-		&PackageIdMap](const FContainerSourceFile& SourceFile, const FString& RelativeFileName, FContainerTargetFile& OutTargetFile) -> bool
+		&PackageIdMap](const FContainerSourceFile& SourceFile, FContainerTargetFile& OutTargetFile) -> bool
 	{
 		const FCookedFileStatData* OriginalCookedFileStatData = Arguments.CookedFileStatMap.Find(SourceFile.NormalizedPath);
 		if (!OriginalCookedFileStatData)
@@ -1562,12 +1677,12 @@ void InitializeContainerTargetsAndPackages(
 		else if (CookedFileStatData->FileExt == FCookedFileStatData::EFileExt::UMappedBulk)
 		{
 			OutTargetFile.ChunkType = EContainerChunkType::MemoryMappedBulkData;
-			FString TmpFileName = FString(RelativeFileName.Len() - 8, GetData(RelativeFileName)) + TEXT(".ubulk");
+			FString TmpFileName = FString(SourceFile.NormalizedPath.Len() - 8, GetData(SourceFile.NormalizedPath)) + TEXT(".ubulk");
 			OutTargetFile.Package = FindOrAddPackage(Arguments, *TmpFileName, Packages, PackageNameMap, PackageIdMap);
 		}
 		else
 		{
-			OutTargetFile.Package = FindOrAddPackage(Arguments, *RelativeFileName, Packages, PackageNameMap, PackageIdMap);
+			OutTargetFile.Package = FindOrAddPackage(Arguments, *SourceFile.NormalizedPath, Packages, PackageNameMap, PackageIdMap);
 			if (CookedFileStatData->FileExt == FCookedFileStatData::UPtnl)
 			{
 				OutTargetFile.ChunkType = EContainerChunkType::OptionalBulkData;
@@ -1621,10 +1736,9 @@ void InitializeContainerTargetsAndPackages(
 
 	auto CreateTargetFileFromPackageStore = [
 		&Arguments,
-		&ConvertCookedPathToRelativePath,
 		&Packages,
 		&PackageNameMap,
-		&PackageIdMap](const FContainerSourceFile& SourceFile, const FString& RelativeFileName, FContainerTargetFile& OutTargetFile) -> bool
+		&PackageIdMap](const FContainerSourceFile& SourceFile, FContainerTargetFile& OutTargetFile) -> bool
 	{
 		FCookedPackageStore& PackageStore = *Arguments.PackageStore;
 		
@@ -1727,12 +1841,10 @@ void InitializeContainerTargetsAndPackages(
 			IOSTORE_CPU_SCOPE(ProcessSourceFiles);
 			for (const FContainerSourceFile& SourceFile : ContainerSource.SourceFiles)
 			{
-				FString RelativeFileName = ConvertCookedPathToRelativePath(SourceFile.NormalizedPath);
-
 				FContainerTargetFile TargetFile;
-				bool bIsValidTargetFile = Arguments.PackageStore.IsValid()
-					? CreateTargetFileFromPackageStore(SourceFile, RelativeFileName, TargetFile)
-					: CreateTargetFileFromCookedFile(SourceFile, RelativeFileName, TargetFile);
+				bool bIsValidTargetFile = Arguments.PackageStore->HasDataSource()
+					? CreateTargetFileFromPackageStore(SourceFile, TargetFile)
+					: CreateTargetFileFromCookedFile(SourceFile, TargetFile);
 
 				if (!bIsValidTargetFile)
 				{
@@ -1740,7 +1852,6 @@ void InitializeContainerTargetsAndPackages(
 				}
 
 				TargetFile.ContainerTarget			= ContainerTarget;
-				TargetFile.TargetPath				= MoveTemp(RelativeFileName);
 				TargetFile.DestinationPath			= SourceFile.DestinationPath;
 				TargetFile.bForceUncompressed		= !SourceFile.bNeedsCompression;
 				
@@ -1983,7 +2094,7 @@ public:
 
 	IIoStoreWriteRequest* Read(const FContainerTargetFile& InTargetFile)
 	{
-		if (PackageStore)
+		if (PackageStore->HasDataSource())
 		{
 			return new FCookedPackageStoreWriteRequest(*this, InTargetFile);
 		}
@@ -2312,7 +2423,7 @@ int32 CreateTarget(const FIoStoreArguments& Arguments, const FIoStoreWriterSetti
 	FPackageIdMap PackageIdMap;
 
 	FPackageStoreOptimizer PackageStoreOptimizer;
-	PackageStoreOptimizer.Initialize(Arguments.TargetPlatform);
+	PackageStoreOptimizer.Initialize(*Arguments.ScriptObjects);
 	FIoStoreWriteRequestManager WriteRequestManager(PackageStoreOptimizer, Arguments.PackageStore.Get());
 
 	TArray<FContainerTargetSpec*> ContainerTargets;
@@ -2371,7 +2482,7 @@ int32 CreateTarget(const FIoStoreArguments& Arguments, const FIoStoreWriterSetti
 		}
 	}
 
-	if (Arguments.PackageStore.IsValid())
+	if (Arguments.PackageStore->HasDataSource())
 	{
 		ParsePackageAssetsFromPackageStore(*Arguments.PackageStore, Packages, PackageStoreOptimizer);
 	}
@@ -2408,7 +2519,7 @@ int32 CreateTarget(const FIoStoreArguments& Arguments, const FIoStoreWriterSetti
 				if (TargetFile.ChunkType != EContainerChunkType::PackageData)
 				{
 					FIoWriteOptions WriteOptions;
-					WriteOptions.DebugName = *TargetFile.TargetPath;
+					WriteOptions.DebugName = *TargetFile.DestinationPath;
 					WriteOptions.bForceUncompressed = TargetFile.bForceUncompressed;
 					WriteOptions.bIsMemoryMapped = TargetFile.ChunkType == EContainerChunkType::MemoryMappedBulkData;
 					WriteOptions.FileName = TargetFile.DestinationPath;
@@ -2464,7 +2575,7 @@ int32 CreateTarget(const FIoStoreArguments& Arguments, const FIoStoreWriterSetti
 				{
 					check(TargetFile.Package);
 					FIoWriteOptions WriteOptions;
-					WriteOptions.DebugName = *TargetFile.TargetPath;
+					WriteOptions.DebugName = *TargetFile.DestinationPath;
 					WriteOptions.bForceUncompressed = TargetFile.bForceUncompressed;
 					WriteOptions.FileName = TargetFile.DestinationPath;
 					if (TargetFile.SourceBuffer.IsSet())
@@ -2788,7 +2899,7 @@ bool IterateDirectoryIndex(FIoDirectoryIndexHandle Directory, const FString& Pat
 }
 
 int32 ListContainer(
-	const FIoStoreArguments& Arguments,
+	const FKeyChain& KeyChain,
 	const FString& ContainerPathOrWildcard,
 	const FString& CsvPath)
 {
@@ -2824,7 +2935,7 @@ int32 ListContainer(
 
 	for (const FString& ContainerFilePath : ContainerFilePaths)
 	{
-		TUniquePtr<FIoStoreReader> Reader = CreateIoStoreReader(*ContainerFilePath, Arguments.KeyChain);
+		TUniquePtr<FIoStoreReader> Reader = CreateIoStoreReader(*ContainerFilePath, KeyChain);
 		if (!Reader.IsValid())
 		{
 			UE_LOG(LogIoStore, Warning, TEXT("Failed to read container '%s'"), *ContainerFilePath);
@@ -3932,7 +4043,7 @@ static int32 Diff(
 	return 0;
 }
 
-int32 Staged2Zen(const FString& BuildPath, const FKeyChain& KeyChain, const ITargetPlatform* TargetPlatform)
+int32 Staged2Zen(const FString& BuildPath, const FKeyChain& KeyChain, const FString& ProjectName, const ITargetPlatform* TargetPlatform)
 {
 	FString PlatformName = TargetPlatform->PlatformName();
 	FString CookedOutputPath = FPaths::Combine(FPaths::ProjectDir(), TEXT("Saved"), TEXT("Cooked"), PlatformName);
@@ -3967,10 +4078,10 @@ int32 Staged2Zen(const FString& BuildPath, const FKeyChain& KeyChain, const ITar
 	PakPlatformFile.Initialize(&FPlatformFileManager::Get().GetPlatformFile(), TEXT(""));
 	FString CookedEngineContentPath = FPaths::Combine(CookedOutputPath, TEXT("Engine"), TEXT("Content"));
 	IFileManager::Get().MakeDirectory(*CookedEngineContentPath, true);
-	FString CookedProjectContentPath = FPaths::Combine(CookedOutputPath, FApp::GetProjectName(), TEXT("Content"));
+	FString CookedProjectContentPath = FPaths::Combine(CookedOutputPath, ProjectName, TEXT("Content"));
 	IFileManager::Get().MakeDirectory(*CookedProjectContentPath, true);
 	FString EngineContentPakPath = TEXT("../../../Engine/Content/");
-	FString ProjectContentPakPath = FPaths::Combine(TEXT("../../.."), FApp::GetProjectName(), TEXT("Content"));
+	FString ProjectContentPakPath = FPaths::Combine(TEXT("../../.."), ProjectName, TEXT("Content"));
 	for (const FString& PakFilePath : PakFiles)
 	{
 		PakPlatformFile.Mount(*PakFilePath, 0);
@@ -3981,7 +4092,7 @@ int32 Staged2Zen(const FString& BuildPath, const FKeyChain& KeyChain, const ITar
 			FString FileName = FPaths::GetCleanFilename(FileInPak);
 			if (FileName == TEXT("AssetRegistry.bin"))
 			{
-				FString TargetPath = FPaths::Combine(CookedOutputPath, FApp::GetProjectName(), FileName);
+				FString TargetPath = FPaths::Combine(CookedOutputPath, ProjectName, FileName);
 				PakPlatformFile.CopyFile(*TargetPath, *FileInPak);
 			}
 			else if (FileName.EndsWith(TEXT(".ushaderbytecode")))
@@ -4159,7 +4270,7 @@ int32 Staged2Zen(const FString& BuildPath, const FKeyChain& KeyChain, const ITar
 		}
 	}
 
-	FString MetaDataOutputPath = FPaths::Combine(CookedOutputPath, FApp::GetProjectName(), TEXT("Metadata"));
+	FString MetaDataOutputPath = FPaths::Combine(CookedOutputPath, ProjectName, TEXT("Metadata"));
 	TUniquePtr<FZenStoreWriter> ZenStoreWriter = MakeUnique<FZenStoreWriter>(CookedOutputPath, MetaDataOutputPath, TargetPlatform);
 	
 	ICookedPackageWriter::FCookInfo CookInfo;
@@ -4270,7 +4381,7 @@ static bool ParsePakResponseFile(const TCHAR* FilePath, TArray<FContainerSourceF
 	return true;
 }
 
-static bool ParsePakOrderFile(const TCHAR* FilePath, FFileOrderMap& Map)
+static bool ParsePakOrderFile(const TCHAR* FilePath, FFileOrderMap& Map, const FIoStoreArguments& Arguments)
 {
 	IOSTORE_CPU_SCOPE(ParsePakOrderFile);
 
@@ -4287,7 +4398,7 @@ static bool ParsePakOrderFile(const TCHAR* FilePath, FFileOrderMap& Map)
 	for (const FString& OrderLine : OrderFileContents)
 	{
 		const TCHAR* OrderLinePtr = *OrderLine;
-		FString Path;
+		FString PackageName;
 
 		// Skip comments
 		if (FCString::Strncmp(OrderLinePtr, TEXT("#"), 1) == 0 || FCString::Strncmp(OrderLinePtr, TEXT("//"), 2) == 0)
@@ -4295,19 +4406,25 @@ static bool ParsePakOrderFile(const TCHAR* FilePath, FFileOrderMap& Map)
 			continue;
 		}
 
-		if (!FParse::Token(OrderLinePtr, Path, false))
+		if (!FParse::Token(OrderLinePtr, PackageName, false))
 		{
 			UE_LOG(LogIoStore, Error, TEXT("Invalid line in order file '%s'."), *OrderLine);
 			return false;
 		}
-		FString PackageName;
-		if (!FPackageName::TryConvertFilenameToLongPackageName(Path, PackageName, nullptr))
+
+		FName PackageFName;
+		if (FPackageName::IsValidTextForLongPackageName(PackageName))
+		{
+			PackageFName = FName(PackageName);
+		}
+		else if (PackageName.StartsWith(TEXT("../../../")))
 		{
-			continue;;
+			FString FullFileName = FPaths::Combine(Arguments.CookedDir, PackageName.RightChop(9));
+			FPaths::NormalizeFilename(FullFileName);
+			PackageFName = Arguments.PackageStore->GetPackageNameFromFileName(FullFileName);
 		}
 
-		FName PackageFName(MoveTemp(PackageName));
-		if (!Map.PackageNameToOrder.Contains(PackageFName))
+		if (!PackageFName.IsNone() && !Map.PackageNameToOrder.Contains(PackageFName))
 		{
 			Map.PackageNameToOrder.Emplace(PackageFName, NextOrder++);
 		}
@@ -4444,33 +4561,8 @@ static bool ParseSizeArgument(const TCHAR* CmdLine, const TCHAR* Argument, uint6
 	}
 }
 
-int32 CreateIoStoreContainerFiles(const TCHAR* CmdLine)
+static bool ParseOrderFileArguments(FIoStoreArguments& Arguments)
 {
-	IOSTORE_CPU_SCOPE(CreateIoStoreContainerFiles);
-
-	UE_LOG(LogIoStore, Display, TEXT("==================== IoStore Utils ===================="));
-
-	FIoStoreArguments Arguments;
-
-	LoadKeyChain(FCommandLine::Get(), Arguments.KeyChain);
-
-	if (FParse::Param(FCommandLine::Get(), TEXT("sign")))
-	{
-		Arguments.bSign = true;
-	}
-
-	UE_LOG(LogIoStore, Display, TEXT("Container signing - %s"), Arguments.bSign ? TEXT("ENABLED") : TEXT("DISABLED"));
-
-	Arguments.bCreateDirectoryIndex = !FParse::Param(FCommandLine::Get(), TEXT("NoDirectoryIndex"));
-	UE_LOG(LogIoStore, Display, TEXT("Directory index - %s"), Arguments.bCreateDirectoryIndex  ? TEXT("ENABLED") : TEXT("DISABLED"));
-
-	FString PatchReferenceCryptoKeysFilename;
-	FKeyChain PatchKeyChain;
-	if (FParse::Value(FCommandLine::Get(), TEXT("PatchCryptoKeys="), PatchReferenceCryptoKeysFilename))
-	{
-		KeyChainUtilities::LoadKeyChainFromFile(PatchReferenceCryptoKeysFilename, Arguments.PatchKeyChain);
-	}
-
 	uint64 OrderMapStartIndex = 0;
 	FString OrderFileStr;
 	if (FParse::Value(FCommandLine::Get(), TEXT("Order="), OrderFileStr, false))
@@ -4503,7 +4595,7 @@ int32 CreateIoStoreContainerFiles(const TCHAR* CmdLine)
 			if (PriorityStrings.Num() != OrderFilePaths.Num())
 			{
 				UE_LOG(LogIoStore, Error, TEXT("Number of parameters to -Order= and -OrderPriority= do not match"));
-				return -1;
+				return false;
 			}
 
 			for (const FString& PriorityString : PriorityStrings)
@@ -4520,42 +4612,46 @@ int32 CreateIoStoreContainerFiles(const TCHAR* CmdLine)
 		check(OrderFilePaths.Num() == OrderFilePriorities.Num());
 
 		bool bMerge = false;
-		for (int32 i=0; i < OrderFilePaths.Num(); ++i)
+		for (int32 i = 0; i < OrderFilePaths.Num(); ++i)
 		{
 			FString& OrderFile = OrderFilePaths[i];
 			int32 Priority = OrderFilePriorities[i];
 
 			FFileOrderMap OrderMap(Priority, i);
-			if (!ParsePakOrderFile(*OrderFile, OrderMap))
+			if (!ParsePakOrderFile(*OrderFile, OrderMap, Arguments))
 			{
-				return -1;
+				return false;
 			}
 			Arguments.OrderMaps.Add(OrderMap);
 		}
 	}
 
-	FString TargetPlatformName;
-	if (FParse::Value(FCommandLine::Get(), TEXT("TargetPlatform="), TargetPlatformName))
+	Arguments.bClusterByOrderFilePriority = !FParse::Param(FCommandLine::Get(), TEXT("DoNotClusterByOrderPriority"));
+	
+	return true;
+}
+
+bool ParseContainerGenerationArguments(FIoStoreArguments& Arguments, FIoStoreWriterSettings& WriterSettings)
+{
+	if (FParse::Param(FCommandLine::Get(), TEXT("sign")))
 	{
-		UE_LOG(LogIoStore, Display, TEXT("Using target platform '%s'"), *TargetPlatformName);
-		ITargetPlatformManagerModule& TPM = GetTargetPlatformManagerRef();
-		Arguments.TargetPlatform = TPM.FindTargetPlatform(TargetPlatformName);
-		if (!Arguments.TargetPlatform)
-		{
-			UE_LOG(LogIoStore, Error, TEXT("Invalid TargetPlatform: '%s'"), *TargetPlatformName);
-			return 1;
-		}
+		Arguments.bSign = true;
 	}
 
-	Arguments.bClusterByOrderFilePriority = !FParse::Param(CmdLine, TEXT("DoNotClusterByOrderPriority"));
+	UE_LOG(LogIoStore, Display, TEXT("Container signing - %s"), Arguments.bSign ? TEXT("ENABLED") : TEXT("DISABLED"));
 
-	FIoStoreWriterSettings GeneralIoWriterSettings { DefaultCompressionMethod, DefaultCompressionBlockSize, false };
-	GeneralIoWriterSettings.bEnableCsvOutput = FParse::Param(CmdLine, TEXT("-csvoutput"));
+	Arguments.bCreateDirectoryIndex = !FParse::Param(FCommandLine::Get(), TEXT("NoDirectoryIndex"));
+	UE_LOG(LogIoStore, Display, TEXT("Directory index - %s"), Arguments.bCreateDirectoryIndex ? TEXT("ENABLED") : TEXT("DISABLED"));
+
+	WriterSettings.CompressionMethod = DefaultCompressionMethod;
+	WriterSettings.CompressionBlockSize = DefaultCompressionBlockSize;
+
+	WriterSettings.bEnableCsvOutput = FParse::Param(FCommandLine::Get(), TEXT("csvoutput"));
 
 	TArray<FName> CompressionFormats;
 	FString DesiredCompressionFormats;
-	if (FParse::Value(CmdLine, TEXT("-compressionformats="), DesiredCompressionFormats) ||
-		FParse::Value(CmdLine, TEXT("-compressionformat="), DesiredCompressionFormats))
+	if (FParse::Value(FCommandLine::Get(), TEXT("-compressionformats="), DesiredCompressionFormats) ||
+		FParse::Value(FCommandLine::Get(), TEXT("-compressionformat="), DesiredCompressionFormats))
 	{
 		TArray<FString> Formats;
 		DesiredCompressionFormats.ParseIntoArray(Formats, TEXT(","));
@@ -4566,69 +4662,89 @@ int32 CreateIoStoreContainerFiles(const TCHAR* CmdLine)
 
 			if (FCompression::IsFormatValid(FormatName))
 			{
-				GeneralIoWriterSettings.CompressionMethod = FormatName;
+				WriterSettings.CompressionMethod = FormatName;
 				break;
 			}
 		}
 
-		if (GeneralIoWriterSettings.CompressionMethod == NAME_None)
+		if (WriterSettings.CompressionMethod == NAME_None)
 		{
 			UE_LOG(LogIoStore, Warning, TEXT("Failed to find desired compression format(s) '%s'. Using falling back to '%s'"),
 				*DesiredCompressionFormats, *DefaultCompressionMethod.ToString());
 		}
 		else
 		{
-			UE_LOG(LogIoStore, Display, TEXT("Using compression format '%s'"), *GeneralIoWriterSettings.CompressionMethod.ToString());
+			UE_LOG(LogIoStore, Display, TEXT("Using compression format '%s'"), *WriterSettings.CompressionMethod.ToString());
 		}
 	}
 
-	ParseSizeArgument(CmdLine, TEXT("-alignformemorymapping="), GeneralIoWriterSettings.MemoryMappingAlignment, DefaultMemoryMappingAlignment);
-	ParseSizeArgument(CmdLine, TEXT("-compressionblocksize="), GeneralIoWriterSettings.CompressionBlockSize, DefaultCompressionBlockSize);
-		
-	GeneralIoWriterSettings.CompressionBlockAlignment = DefaultCompressionBlockAlignment;
-	
+	ParseSizeArgument(FCommandLine::Get(), TEXT("-alignformemorymapping="), WriterSettings.MemoryMappingAlignment, DefaultMemoryMappingAlignment);
+	ParseSizeArgument(FCommandLine::Get(), TEXT("-compressionblocksize="), WriterSettings.CompressionBlockSize, DefaultCompressionBlockSize);
+
+	WriterSettings.CompressionBlockAlignment = DefaultCompressionBlockAlignment;
+
 	uint64 BlockAlignment = 0;
-	if (ParseSizeArgument(CmdLine, TEXT("-blocksize="), BlockAlignment))
+	if (ParseSizeArgument(FCommandLine::Get(), TEXT("-blocksize="), BlockAlignment))
 	{
-		GeneralIoWriterSettings.CompressionBlockAlignment = BlockAlignment;
+		WriterSettings.CompressionBlockAlignment = BlockAlignment;
 	}
-	
+
 	uint64 PatchPaddingAlignment = 0;
-	if (ParseSizeArgument(CmdLine, TEXT("-patchpaddingalign="), PatchPaddingAlignment))
+	if (ParseSizeArgument(FCommandLine::Get(), TEXT("-patchpaddingalign="), PatchPaddingAlignment))
 	{
-		if (PatchPaddingAlignment < GeneralIoWriterSettings.CompressionBlockAlignment)
+		if (PatchPaddingAlignment < WriterSettings.CompressionBlockAlignment)
 		{
-			GeneralIoWriterSettings.CompressionBlockAlignment = PatchPaddingAlignment;
+			WriterSettings.CompressionBlockAlignment = PatchPaddingAlignment;
 		}
 	}
-	
+
 	// Temporary, this command-line allows us to explicitly override the value otherwise shared between pak building and iostore
 	uint64 IOStorePatchPaddingAlignment = 0;
-	if (ParseSizeArgument(CmdLine, TEXT("-iostorepatchpaddingalign="), IOStorePatchPaddingAlignment))
+	if (ParseSizeArgument(FCommandLine::Get(), TEXT("-iostorepatchpaddingalign="), IOStorePatchPaddingAlignment))
 	{
-		GeneralIoWriterSettings.CompressionBlockAlignment = IOStorePatchPaddingAlignment;
+		WriterSettings.CompressionBlockAlignment = IOStorePatchPaddingAlignment;
 	}
 
 	uint64 MaxPartitionSize = 0;
-	if (ParseSizeArgument(CmdLine, TEXT("-maxPartitionSize="), MaxPartitionSize))
+	if (ParseSizeArgument(FCommandLine::Get(), TEXT("-maxPartitionSize="), MaxPartitionSize))
+	{
+		WriterSettings.MaxPartitionSize = MaxPartitionSize;
+	}
+
+	int32 CompressionMinBytesSaved = 0;
+	if (FParse::Value(FCommandLine::Get(), TEXT("-compressionMinBytesSaved="), CompressionMinBytesSaved))
 	{
-		GeneralIoWriterSettings.MaxPartitionSize = MaxPartitionSize;
+		WriterSettings.CompressionMinBytesSaved = CompressionMinBytesSaved;
 	}
 
-	if (Arguments.TargetPlatform)
+	int32 CompressionMinPercentSaved = 0;
+	if (FParse::Value(FCommandLine::Get(), TEXT("-compressionMinPercentSaved="), CompressionMinPercentSaved))
 	{
-		GeneralIoWriterSettings.InitializePlatformSpecificSettings(Arguments.TargetPlatform);
+		WriterSettings.CompressionMinPercentSaved = CompressionMinPercentSaved;
 	}
 
-	UE_LOG(LogIoStore, Display, TEXT("Using memory mapping alignment '%ld'"), GeneralIoWriterSettings.MemoryMappingAlignment);
-	UE_LOG(LogIoStore, Display, TEXT("Using compression block size '%ld'"), GeneralIoWriterSettings.CompressionBlockSize);
-	UE_LOG(LogIoStore, Display, TEXT("Using compression block alignment '%ld'"), GeneralIoWriterSettings.CompressionBlockAlignment);
-	UE_LOG(LogIoStore, Display, TEXT("Using compression min bytes saved '%d'"), GeneralIoWriterSettings.CompressionMinBytesSaved);
-	UE_LOG(LogIoStore, Display, TEXT("Using compression min percent saved '%d'"), GeneralIoWriterSettings.CompressionMinPercentSaved);
-	UE_LOG(LogIoStore, Display, TEXT("Using max partition size '%lld'"), GeneralIoWriterSettings.MaxPartitionSize);
+	WriterSettings.bCompressionEnableDDC = FParse::Param(FCommandLine::Get(), TEXT("compressionEnableDDC"));
 
-	FParse::Value(CmdLine, TEXT("-MetaOutputDirectory="), Arguments.MetaOutputDir);
-	FParse::Value(CmdLine, TEXT("-MetaInputDirectory="), Arguments.MetaInputDir);
+	int32 CompressionMinSizeToConsiderDDC = 0;
+	if (FParse::Value(FCommandLine::Get(), TEXT("-compressionMinSizeToConsiderDDC="), CompressionMinSizeToConsiderDDC))
+	{
+		WriterSettings.CompressionMinSizeToConsiderDDC = CompressionMinSizeToConsiderDDC;
+	}
+
+	UE_LOG(LogIoStore, Display, TEXT("Using memory mapping alignment '%ld'"), WriterSettings.MemoryMappingAlignment);
+	UE_LOG(LogIoStore, Display, TEXT("Using compression block size '%ld'"), WriterSettings.CompressionBlockSize);
+	UE_LOG(LogIoStore, Display, TEXT("Using compression block alignment '%ld'"), WriterSettings.CompressionBlockAlignment);
+	UE_LOG(LogIoStore, Display, TEXT("Using compression min bytes saved '%d'"), WriterSettings.CompressionMinBytesSaved);
+	UE_LOG(LogIoStore, Display, TEXT("Using compression min percent saved '%d'"), WriterSettings.CompressionMinPercentSaved);
+	UE_LOG(LogIoStore, Display, TEXT("Using max partition size '%lld'"), WriterSettings.MaxPartitionSize);
+	if (WriterSettings.bCompressionEnableDDC)
+	{
+		UE_LOG(LogIoStore, Display, TEXT("Using DDC for compression with min size '%d'"), WriterSettings.CompressionMinSizeToConsiderDDC);
+	}
+	else
+	{
+		UE_LOG(LogIoStore, Display, TEXT("Not using DDC for compression"));
+	}
 
 	FString CommandListFile;
 	if (FParse::Value(FCommandLine::Get(), TEXT("Commands="), CommandListFile))
@@ -4638,7 +4754,7 @@ int32 CreateIoStoreContainerFiles(const TCHAR* CmdLine)
 		if (!FFileHelper::LoadFileToStringArray(Commands, *CommandListFile))
 		{
 			UE_LOG(LogIoStore, Error, TEXT("Failed to read command list file '%s'."), *CommandListFile);
-			return -1;
+			return false;
 		}
 
 		Arguments.Containers.Reserve(Commands.Num());
@@ -4649,7 +4765,7 @@ int32 CreateIoStoreContainerFiles(const TCHAR* CmdLine)
 			if (!FParse::Value(*Command, TEXT("Output="), ContainerSpec.OutputPath))
 			{
 				UE_LOG(LogIoStore, Error, TEXT("Output argument missing from command '%s'"), *Command);
-				return -1;
+				return false;
 			}
 			ContainerSpec.OutputPath = FPaths::ChangeExtension(ContainerSpec.OutputPath, TEXT(""));
 
@@ -4681,7 +4797,7 @@ int32 CreateIoStoreContainerFiles(const TCHAR* CmdLine)
 				if (!ParsePakResponseFile(*ResponseFilePath, ContainerSpec.SourceFiles))
 				{
 					UE_LOG(LogIoStore, Error, TEXT("Failed to parse Pak response file '%s'"), *ResponseFilePath);
-					return -1;
+					return false;
 				}
 
 				FString EncryptionKeyOverrideGuidString;
@@ -4693,19 +4809,166 @@ int32 CreateIoStoreContainerFiles(const TCHAR* CmdLine)
 		}
 	}
 
-	if (FParse::Value(FCommandLine::Get(), TEXT("BasedOnReleaseVersionPath="), Arguments.BasedOnReleaseVersionPath))
+	for (const FContainerSourceSpec& Container : Arguments.Containers)
+	{
+		if (Container.Name.IsNone())
+		{
+			UE_LOG(LogIoStore, Error, TEXT("ContainerName argument missing for container '%s'"), *Container.OutputPath);
+			return false;
+		}
+	}
+
+	Arguments.bFileRegions = FParse::Param(FCommandLine::Get(), TEXT("FileRegions"));
+	WriterSettings.bEnableFileRegions = Arguments.bFileRegions;
+
+	FString PatchReferenceCryptoKeysFilename;
+	FKeyChain PatchKeyChain;
+	if (FParse::Value(FCommandLine::Get(), TEXT("PatchCryptoKeys="), PatchReferenceCryptoKeysFilename))
+	{
+		KeyChainUtilities::LoadKeyChainFromFile(PatchReferenceCryptoKeysFilename, Arguments.PatchKeyChain);
+	}
+
+	return true;
+}
+
+int32 CreateIoStoreContainerFiles(const TCHAR* CmdLine)
+{
+	IOSTORE_CPU_SCOPE(CreateIoStoreContainerFiles);
+
+	UE_LOG(LogIoStore, Display, TEXT("==================== IoStore Utils ===================="));
+
+	FIoStoreArguments Arguments;
+	FIoStoreWriterSettings WriterSettings;
+
+	LoadKeyChain(FCommandLine::Get(), Arguments.KeyChain);
+	
+	FString ArgumentValue;
+	if (FParse::Value(FCommandLine::Get(), TEXT("List="), ArgumentValue))
+	{
+		FString ContainerPathOrWildcard = MoveTemp(ArgumentValue);
+		FString CsvPath;
+		if (!FParse::Value(FCommandLine::Get(), TEXT("csv="), CsvPath))
+		{
+			UE_LOG(LogIoStore, Error, TEXT("Incorrect arguments. Expected: -list=<ContainerFile> -csv=<path>"));
+		}
+
+		return ListContainer(Arguments.KeyChain, ContainerPathOrWildcard, CsvPath);
+	}
+	else if (FParse::Value(FCommandLine::Get(), TEXT("Describe="), ArgumentValue))
+	{
+		FString ContainerPathOrWildcard = ArgumentValue;
+		FString PackageFilter;
+		FParse::Value(FCommandLine::Get(), TEXT("PackageFilter="), PackageFilter);
+		FString OutPath;
+		FParse::Value(FCommandLine::Get(), TEXT("DumpToFile="), OutPath);
+		bool bIncludeExportHashes = FParse::Param(FCommandLine::Get(), TEXT("IncludeExportHashes"));
+		return Describe(ContainerPathOrWildcard, Arguments.KeyChain, PackageFilter, OutPath, bIncludeExportHashes);
+	}
+	else if (FParse::Param(FCommandLine::Get(), TEXT("Diff")))
+	{
+		FString SourcePath, TargetPath, OutPath;
+		FKeyChain SourceKeyChain, TargetKeyChain;
+
+		if (!FParse::Value(FCommandLine::Get(), TEXT("Source="), SourcePath))
+		{
+			UE_LOG(LogIoStore, Error, TEXT("Incorrect arguments. Expected: -Diff -Source=<Path> -Target=<path>"));
+			return -1;
+		}
+
+		if (!IFileManager::Get().DirectoryExists(*SourcePath))
+		{
+			UE_LOG(LogIoStore, Error, TEXT("Source directory '%s' doesn't exist"), *SourcePath);
+			return -1;
+		}
+
+		if (!FParse::Value(FCommandLine::Get(), TEXT("Target="), TargetPath))
+		{
+			UE_LOG(LogIoStore, Error, TEXT("Incorrect arguments. Expected: -Diff -Source=<Path> -Target=<path>"));
+		}
+
+		if (!IFileManager::Get().DirectoryExists(*TargetPath))
+		{
+			UE_LOG(LogIoStore, Error, TEXT("Target directory '%s' doesn't exist"), *TargetPath);
+			return -1;
+		}
+
+		FParse::Value(FCommandLine::Get(), TEXT("DumpToFile="), OutPath);
+
+		FString CryptoKeysCacheFilename;
+		if (FParse::Value(CmdLine, TEXT("SourceCryptoKeys="), CryptoKeysCacheFilename))
+		{
+			UE_LOG(LogIoStore, Display, TEXT("Parsing source crypto keys from '%s'"), *CryptoKeysCacheFilename);
+			KeyChainUtilities::LoadKeyChainFromFile(CryptoKeysCacheFilename, SourceKeyChain);
+		}
+
+		if (FParse::Value(CmdLine, TEXT("TargetCryptoKeys="), CryptoKeysCacheFilename))
+		{
+			UE_LOG(LogIoStore, Display, TEXT("Parsing target crypto keys from '%s'"), *CryptoKeysCacheFilename);
+			KeyChainUtilities::LoadKeyChainFromFile(CryptoKeysCacheFilename, TargetKeyChain);
+		}
+
+		return Diff(SourcePath, SourceKeyChain, TargetPath, TargetKeyChain, OutPath);
+	}
+	else if (FParse::Param(FCommandLine::Get(), TEXT("Staged2Zen")))
 	{
-		UE_LOG(LogIoStore, Display, TEXT("Based on release version path: '%s'"), *Arguments.BasedOnReleaseVersionPath);
+		ITargetPlatform* TargetPlatform = nullptr;
+		FString TargetPlatformName;
+		if (FParse::Value(FCommandLine::Get(), TEXT("TargetPlatform="), TargetPlatformName))
+		{
+			ITargetPlatformManagerModule& TPM = GetTargetPlatformManagerRef();
+			TargetPlatform = TPM.FindTargetPlatform(TargetPlatformName);
+			if (!TargetPlatform)
+			{
+				UE_LOG(LogIoStore, Error, TEXT("Invalid TargetPlatform: '%s'"), *TargetPlatformName);
+				return 1;
+			}
+		}
+
+		FString BuildPath;
+		FString ProjectName;
+		if (!FParse::Value(FCommandLine::Get(), TEXT("BuildPath="), BuildPath) ||
+			!FParse::Value(FCommandLine::Get(), TEXT("ProjectName="), ProjectName) ||
+			!TargetPlatform)
+		{
+			UE_LOG(LogIoStore, Error, TEXT("Incorrect arguments. Expected: -Staged2Zen -BuildPath=<Path> -ProjectName=<ProjectName> -TargetPlatform=<Platform>"));
+		}
+		return Staged2Zen(BuildPath, Arguments.KeyChain, ProjectName, TargetPlatform);
 	}
+	else if (FParse::Param(FCommandLine::Get(), TEXT("CreateContentPatch")))
+	{
+		if (!ParseContainerGenerationArguments(Arguments, WriterSettings))
+		{
+			return -1;
+		}
 
-	if (FParse::Value(FCommandLine::Get(), TEXT("DLCFile="), Arguments.DLCPluginPath))
+		for (const FContainerSourceSpec& Container : Arguments.Containers)
+		{
+			if (Container.PatchTargetFile.IsEmpty())
+			{
+				UE_LOG(LogIoStore, Error, TEXT("PatchTarget argument missing for container '%s'"), *Container.OutputPath);
+				return -1;
+			}
+		}
+
+		return CreateContentPatch(Arguments, WriterSettings);
+	}
+	else if (FParse::Value(FCommandLine::Get(), TEXT("CreateDLCContainer="), Arguments.DLCPluginPath))
 	{
+		if (!ParseContainerGenerationArguments(Arguments, WriterSettings))
+		{
+			return -1;
+		}
+		
 		Arguments.DLCName = FPaths::GetBaseFilename(*Arguments.DLCPluginPath);
 		Arguments.bRemapPluginContentToGame = FParse::Param(FCommandLine::Get(), TEXT("RemapPluginContentToGame"));
 
 		UE_LOG(LogIoStore, Display, TEXT("DLC: '%s'"), *Arguments.DLCPluginPath);
 		UE_LOG(LogIoStore, Display, TEXT("Remapping plugin content to game: '%s'"), Arguments.bRemapPluginContentToGame ? TEXT("True") : TEXT("False"));
 
+		if (FParse::Value(FCommandLine::Get(), TEXT("BasedOnReleaseVersionPath="), Arguments.BasedOnReleaseVersionPath))
+		{
+			UE_LOG(LogIoStore, Display, TEXT("Based on release version path: '%s'"), *Arguments.BasedOnReleaseVersionPath);
+		}
 		if (Arguments.BasedOnReleaseVersionPath.IsEmpty())
 		{
 			UE_LOG(LogIoStore, Error, TEXT("Based on release version path is needed for DLC"));
@@ -4742,176 +5005,83 @@ int32 CreateIoStoreContainerFiles(const TCHAR* CmdLine)
 			UE_LOG(LogIoStore, Warning, TEXT("Failed to load Asset registry '%s'. Needed to verify DLC package names"), *DevelopmentAssetRegistryPath);
 		}
 	}
-
-	if (FParse::Value(FCommandLine::Get(), TEXT("CreateGlobalContainer="), Arguments.GlobalContainerPath))
+	else if (FParse::Value(FCommandLine::Get(), TEXT("CreateGlobalContainer="), Arguments.GlobalContainerPath))
 	{
 		Arguments.GlobalContainerPath = FPaths::ChangeExtension(Arguments.GlobalContainerPath, TEXT(""));
-	}
-
-	if (Arguments.ShouldCreateContainers())
-	{
-		if (!Arguments.TargetPlatform)
-		{
-			UE_LOG(LogIoStore, Error, TEXT("TargetPlatform must be specified"));
-			return 1;
-		}
 
-		// Now that we have the target platform we need to error out if LegacyBulkDataOffsets is enabled for it
+		if (!ParseContainerGenerationArguments(Arguments, WriterSettings))
 		{
-			FConfigFile PlatformEngineIni;
-			FConfigCacheIni::LoadLocalIniFile(PlatformEngineIni, TEXT("Engine"), true, *Arguments.TargetPlatform->IniPlatformName());
-
-			bool bLegacyBulkDataOffsets = false;
-			PlatformEngineIni.GetBool(TEXT("Core.System"), TEXT("LegacyBulkDataOffsets"), bLegacyBulkDataOffsets);
-
-			if (bLegacyBulkDataOffsets)
-			{
-				UE_LOG(LogIoStore, Error, TEXT("'LegacyBulkDataOffsets' is enabled for the target platform '%s', this needs to be disabled and the data recooked in order for the IoStore to work"), *Arguments.TargetPlatform->IniPlatformName());
-				return 1;
-			}
-		}
-
-		if (!FParse::Value(FCommandLine::Get(), TEXT("CookedDirectory="), Arguments.CookedDir))
-		{
-			UE_LOG(LogIoStore, Error, TEXT("CookedDirectory must be specified"));
-			return 1;
-		}
-
-		for (const FContainerSourceSpec& Container : Arguments.Containers)
-		{
-			if (Container.Name.IsNone())
-			{
-				UE_LOG(LogIoStore, Error, TEXT("ContainerName argument missing for container '%s'"), *Container.OutputPath);
-				return -1;
-			}
+			return -1;
 		}
+	}
+	else
+	{
+		UE_LOG(LogIoStore, Display, TEXT("Usage:"));
+		UE_LOG(LogIoStore, Display, TEXT(" -List=</path/to/[container.utoc|*.utoc]> -CSV=<list.csv> [-CryptoKeys=</path/to/crypto.json>]"));
+		UE_LOG(LogIoStore, Display, TEXT(" -Describe=</path/to/global.utoc> [-PackageFilter=<PackageName>] [-DumpToFile=<describe.txt>] [-CryptoKeys=</path/to/crypto.json>]"));
+		return -1;
+	}
 
-		// Enable file region metadata if required by the target platform.
-		Arguments.bFileRegions = Arguments.TargetPlatform->SupportsFeature(ETargetPlatformFeatures::CookFileRegionMetadata);
-		GeneralIoWriterSettings.bEnableFileRegions = Arguments.bFileRegions;
+	// Common path for creating containers
+	FParse::Value(FCommandLine::Get(), TEXT("CookedDirectory="), Arguments.CookedDir);
+	if (Arguments.CookedDir.IsEmpty())
+	{
+		UE_LOG(LogIoStore, Error, TEXT("CookedDirectory must be specified"));
+		return -1;
+	}
 
-		TUniquePtr<FCookedPackageStore> PackageStore = MakeUnique<FCookedPackageStore>();
-		FIoStatus Status = PackageStore->Load(Arguments.CookedDir);
+	FString PackageStoreManifestFilename;
+	if (FParse::Value(FCommandLine::Get(), TEXT("PackageStoreManifest="), PackageStoreManifestFilename))
+	{
+		TUniquePtr<FCookedPackageStore> PackageStore = MakeUnique<FCookedPackageStore>(Arguments.CookedDir);
+		FIoStatus Status = PackageStore->Load(*PackageStoreManifestFilename);
 		if (Status.IsOk())
 		{
-			UE_LOG(LogIoStore, Display, TEXT("Using package store"));
 			Arguments.PackageStore = MoveTemp(PackageStore);
 		}
 		else
 		{
-			IOSTORE_CPU_SCOPE(FindCookedAssets);
-			UE_LOG(LogIoStore, Display, TEXT("Searching for cooked assets in folder '%s'"), *Arguments.CookedDir);
-			FCookedFileVisitor CookedFileVistor(Arguments.CookedFileStatMap, nullptr, Arguments.bFileRegions);
-			IFileManager::Get().IterateDirectoryStatRecursively(*Arguments.CookedDir, CookedFileVistor);
-			UE_LOG(LogIoStore, Display, TEXT("Found '%d' files"), Arguments.CookedFileStatMap.Num());
+			UE_LOG(LogIoStore, Fatal, TEXT("Failed loading package store manifest '%s'"), *PackageStoreManifestFilename);
 		}
-
-		const int32 ReturnValue = CreateTarget(Arguments, GeneralIoWriterSettings);
-		return ReturnValue;
 	}
-	else if (FParse::Param(FCommandLine::Get(), TEXT("CreateContentPatch")))
+	else
 	{
-		for (const FContainerSourceSpec& Container : Arguments.Containers)
-		{
-			if (Container.PatchTargetFile.IsEmpty())
-			{
-				UE_LOG(LogIoStore, Error, TEXT("PatchTarget argument missing for container '%s'"), *Container.OutputPath);
-				return -1;
-			}
-		}
+		UE_LOG(LogIoStore, Error, TEXT("Expected -PackageStoreManifest=<path to package store manifest>"));
+		return -1;
+	}
 
-		int32 ReturnValue = CreateContentPatch(Arguments, GeneralIoWriterSettings);
-		if (ReturnValue != 0)
+	if (!ParseOrderFileArguments(Arguments))
+	{
+		return false;
+	}
+
+	FString ScriptObjectsFile;
+	if (FParse::Value(FCommandLine::Get(), TEXT("ScriptObjects="), ScriptObjectsFile))
+	{
+		TArray<uint8> ScriptObjectsData;
+		if (!FFileHelper::LoadFileToArray(ScriptObjectsData, *ScriptObjectsFile))
 		{
-			return ReturnValue;
+			UE_LOG(LogIoStore, Fatal, TEXT("Failed reading script objects file '%s'"), *ScriptObjectsFile);
 		}
+		Arguments.ScriptObjects = MakeUnique<FIoBuffer>(FIoBuffer::Clone, ScriptObjectsData.GetData(), ScriptObjectsData.Num());
 	}
 	else
 	{
-		FString ContainerPathOrWildcard;
-		if (FParse::Value(FCommandLine::Get(), TEXT("List="), ContainerPathOrWildcard))
-		{
-			FString CsvPath;
-			if (!FParse::Value(FCommandLine::Get(), TEXT("csv="), CsvPath))
-			{
-				UE_LOG(LogIoStore, Error, TEXT("Incorrect arguments. Expected: -list=<ContainerFile> -csv=<path>"));
-			}
-
-			return ListContainer(Arguments, ContainerPathOrWildcard, CsvPath);
-		}
-		else if (FParse::Value(FCommandLine::Get(), TEXT("Describe="), ContainerPathOrWildcard))
-		{
-			FString PackageFilter;
-			FParse::Value(FCommandLine::Get(), TEXT("PackageFilter="), PackageFilter);
-			FString OutPath;
-			FParse::Value(FCommandLine::Get(), TEXT("DumpToFile="), OutPath);
-			bool bIncludeExportHashes = FParse::Param(FCommandLine::Get(), TEXT("IncludeExportHashes"));
-			return Describe(ContainerPathOrWildcard, Arguments.KeyChain, PackageFilter, OutPath, bIncludeExportHashes);
-		}
-		else if (FParse::Param(FCommandLine::Get(), TEXT("Diff")))
-		{
-			FString SourcePath, TargetPath, OutPath;
-			FKeyChain SourceKeyChain, TargetKeyChain;
-
-			if (!FParse::Value(FCommandLine::Get(), TEXT("Source="), SourcePath))
-			{
-				UE_LOG(LogIoStore, Error, TEXT("Incorrect arguments. Expected: -Diff -Source=<Path> -Target=<path>"));
-				return -1;
-			}
-
-			if (!IFileManager::Get().DirectoryExists(*SourcePath))
-			{
-				UE_LOG(LogIoStore, Error, TEXT("Source directory '%s' doesn't exist"), *SourcePath);
-				return -1;
-			}
-
-			if (!FParse::Value(FCommandLine::Get(), TEXT("Target="), TargetPath))
-			{
-				UE_LOG(LogIoStore, Error, TEXT("Incorrect arguments. Expected: -Diff -Source=<Path> -Target=<path>"));
-			}
-
-			if (!IFileManager::Get().DirectoryExists(*TargetPath))
-			{
-				UE_LOG(LogIoStore, Error, TEXT("Target directory '%s' doesn't exist"), *TargetPath);
-				return -1;
-			}
-
-			FParse::Value(FCommandLine::Get(), TEXT("DumpToFile="), OutPath);
-
-			FString CryptoKeysCacheFilename;
-			if (FParse::Value(CmdLine, TEXT("SourceCryptoKeys="), CryptoKeysCacheFilename))
-			{
-				UE_LOG(LogIoStore, Display, TEXT("Parsing source crypto keys from '%s'"), *CryptoKeysCacheFilename);
-				KeyChainUtilities::LoadKeyChainFromFile(CryptoKeysCacheFilename, SourceKeyChain);
-			}
+		UE_CLOG(!Arguments.PackageStore->HasDataSource(), LogIoStore, Fatal, TEXT("Expected -ScriptObjects=<path to script objects file> argument"));
+		TIoStatusOr<FIoBuffer> Status = Arguments.PackageStore->ReadChunk(CreateIoChunkId(0, 0, EIoChunkType::ScriptObjects));
+		UE_CLOG(!Status.IsOk(), LogIoStore, Fatal, TEXT("Failed reading script objects chunk '%s'"), *Status.Status().ToString());
+		Arguments.ScriptObjects = MakeUnique<FIoBuffer>(Status.ConsumeValueOrDie());
+	}
 
-			if (FParse::Value(CmdLine, TEXT("TargetCryptoKeys="), CryptoKeysCacheFilename))
-			{
-				UE_LOG(LogIoStore, Display, TEXT("Parsing target crypto keys from '%s'"), *CryptoKeysCacheFilename);
-				KeyChainUtilities::LoadKeyChainFromFile(CryptoKeysCacheFilename, TargetKeyChain);
-			}
-			
-			return Diff(SourcePath, SourceKeyChain, TargetPath, TargetKeyChain, OutPath);
-		}
-		else if (FParse::Param(FCommandLine::Get(), TEXT("Staged2Zen")))
-		{
-			FString BuildPath;
-			if (!FParse::Value(FCommandLine::Get(), TEXT("BuildPath="), BuildPath) || !Arguments.TargetPlatform)
-			{
-				UE_LOG(LogIoStore, Error, TEXT("Incorrect arguments. Expected: -Staged2Zen -BuildPath=<Path> -TargetPlatform=<Platform>"));
-			}
-			return Staged2Zen(BuildPath, Arguments.KeyChain, Arguments.TargetPlatform);
-		}
-		else
-		{
-			UE_LOG(LogIoStore, Display, TEXT("Usage:"));
-			UE_LOG(LogIoStore, Display, TEXT(" -List=</path/to/[container.utoc|*.utoc]> -CSV=<list.csv> [-CryptoKeys=</path/to/crypto.json>]"));
-			UE_LOG(LogIoStore, Display, TEXT(" -Describe=</path/to/global.utoc> [-PackageFilter=<PackageName>] [-DumpToFile=<describe.txt>] [-CryptoKeys=</path/to/crypto.json>]"));
-			return -1;
-		}
+	{
+		IOSTORE_CPU_SCOPE(FindCookedAssets);
+		UE_LOG(LogIoStore, Display, TEXT("Searching for cooked assets in folder '%s'"), *Arguments.CookedDir);
+		FCookedFileVisitor CookedFileVistor(Arguments.CookedFileStatMap, nullptr, Arguments.bFileRegions);
+		IFileManager::Get().IterateDirectoryStatRecursively(*Arguments.CookedDir, CookedFileVistor);
+		UE_LOG(LogIoStore, Display, TEXT("Found '%d' files"), Arguments.CookedFileStatMap.Num());
 	}
 
-	return 0;
+	return CreateTarget(Arguments, WriterSettings);
 }
 
 #endif
diff --git a/Engine/Source/Developer/IoStoreUtilities/Private/PackageStoreManifest.cpp b/Engine/Source/Developer/IoStoreUtilities/Private/PackageStoreManifest.cpp
--- a/Engine/Source/Developer/IoStoreUtilities/Private/PackageStoreManifest.cpp
+++ b/Engine/Source/Developer/IoStoreUtilities/Private/PackageStoreManifest.cpp
@@ -7,6 +7,13 @@
 #include "Dom/JsonObject.h"
 #include "Misc/FileHelper.h"
 #include "Misc/ScopeLock.h"
+#include "Misc/Paths.h"
+
+FPackageStoreManifest::FPackageStoreManifest(const FString& InCookedOutputPath)
+	: CookedOutputPath(InCookedOutputPath)
+{
+	FPaths::NormalizeFilename(CookedOutputPath);
+}
 
 void FPackageStoreManifest::BeginPackage(FName PackageName)
 {
@@ -78,7 +85,9 @@ FIoStatus FPackageStoreManifest::Save(const TCHAR* Filename) const
 	for (const auto& KV : FileNameByChunkIdMap)
 	{
 		Writer->WriteObjectStart();
-		Writer->WriteValue(TEXT("Path"), KV.Value);
+		FString RelativePath = KV.Value;
+		FPaths::MakePathRelativeTo(RelativePath, *CookedOutputPath);
+		Writer->WriteValue(TEXT("Path"), RelativePath);
 		Writer->WriteValue(TEXT("ChunkId"), ChunkIdToString(KV.Key));
 		Writer->WriteObjectEnd();
 	}
@@ -162,7 +171,8 @@ FIoStatus FPackageStoreManifest::Load(const TCHAR* Filename)
 	{
 		TSharedPtr<FJsonObject> FileObject = FileValue->AsObject();
 		FIoChunkId ChunkId = ChunkIdFromString(FileObject->Values.FindRef(TEXT("ChunkId"))->AsString());
-		FileNameByChunkIdMap.Add(ChunkId, FileObject->Values.FindRef(TEXT("Path"))->AsString());
+		FString RelativePath = FileObject->Values.FindRef(TEXT("Path"))->AsString();
+		FileNameByChunkIdMap.Add(ChunkId, FPaths::Combine(CookedOutputPath, RelativePath));
 	}
 
 	TArray<TSharedPtr<FJsonValue>> PackagesArray = JsonObject->Values.FindRef(TEXT("Packages"))->AsArray();
diff --git a/Engine/Source/Developer/IoStoreUtilities/Private/PackageStoreOptimizer.cpp b/Engine/Source/Developer/IoStoreUtilities/Private/PackageStoreOptimizer.cpp
--- a/Engine/Source/Developer/IoStoreUtilities/Private/PackageStoreOptimizer.cpp
+++ b/Engine/Source/Developer/IoStoreUtilities/Private/PackageStoreOptimizer.cpp
@@ -14,6 +14,7 @@
 #include "Serialization/MemoryReader.h"
 #include "UObject/Package.h"
 #include "Misc/SecureHash.h"
+#include "Serialization/LargeMemoryReader.h"
 
 DEFINE_LOG_CATEGORY_STATIC(LogPackageStoreOptimizer, Log, All);
 
@@ -102,6 +103,11 @@ void FPackageStoreOptimizer::Initialize(const ITargetPlatform* TargetPlatform)
 	FindScriptObjects(TargetPlatform);
 }
 
+void FPackageStoreOptimizer::Initialize(const FIoBuffer& ScriptObjectsBuffer)
+{
+	LoadScriptObjectsBuffer(ScriptObjectsBuffer);
+}
+
 FPackageStorePackage* FPackageStoreOptimizer::CreateMissingPackage(const FName& Name) const
 {
 	FPackageStorePackage* Package = new FPackageStorePackage();
@@ -397,6 +403,11 @@ void FPackageStoreOptimizer::ProcessImports(const FCookedHeaderData& CookedHeade
 		FPackageStorePackage::FUnresolvedImport& UnresolvedImport = UnresolvedImports[ImportIndex];
 		if (UnresolvedImport.bIsScriptImport)
 		{
+			FPackageObjectIndex ScriptObjectIndex = FPackageObjectIndex::FromScriptPath(UnresolvedImport.FullName);
+			if (!ScriptObjectsMap.Contains(ScriptObjectIndex))
+			{
+				UE_LOG(LogPackageStoreOptimizer, Warning, TEXT("Package '%s' is referencing missing script import '%s'"), *Package->Name.ToString(), *UnresolvedImport.FullName);
+			}
 			Package->Imports[ImportIndex] = FPackageObjectIndex::FromScriptPath(UnresolvedImport.FullName);
 		}
 		else if (!UnresolvedImport.bIsImportOfPackage)
@@ -1682,6 +1693,25 @@ FIoBuffer FPackageStoreOptimizer::CreateScriptObjectsBuffer() const
 	return FIoBuffer(FIoBuffer::AssumeOwnership, ScriptObjectsArchive.ReleaseOwnership(), DataSize);
 }
 
+void FPackageStoreOptimizer::LoadScriptObjectsBuffer(const FIoBuffer& ScriptObjectsBuffer)
+{
+	FLargeMemoryReader ScriptObjectsArchive(ScriptObjectsBuffer.Data(), ScriptObjectsBuffer.DataSize());
+	TArray<FNameEntryId> NameMap = LoadNameBatch(ScriptObjectsArchive);
+	int32 NumScriptObjects;
+	ScriptObjectsArchive << NumScriptObjects;
+	for (int32 Index = 0; Index < NumScriptObjects; ++Index)
+	{
+		FScriptObjectEntry Entry;
+		ScriptObjectsArchive << Entry;
+		FScriptObjectData& ImportData = ScriptObjectsMap.Add(Entry.GlobalIndex);
+		const FMappedName& MappedName = FMappedName::FromMinimalName(Entry.ObjectName);
+		ImportData.ObjectName = FName::CreateFromDisplayId(NameMap[MappedName.GetIndex()], MappedName.GetNumber());
+		ImportData.GlobalIndex = Entry.GlobalIndex;
+		ImportData.OuterIndex = Entry.OuterIndex;
+		ImportData.CDOClassIndex = Entry.CDOClassIndex;
+	}
+}
+
 FPackageStoreEntryResource FPackageStoreOptimizer::CreatePackageStoreEntry(const FPackageStorePackage* Package) const
 {
 	FPackageStoreEntryResource Result;
diff --git a/Engine/Source/Developer/IoStoreUtilities/Private/PackageStoreOptimizer.h b/Engine/Source/Developer/IoStoreUtilities/Private/PackageStoreOptimizer.h
deleted file mode 100644
diff --git a/Engine/Source/Developer/IoStoreUtilities/Private/ZenStoreHttpClient.cpp b/Engine/Source/Developer/IoStoreUtilities/Private/ZenStoreHttpClient.cpp
--- a/Engine/Source/Developer/IoStoreUtilities/Private/ZenStoreHttpClient.cpp
+++ b/Engine/Source/Developer/IoStoreUtilities/Private/ZenStoreHttpClient.cpp
@@ -229,7 +229,12 @@ TFuture<TIoStatusOr<uint64>> FZenStoreHttpClient::AppendOp(FCbPackage OpEntry)
 
 	TRACE_CPUPROFILER_EVENT_SCOPE(ZenStoreHttp_AppendOp);
 
-	return Async(EAsyncExecution::LargeThreadPool, [this, OpEntry]
+#if WITH_EDITOR
+	EAsyncExecution ThreadPool = EAsyncExecution::LargeThreadPool;
+#else
+	EAsyncExecution ThreadPool = EAsyncExecution::ThreadPool;
+#endif
+	return Async(ThreadPool, [this, OpEntry]
 	{
 		TRACE_CPUPROFILER_EVENT_SCOPE(Zen_AppendOp_Async);
 		FLargeMemoryWriter SerializedPackage;
@@ -463,7 +468,12 @@ TIoStatusOr<FIoBuffer> FZenStoreHttpClient::ReadOpLogUri(FStringBuilderBase& Chu
 
 TFuture<TIoStatusOr<FCbObject>> FZenStoreHttpClient::GetOplog()
 {
-	return Async(EAsyncExecution::LargeThreadPool, [this]
+#if WITH_EDITOR
+	EAsyncExecution ThreadPool = EAsyncExecution::LargeThreadPool;
+#else
+	EAsyncExecution ThreadPool = EAsyncExecution::ThreadPool;
+#endif
+	return Async(ThreadPool, [this]
 	{
 		UE::Zen::FZenScopedRequestPtr Request(RequestPool.Get());
 		
@@ -487,7 +497,12 @@ TFuture<TIoStatusOr<FCbObject>> FZenStoreHttpClient::GetOplog()
 
 TFuture<TIoStatusOr<FCbObject>> FZenStoreHttpClient::GetFiles()
 {
-	return Async(EAsyncExecution::LargeThreadPool, [this]
+#if WITH_EDITOR
+	EAsyncExecution ThreadPool = EAsyncExecution::LargeThreadPool;
+#else
+	EAsyncExecution ThreadPool = EAsyncExecution::ThreadPool;
+#endif
+	return Async(ThreadPool, [this]
 	{
 		UE::Zen::FZenScopedRequestPtr Request(RequestPool.Get());
 		
diff --git a/Engine/Source/Developer/IoStoreUtilities/Private/ZenStoreWriter.cpp b/Engine/Source/Developer/IoStoreUtilities/Private/ZenStoreWriter.cpp
--- a/Engine/Source/Developer/IoStoreUtilities/Private/ZenStoreWriter.cpp
+++ b/Engine/Source/Developer/IoStoreUtilities/Private/ZenStoreWriter.cpp
@@ -251,6 +251,7 @@ FZenStoreWriter::FZenStoreWriter(
 	: TargetPlatform(*InTargetPlatform)
 	, OutputPath(InOutputPath)
 	, MetadataDirectoryPath(InMetadataDirectoryPath)
+	, PackageStoreManifest(InOutputPath)
 	, PackageStoreOptimizer(new FPackageStoreOptimizer())
 	, CookMode(ICookedPackageWriter::FCookInfo::CookByTheBookMode)
 	, bInitialized(false)
diff --git a/Engine/Source/Developer/IoStoreUtilities/Public/PackageStoreManifest.h b/Engine/Source/Developer/IoStoreUtilities/Public/PackageStoreManifest.h
--- a/Engine/Source/Developer/IoStoreUtilities/Public/PackageStoreManifest.h
+++ b/Engine/Source/Developer/IoStoreUtilities/Public/PackageStoreManifest.h
@@ -28,7 +28,7 @@ public:
 		FString OplogId;
 	};
 
-	IOSTOREUTILITIES_API FPackageStoreManifest() = default;
+	IOSTOREUTILITIES_API FPackageStoreManifest(const FString& CookedOutputPath);
 	IOSTOREUTILITIES_API ~FPackageStoreManifest() = default;
 
 	IOSTOREUTILITIES_API void BeginPackage(FName PackageName);
@@ -46,6 +46,7 @@ public:
 
 private:
 	mutable FCriticalSection CriticalSection;
+	FString CookedOutputPath;
 	TMap<FName, FPackageInfo> PackageInfoByNameMap;
 	TMap<FIoChunkId, FString> FileNameByChunkIdMap;
 	TUniquePtr<FZenServerInfo> ZenServerInfo;
diff --git a/Engine/Source/Developer/PakFileUtilities/PakFileUtilities.Build.cs b/Engine/Source/Developer/PakFileUtilities/PakFileUtilities.Build.cs
--- a/Engine/Source/Developer/PakFileUtilities/PakFileUtilities.Build.cs
+++ b/Engine/Source/Developer/PakFileUtilities/PakFileUtilities.Build.cs
@@ -12,6 +12,7 @@ public class PakFileUtilities : ModuleRules
             "Json",
             "Projects",
             "RSA",
+            "IoStoreUtilities",
         });
 
         PrivateIncludePathModuleNames.AddRange(new string[] {
diff --git a/Engine/Source/Developer/PakFileUtilities/Private/PakFileUtilities.cpp b/Engine/Source/Developer/PakFileUtilities/Private/PakFileUtilities.cpp
--- a/Engine/Source/Developer/PakFileUtilities/Private/PakFileUtilities.cpp
+++ b/Engine/Source/Developer/PakFileUtilities/Private/PakFileUtilities.cpp
@@ -27,6 +27,7 @@
 #include "Serialization/FileRegions.h"
 #include "Misc/ICompressionFormat.h"
 #include "Misc/KeyChainUtilities.h"
+#include "IoStoreUtilities.h"
 
 IMPLEMENT_MODULE(FDefaultModuleImpl, PakFileUtilities);
 
@@ -4762,6 +4763,13 @@ void CheckAndReallocThreadPool()
  */
 bool ExecuteUnrealPak(const TCHAR* CmdLine)
 {
+	FString IoStoreArg;
+	if (FParse::Value(CmdLine, TEXT("-CreateGlobalContainer="), IoStoreArg) ||
+		FParse::Value(CmdLine, TEXT("-CreateDLCContainer="), IoStoreArg))
+	{
+		return CreateIoStoreContainerFiles(CmdLine) == 0;
+	}
+
 	// Parse all the non-option arguments from the command line
 	TArray<FString> NonOptionArguments;
 	for (const TCHAR* CmdLineEnd = CmdLine; *CmdLineEnd != 0;)
diff --git a/Engine/Source/Editor/UnrealEd/Private/CookOnTheFlyServer.cpp b/Engine/Source/Editor/UnrealEd/Private/CookOnTheFlyServer.cpp
--- a/Engine/Source/Editor/UnrealEd/Private/CookOnTheFlyServer.cpp
+++ b/Engine/Source/Editor/UnrealEd/Private/CookOnTheFlyServer.cpp
@@ -4464,7 +4464,7 @@ void FSaveCookedPackageContext::SetupPlatform(const ITargetPlatform* InTargetPla
 	PackageWriter = CookContext->PackageWriter;
 	ICookedPackageWriter::FBeginPackageInfo Info;
 	Info.PackageName = Package->GetFName();
-
+	Info.LooseFilePath = PlatFilename;
 	PackageWriter->BeginPackage(Info);
 
 	// Indicate Setup was successful
@@ -7469,32 +7469,30 @@ UE::Cook::FCookSavePackageContext* UCookOnTheFlyServer::CreateSaveContext(const
 		const_cast<UCookOnTheFlyServer&>(*this).CreateSandboxFile();
 	}
 
-	const FString RootPath = FPaths::RootDir();
-	const FString RootPathSandbox = ConvertToFullSandboxPath(*RootPath, true);
-
-	const FString ProjectPath = FPaths::ProjectDir();
-	const FString ProjectPathSandbox = ConvertToFullSandboxPath(*ProjectPath, true);
-
-	const bool bIsDiffOnly = FParse::Param(FCommandLine::Get(), TEXT("DIFFONLY"));
-	const bool bIsCleanBuild = !IsCookFlagSet(ECookInitializationFlags::Iterative);
-
+	const FString RootPathSandbox = ConvertToFullSandboxPath(FPaths::RootDir(), true);
+	FString MetadataPathSandbox;
+	if (IsCookingDLC())
+	{
+		MetadataPathSandbox = ConvertToFullSandboxPath(GetBaseDirectoryForDLC() / "Metadata", true);
+	}
+	else
+	{
+		MetadataPathSandbox = ConvertToFullSandboxPath(FPaths::ProjectDir() / "Metadata", true);
+	}
 	const FString PlatformString = TargetPlatform->PlatformName();
-
 	const FString ResolvedRootPath = RootPathSandbox.Replace(TEXT("[Platform]"), *PlatformString);
-	const FString ResolvedProjectPath = ProjectPathSandbox.Replace(TEXT("[Platform]"), *PlatformString);
+	const FString ResolvedMetadataPath = MetadataPathSandbox.Replace(TEXT("[Platform]"), *PlatformString);
 
 	ICookedPackageWriter* PackageWriter = nullptr;
-	FString MetadataDirectoryPath = ProjectPathSandbox / TEXT("Metadata");
-	MetadataDirectoryPath.ReplaceInline(TEXT("[Platform]"), *PlatformString);
 	FString WriterDebugName;
 	if (IsUsingZenStore())
 	{
-		PackageWriter = new FZenStoreWriter(ResolvedRootPath, MetadataDirectoryPath, TargetPlatform);
+		PackageWriter = new FZenStoreWriter(ResolvedRootPath, ResolvedMetadataPath, TargetPlatform);
 		WriterDebugName = TEXT("ZenStore");
 	}
 	else
 	{
-		PackageWriter = new FLooseCookedPackageWriter(ResolvedRootPath, MetadataDirectoryPath, TargetPlatform,
+		PackageWriter = new FLooseCookedPackageWriter(ResolvedRootPath, ResolvedMetadataPath, TargetPlatform,
 			const_cast<UCookOnTheFlyServer&>(*this).GetAsyncIODelete(), GetPackageNameCache(), PluginsToRemap);
 		WriterDebugName = TEXT("DirectoryWriter");
 	}
diff --git a/Engine/Source/Editor/UnrealEd/Private/Cooker/LooseCookedPackageWriter.cpp b/Engine/Source/Editor/UnrealEd/Private/Cooker/LooseCookedPackageWriter.cpp
--- a/Engine/Source/Editor/UnrealEd/Private/Cooker/LooseCookedPackageWriter.cpp
+++ b/Engine/Source/Editor/UnrealEd/Private/Cooker/LooseCookedPackageWriter.cpp
@@ -22,6 +22,7 @@
 #include "ProfilingDebugging/CpuProfilerTrace.h"
 #include "Serialization/ArrayReader.h"
 #include "UObject/Package.h"
+#include "PackageStoreOptimizer.h"
 
 const static FName NAME_DummyCookedFilename(TEXT("DummyCookedFilename"));
 
@@ -31,6 +32,7 @@ FLooseCookedPackageWriter::FLooseCookedPackageWriter(const FString& InOutputPath
 	, MetadataDirectoryPath(InMetadataDirectoryPath)
 	, TargetPlatform(*InTargetPlatform)
 	, PackageNameCache(InPackageNameCache)
+	, PackageStoreManifest(InOutputPath)
 	, PluginsToRemap(InPluginsToRemap)
 	, AsyncIODelete(InAsyncIODelete)
 	, bIterateSharedBuild(false)
@@ -43,6 +45,10 @@ FLooseCookedPackageWriter::~FLooseCookedPackageWriter()
 
 void FLooseCookedPackageWriter::BeginPackage(const FBeginPackageInfo& Info)
 {
+	PackageStoreManifest.BeginPackage(Info.PackageName);
+	// Temp solution until WritePackageData gets called also for LooseCookedPackageWriter
+	FIoChunkId ChunkId = CreateIoChunkId(FPackageId::FromName(Info.PackageName).Value(), 0, EIoChunkType::ExportBundleData);
+	PackageStoreManifest.AddPackageData(Info.PackageName, Info.LooseFilePath, ChunkId);
 }
 
 void FLooseCookedPackageWriter::CommitPackage(const FCommitPackageInfo& Info)
@@ -87,10 +93,20 @@ void FLooseCookedPackageWriter::BeginCook(const FCookInfo& Info)
 	{
 		DeleteSandboxDirectory();
 	}
+	{
+		TRACE_CPUPROFILER_EVENT_SCOPE(SaveScriptObjects);
+		FPackageStoreOptimizer PackageStoreOptimizer;
+		PackageStoreOptimizer.Initialize(&TargetPlatform);
+		FIoBuffer ScriptObjectsBuffer = PackageStoreOptimizer.CreateScriptObjectsBuffer();
+		FFileHelper::SaveArrayToFile(
+			MakeArrayView(ScriptObjectsBuffer.Data(), ScriptObjectsBuffer.DataSize()),
+			*(MetadataDirectoryPath / TEXT("scriptobjects.bin")));
+	}
 }
 
 void FLooseCookedPackageWriter::EndCook()
 {
+	PackageStoreManifest.Save(*(MetadataDirectoryPath / TEXT("packagestore.manifest")));
 }
 
 void FLooseCookedPackageWriter::Flush()
diff --git a/Engine/Source/Editor/UnrealEd/Private/Cooker/LooseCookedPackageWriter.h b/Engine/Source/Editor/UnrealEd/Private/Cooker/LooseCookedPackageWriter.h
--- a/Engine/Source/Editor/UnrealEd/Private/Cooker/LooseCookedPackageWriter.h
+++ b/Engine/Source/Editor/UnrealEd/Private/Cooker/LooseCookedPackageWriter.h
@@ -6,6 +6,7 @@
 #include "Serialization/PackageWriter.h"
 #include "Templates/SharedPointer.h"
 #include "Templates/UniquePtr.h"
+#include "PackageStoreManifest.h"
 
 class FAsyncIODelete;
 class IPlugin;
@@ -71,6 +72,7 @@ private:
 	FString MetadataDirectoryPath;
 	const ITargetPlatform& TargetPlatform;
 	const FPackageNameCache& PackageNameCache;
+	FPackageStoreManifest PackageStoreManifest;
 	TUniquePtr<FAssetRegistryState> PreviousState;
 	const TArray<TSharedRef<IPlugin>>& PluginsToRemap;
 	FAsyncIODelete& AsyncIODelete;
diff --git a/Engine/Source/Programs/AutomationTool/AutomationUtils/DeploymentContext.cs b/Engine/Source/Programs/AutomationTool/AutomationUtils/DeploymentContext.cs
--- a/Engine/Source/Programs/AutomationTool/AutomationUtils/DeploymentContext.cs
+++ b/Engine/Source/Programs/AutomationTool/AutomationUtils/DeploymentContext.cs
@@ -94,6 +94,26 @@ public class FilesToStage
 	}
 }
 
+public class PackageStoreManifest
+{
+	public class ZenServerInfo
+	{
+		public string HostName { get; set; }
+		public int Port { get; set; }
+		public string ProjectId { get; set; }
+		public string OplogId { get; set; }
+	}
+	public class FileInfo
+	{
+		public string Path { get; set; }
+	}
+
+	public string FullPath { get; set; }
+
+	public ZenServerInfo ZenServer { get; set; }
+	public IList<FileInfo> Files { get; set; }
+}
+
 public class DeploymentContext //: ProjectParams
 {
 	/// <summary>
@@ -226,6 +246,11 @@ public class DeploymentContext //: ProjectParams
 	/// </summary>
 	public string PakFileInternalRoot = "../../../";
 
+	/// <summary>
+	/// Cooked package store manifest if available
+	/// </summary>
+	public PackageStoreManifest PackageStoreManifest;
+
 	/// <summary>
 	/// List of files to be staged
 	/// </summary>
diff --git a/Engine/Source/Programs/AutomationTool/Scripts/CopyBuildToStagingDirectory.Automation.cs b/Engine/Source/Programs/AutomationTool/Scripts/CopyBuildToStagingDirectory.Automation.cs
--- a/Engine/Source/Programs/AutomationTool/Scripts/CopyBuildToStagingDirectory.Automation.cs
+++ b/Engine/Source/Programs/AutomationTool/Scripts/CopyBuildToStagingDirectory.Automation.cs
@@ -604,27 +604,26 @@ namespace AutomationScripts
 			return Plugins.ToList();
 		}
 
-		private class PackageStoreManifest
+		private static void LoadPackageStoreManifest(ProjectParams Params, DeploymentContext SC)
 		{
-			public static IList<FileReference> GetCookedFiles(FileReference PackageStoreManifestFile)
+			if (!Params.Stage)
 			{
-				List<FileReference> CookedFiles = new List<FileReference>();
-				if (FileReference.Exists(PackageStoreManifestFile))
-				{
-					String JsonString = FileReference.ReadAllText(PackageStoreManifestFile);
-					PackageStoreManifest PackageStoreManifest = JsonSerializer.Deserialize<PackageStoreManifest>(JsonString);
-					foreach (PackageStoreManifest.FileInfo File in PackageStoreManifest.Files)
-					{
-						CookedFiles.Add(new FileReference(File.Path));
-					}
-				}
-				return CookedFiles;
+				return;
 			}
-			public IList<FileInfo> Files { get; set; }
-
-			public class FileInfo
+			if (!Params.ZenStore && !Params.IoStore)
+			{
+				return;
+			}
+			FileReference PackageStoreManifestFile = FileReference.Combine(SC.MetadataDir, "packagestore.manifest");
+			if (FileReference.Exists(PackageStoreManifestFile))
 			{
-				public String Path { get; set; }
+				String JsonString = FileReference.ReadAllText(PackageStoreManifestFile);
+				SC.PackageStoreManifest = JsonSerializer.Deserialize<PackageStoreManifest>(JsonString);
+				SC.PackageStoreManifest.FullPath = PackageStoreManifestFile.FullName;
+				foreach (PackageStoreManifest.FileInfo FileInfo in SC.PackageStoreManifest.Files)
+				{
+					FileInfo.Path = FileReference.Combine(SC.PlatformCookDir, FileInfo.Path).FullName;
+				}
 			}
 		}
 
@@ -724,10 +723,14 @@ namespace AutomationScripts
 				// Stage all the cooked data, this is the same rule as normal stage except we may skip Engine
 				List<FileReference> CookedFiles = DirectoryReference.EnumerateFiles(SC.PlatformCookDir, "*", SearchOption.AllDirectories).ToList();
 
-				// When cooking directly to I/O store contaier files look for the package store manifest
-				if (Params.IoStore || Params.ZenStore)
+				// When cooking to Zen get the list of cooked package files from the manifest
+				LoadPackageStoreManifest(Params, SC);
+				if (SC.PackageStoreManifest != null && SC.PackageStoreManifest.ZenServer != null)
 				{
-					CookedFiles.AddRange(PackageStoreManifest.GetCookedFiles(FileReference.Combine(SC.MetadataDir, "packagestore.manifest")));
+					foreach (PackageStoreManifest.FileInfo File in SC.PackageStoreManifest.Files)
+					{
+						CookedFiles.Add(new FileReference(File.Path));
+					}
 				}
 
 				foreach (FileReference CookedFile in CookedFiles)
@@ -1055,10 +1058,14 @@ namespace AutomationScripts
 					{
 						List<FileReference> CookedFiles = DirectoryReference.EnumerateFiles(CookOutputDir, "*", SearchOption.AllDirectories).ToList();
 
-						// When cooking directly to I/O store contaier files look for the package store manifest
-						if (Params.IoStore || Params.ZenStore)
+						// When cooking to Zen get the list of cooked package files from the manifest
+						LoadPackageStoreManifest(Params, SC);
+						if (SC.PackageStoreManifest != null && SC.PackageStoreManifest.ZenServer != null)
 						{
-							CookedFiles.AddRange(PackageStoreManifest.GetCookedFiles(FileReference.Combine(SC.MetadataDir, "packagestore.manifest")));
+							foreach (PackageStoreManifest.FileInfo File in SC.PackageStoreManifest.Files)
+							{
+								CookedFiles.Add(new FileReference(File.Path));
+							}
 						}
 
 						foreach (FileReference CookedFile in CookedFiles)
@@ -2816,6 +2823,12 @@ namespace AutomationScripts
 						Dictionary<string, string> UnrealPakResponseFile = PakParams.UnrealPakResponseFile;
 						if (ShouldCreateIoStoreContainerFiles(Params, SC.StageTargetPlatform, SC.CustomConfig))
 						{
+							bool bLegacyBulkDataOffsets = false;
+							if (PlatformEngineConfig.GetBool("Core.System", "LegacyBulkDataOffsets", out bLegacyBulkDataOffsets) && bLegacyBulkDataOffsets)
+							{
+								throw new AutomationException("'LegacyBulkDataOffsets' is enabled, this needs to be disabled and the data recooked in order for IoStore to work");
+							}
+
 							bool bAllowBulkDataInIoStore = true;
 							if(!PlatformEngineConfig.GetBool("Core.System", "AllowBulkDataInIoStore", out bAllowBulkDataInIoStore))
 							{
@@ -2972,6 +2985,42 @@ namespace AutomationScripts
 					}
 				}
 
+				int CompressionMinBytesSaved = 0;
+				if (PlatformGameConfig.GetInt32("/Script/UnrealEd.ProjectPackagingSettings", "PackageCompressionMinBytesSaved", out CompressionMinBytesSaved))
+				{
+					if (CompressionMinBytesSaved > 0)
+					{
+						AdditionalArgs += String.Format(" -compressionMinBytesSaved={0}", CompressionMinBytesSaved);
+					}
+				}
+
+				int CompressionMinPercentSaved = 0;
+				if (PlatformGameConfig.GetInt32("/Script/UnrealEd.ProjectPackagingSettings", "PackageCompressionMinPercentSaved", out CompressionMinPercentSaved))
+				{
+					if (CompressionMinBytesSaved > 0)
+					{
+						AdditionalArgs += String.Format(" -compressionMinPercentSaved={0}", CompressionMinPercentSaved);
+					}
+				}
+
+				bool CompresionEnableDDC = false;
+				if (PlatformGameConfig.GetBool("/Script/UnrealEd.ProjectPackagingSettings", "bPackageCompressionEnableDDC", out CompresionEnableDDC))
+				{
+					if (CompresionEnableDDC)
+					{
+						AdditionalArgs += " -compressionEnableDDC";
+					}
+				}
+
+				int CompressionMinSizeToConsiderDDC = 0;
+				if (PlatformGameConfig.GetInt32("/Script/UnrealEd.ProjectPackagingSettings", "PackageCompressionMinSizeToConsiderDDC", out CompressionMinSizeToConsiderDDC))
+				{
+					if (CompressionMinSizeToConsiderDDC > 0)
+					{
+						AdditionalArgs += String.Format(" -compressionMinSizeToConsiderDDC={0}", CompressionMinSizeToConsiderDDC);
+					}
+				}
+
 				AdditionalArgs += " " + Params.AdditionalIoStoreOptions;
 
 				RunIoStore(Params, SC, IoStoreCommandsFileName, FinalOrderFiles, AdditionalArgs);
@@ -3149,7 +3198,7 @@ namespace AutomationScripts
 
 			if (Params.HasDLCName)
 			{
-				CommandletParams += String.Format("-DLCFile={0}", MakePathSafeToUseWithCommandLine(Params.DLCFile.FullName));
+				CommandletParams += String.Format("-CreateDLCContainer={0}", MakePathSafeToUseWithCommandLine(Params.DLCFile.FullName));
 
 				DirectoryReference DLCRoot = Params.DLCFile.Directory;
 				string DLCName = Params.DLCFile.GetFileNameWithoutExtension();
@@ -3184,7 +3233,24 @@ namespace AutomationScripts
 				CommandletParams += String.Format("-CreateGlobalContainer={0}", MakePathSafeToUseWithCommandLine(GlobalContainerOutputLocation.FullName));
 			}
 
-			CommandletParams += String.Format(" -CookedDirectory={0} -Commands={1}", MakePathSafeToUseWithCommandLine(SC.PlatformCookDir.ToString()), MakePathSafeToUseWithCommandLine(CommandsFileName));
+			if (SC.PackageStoreManifest == null)
+			{
+				throw new AutomationException("A package store manifest is required when staging to IoStore");
+			}
+			CommandletParams += String.Format(" -CookedDirectory={0} -PackageStoreManifest={1} -Commands={2}",
+				MakePathSafeToUseWithCommandLine(SC.PlatformCookDir.ToString()),
+				MakePathSafeToUseWithCommandLine(SC.PackageStoreManifest.FullPath),
+				MakePathSafeToUseWithCommandLine(CommandsFileName));
+			if (SC.PackageStoreManifest.ZenServer == null)
+			{
+				FileReference ScriptObjectsFile = FileReference.Combine(SC.MetadataDir, "scriptobjects.bin");
+				if (!FileReference.Exists(ScriptObjectsFile))
+				{
+					throw new AutomationException("scriptobjects.bin is required when staging to IoStore");
+				}
+				CommandletParams += String.Format(" -ScriptObjects={0}", MakePathSafeToUseWithCommandLine(ScriptObjectsFile.FullName));
+			}
+
 			if (OrderFiles != null && OrderFiles.Count() > 0)
 			{
 				CommandletParams += String.Format(" -Order={0}", MakePathSafeToUseWithCommandLine(string.Join(",", OrderFiles.Select(u => u.File.FullName).ToArray())));
@@ -3195,15 +3261,13 @@ namespace AutomationScripts
 				CommandletParams += AdditionalArgs;
 			}
 
-			CommandletParams += String.Format(" -TargetPlatform={0}", SC.StageTargetPlatform.GetCookPlatform(Params.DedicatedServer, Params.Client));
-
 			if (Params.HasBasedOnReleaseVersion)
 			{
 				CommandletParams += String.Format(" -BasedOnReleaseVersionPath={0}", Params.GetBasedOnReleaseVersionPath(SC, Params.Client));
 			}
 
-			LogInformation("Running IoStore commandlet with arguments: {0}", CommandletParams);
-			RunCommandlet(SC.RawProjectPath, Params.UE4Exe, "IoStore", CommandletParams);
+			LogInformation("Running UnrealPak with arguments: {0}", CommandletParams);
+			RunAndLog(CmdEnv, GetUnrealPakLocation().FullName, CommandletParams, Options: ERunOptions.Default | ERunOptions.UTF8Output);
 		}
 
 		private static void RunUnrealPakInParallel(List<string> Commands, List<string> LogNames, string AdditionalCompressionOptionsOnCommandLine)
diff --git a/Engine/Source/Programs/UnrealPak/UnrealPak.Build.cs b/Engine/Source/Programs/UnrealPak/UnrealPak.Build.cs
--- a/Engine/Source/Programs/UnrealPak/UnrealPak.Build.cs
+++ b/Engine/Source/Programs/UnrealPak/UnrealPak.Build.cs
@@ -8,7 +8,7 @@ public class UnrealPak : ModuleRules
 	{
 		PublicIncludePaths.Add("Runtime/Launch/Public");
 
-		PrivateDependencyModuleNames.AddRange(new string[] { "Core", "PakFile", "Json", "Projects", "PakFileUtilities", "RSA" });
+		PrivateDependencyModuleNames.AddRange(new string[] { "Core", "PakFile", "Json", "Projects", "PakFileUtilities", "RSA", "ApplicationCore" });
 
 		PrivateIncludePaths.Add("Runtime/Launch/Private");      // For LaunchEngineLoop.cpp include
 
diff --git a/Engine/Source/Programs/UnrealPak/UnrealPak.Target.cs b/Engine/Source/Programs/UnrealPak/UnrealPak.Target.cs
--- a/Engine/Source/Programs/UnrealPak/UnrealPak.Target.cs
+++ b/Engine/Source/Programs/UnrealPak/UnrealPak.Target.cs
@@ -27,7 +27,7 @@ public class UnrealPakTarget : TargetRules
 		bCompileICU = false;
 
 		// UnrealPak is a console application, not a Windows app (sets entry point to main(), instead of WinMain())
-		bCompileAgainstApplicationCore = false;
+		bCompileAgainstApplicationCore = true;
 		bIsBuildingConsoleApplication = true;
 	}
 }
diff --git a/Engine/Source/Runtime/Core/Private/IO/IoStore.cpp b/Engine/Source/Runtime/Core/Private/IO/IoStore.cpp
--- a/Engine/Source/Runtime/Core/Private/IO/IoStore.cpp
+++ b/Engine/Source/Runtime/Core/Private/IO/IoStore.cpp
@@ -2593,13 +2593,3 @@ void FIoStoreReader::GetFilenamesByBlockIndex(const TArray<int32>& InBlockIndexL
 			return true;
 		});
 }
-
-void FIoStoreWriterSettings::InitializePlatformSpecificSettings(const ITargetPlatform* TargetPlatform)
-{
-	FConfigCacheIni* PlatformConfig = TargetPlatform->GetConfigSystem();
-	const TCHAR* IniSection = TEXT("/Script/UnrealEd.ProjectPackagingSettings");
-	PlatformConfig->GetInt(IniSection, TEXT("PackageCompressionMinBytesSaved"), CompressionMinBytesSaved, GGameIni);
-	PlatformConfig->GetInt(IniSection, TEXT("PackageCompressionMinPercentSaved"), CompressionMinPercentSaved, GGameIni);
-	PlatformConfig->GetInt(IniSection, TEXT("PackageCompressionMinSizeToConsiderDDC"), CompressionMinSizeToConsiderDDC, GGameIni);
-	PlatformConfig->GetBool(IniSection, TEXT("bPackageCompressionEnableDDC"), bCompressionEnableDDC, GGameIni);
-}
diff --git a/Engine/Source/Runtime/Core/Public/IO/IoDispatcher.h b/Engine/Source/Runtime/Core/Public/IO/IoDispatcher.h
--- a/Engine/Source/Runtime/Core/Public/IO/IoDispatcher.h
+++ b/Engine/Source/Runtime/Core/Public/IO/IoDispatcher.h
@@ -1072,8 +1072,6 @@ struct FIoStoreWriterSettings
 	bool bEnableCsvOutput = false;
 	bool bEnableFileRegions = false;
 	bool bCompressionEnableDDC = false;
-
-	CORE_API void InitializePlatformSpecificSettings(const ITargetPlatform* TargetPlatform);
 };
 
 enum class EIoContainerFlags : uint8
diff --git a/Engine/Source/Runtime/Core/Public/Serialization/PackageWriter.h b/Engine/Source/Runtime/Core/Public/Serialization/PackageWriter.h
--- a/Engine/Source/Runtime/Core/Public/Serialization/PackageWriter.h
+++ b/Engine/Source/Runtime/Core/Public/Serialization/PackageWriter.h
@@ -45,6 +45,7 @@ public:
 	struct FBeginPackageInfo
 	{
 		FName	PackageName;
+		FString	LooseFilePath;
 	};
 
 	/** Mark the beginning of a package store transaction for the specified package
diff --git a/Engine/Source/Runtime/RenderCore/Private/ShaderCodeArchive.cpp b/Engine/Source/Runtime/RenderCore/Private/ShaderCodeArchive.cpp
--- a/Engine/Source/Runtime/RenderCore/Private/ShaderCodeArchive.cpp
+++ b/Engine/Source/Runtime/RenderCore/Private/ShaderCodeArchive.cpp
@@ -15,8 +15,6 @@
 #include "Templates/Greater.h"
 #endif
 
-uint32 GIoStoreShaderCodeArchiveVersion = 1;
-
 int32 GShaderCodeLibraryAsyncLoadingPriority = int32(AIOP_Normal);
 static FAutoConsoleVariableRef CVarShaderCodeLibraryAsyncLoadingPriority(
 	TEXT("r.ShaderCodeLibrary.DefaultAsyncIOPriority"),
@@ -33,24 +31,6 @@ static FAutoConsoleVariableRef CVarShaderCodeLibraryAsyncLoadingAllowDontCache(
 	ECVF_Default
 );
 
-FIoChunkId GetShaderCodeArchiveChunkId(const FString& LibraryName, FName FormatName)
-{
-	FString Name = FString::Printf(TEXT("%s-%s"), *LibraryName, *FormatName.ToString());
-	Name.ToLowerInline();
-	uint64 Hash = CityHash64(reinterpret_cast<const char*>(*Name), Name.Len() * sizeof(TCHAR));
-	return CreateIoChunkId(Hash, 0, EIoChunkType::ShaderCodeLibrary);
-}
-
-FIoChunkId GetShaderCodeChunkId(const FSHAHash& ShaderHash)
-{
-	uint8 Data[12];
-	FMemory::Memcpy(Data, ShaderHash.Hash, 11);
-	*reinterpret_cast<uint8*>(&Data[11]) = static_cast<uint8>(EIoChunkType::ShaderCode);
-	FIoChunkId ChunkId;
-	ChunkId.Set(Data, 12);
-	return ChunkId;
-}
-
 int32 FSerializedShaderArchive::FindShaderMapWithKey(const FSHAHash& Hash, uint32 Key) const
 {
 	for (uint32 Index = ShaderMapHashTable.First(Key); ShaderMapHashTable.IsValid(Index); Index = ShaderMapHashTable.Next(Index))
@@ -1021,6 +1001,52 @@ TRefCountPtr<FRHIShader> FShaderCodeArchive::CreateShader(int32 Index)
 	return Shader;
 }
 
+FIoChunkId FIoStoreShaderCodeArchive::GetShaderCodeArchiveChunkId(const FString& LibraryName, FName FormatName)
+{
+	FString Name = FString::Printf(TEXT("%s-%s"), *LibraryName, *FormatName.ToString());
+	Name.ToLowerInline();
+	uint64 Hash = CityHash64(reinterpret_cast<const char*>(*Name), Name.Len() * sizeof(TCHAR));
+	return CreateIoChunkId(Hash, 0, EIoChunkType::ShaderCodeLibrary);
+}
+
+FIoChunkId FIoStoreShaderCodeArchive::GetShaderCodeChunkId(const FSHAHash& ShaderHash)
+{
+	uint8 Data[12];
+	FMemory::Memcpy(Data, ShaderHash.Hash, 11);
+	*reinterpret_cast<uint8*>(&Data[11]) = static_cast<uint8>(EIoChunkType::ShaderCode);
+	FIoChunkId ChunkId;
+	ChunkId.Set(Data, 12);
+	return ChunkId;
+}
+
+void FIoStoreShaderCodeArchive::SaveIoStoreShaderCodeArchive(const FSerializedShaderArchive& SerializedShaders, FArchive& OutLibraryAr)
+{
+	uint32 Version = CurrentVersion;
+	OutLibraryAr << Version;
+	OutLibraryAr << const_cast<FSerializedShaderArchive&>(SerializedShaders).ShaderMapHashes;
+	OutLibraryAr << const_cast<FSerializedShaderArchive&>(SerializedShaders).ShaderHashes;
+	TArray<FIoStoreShaderMapEntry> ShaderMapEntries;
+	ShaderMapEntries.Reserve(SerializedShaders.ShaderMapEntries.Num());
+	for (const FShaderMapEntry& ShaderMapEntry : SerializedShaders.ShaderMapEntries)
+	{
+		FIoStoreShaderMapEntry& IoStoreShaderMapEntry = ShaderMapEntries.AddDefaulted_GetRef();
+		IoStoreShaderMapEntry.ShaderIndicesOffset = ShaderMapEntry.ShaderIndicesOffset;
+		IoStoreShaderMapEntry.NumShaders = ShaderMapEntry.NumShaders;
+	}
+	OutLibraryAr << ShaderMapEntries;
+	TArray<FIoStoreShaderCodeEntry> ShaderEntries;
+	ShaderEntries.Reserve(SerializedShaders.ShaderEntries.Num());
+	for (const FShaderCodeEntry& ShaderEntry : SerializedShaders.ShaderEntries)
+	{
+		FIoStoreShaderCodeEntry& IoStoreShaderEntry = ShaderEntries.AddDefaulted_GetRef();
+		IoStoreShaderEntry.UncompressedSize = ShaderEntry.UncompressedSize;
+		IoStoreShaderEntry.CompressedSize = ShaderEntry.Size;
+		IoStoreShaderEntry.Frequency = ShaderEntry.Frequency;
+	}
+	OutLibraryAr << ShaderEntries;
+	OutLibraryAr << const_cast<FSerializedShaderArchive&>(SerializedShaders).ShaderIndices;
+}
+
 FIoStoreShaderCodeArchive* FIoStoreShaderCodeArchive::Create(EShaderPlatform InPlatform, const FString& InLibraryName, FIoDispatcher& InIoDispatcher)
 {
 	const FName PlatformName = LegacyShaderPlatformToShaderFormat(InPlatform);
@@ -1037,7 +1063,7 @@ FIoStoreShaderCodeArchive* FIoStoreShaderCodeArchive::Create(EShaderPlatform InP
 		FMemoryReaderView Ar(MakeArrayView(IoBuffer.Data(), IoBuffer.DataSize()));
 		uint32 Version = 0;
 		Ar << Version;
-		if (Version == GIoStoreShaderCodeArchiveVersion)
+		if (Version == CurrentVersion)
 		{
 			FIoStoreShaderCodeArchive* Library = new FIoStoreShaderCodeArchive(InPlatform, InLibraryName, InIoDispatcher);
 			Ar << Library->ShaderMapHashes;
diff --git a/Engine/Source/Runtime/RenderCore/Private/ShaderCodeLibrary.cpp b/Engine/Source/Runtime/RenderCore/Private/ShaderCodeLibrary.cpp
--- a/Engine/Source/Runtime/RenderCore/Private/ShaderCodeLibrary.cpp
+++ b/Engine/Source/Runtime/RenderCore/Private/ShaderCodeLibrary.cpp
@@ -61,7 +61,6 @@ DEFINE_LOG_CATEGORY(LogShaderLibrary);
 
 static uint32 GShaderCodeArchiveVersion = 2;
 static uint32 GShaderPipelineArchiveVersion = 1;
-extern uint32 GIoStoreShaderCodeArchiveVersion;
 
 static FString ShaderExtension = TEXT(".ushaderbytecode");
 static FString ShaderAssetInfoExtension = TEXT(".assetinfo.json");
@@ -1361,33 +1360,6 @@ struct FEditorShaderCodeArchive
 		}
 	}
 
-	static void SaveIoStoreShaderCodeArchive(const FSerializedShaderArchive& SerializedShaders, FArchive& OutLibraryAr)
-	{
-		OutLibraryAr << GIoStoreShaderCodeArchiveVersion;
-		OutLibraryAr << const_cast<FSerializedShaderArchive&>(SerializedShaders).ShaderMapHashes;
-		OutLibraryAr << const_cast<FSerializedShaderArchive&>(SerializedShaders).ShaderHashes;
-		TArray<FIoStoreShaderMapEntry> ShaderMapEntries;
-		ShaderMapEntries.Reserve(SerializedShaders.ShaderMapEntries.Num());
-		for (const FShaderMapEntry& ShaderMapEntry : SerializedShaders.ShaderMapEntries)
-		{
-			FIoStoreShaderMapEntry& IoStoreShaderMapEntry = ShaderMapEntries.AddDefaulted_GetRef();
-			IoStoreShaderMapEntry.ShaderIndicesOffset = ShaderMapEntry.ShaderIndicesOffset;
-			IoStoreShaderMapEntry.NumShaders = ShaderMapEntry.NumShaders;
-		}
-		OutLibraryAr << ShaderMapEntries;
-		TArray<FIoStoreShaderCodeEntry> ShaderEntries;
-		ShaderEntries.Reserve(SerializedShaders.ShaderEntries.Num());
-		for (const FShaderCodeEntry& ShaderEntry : SerializedShaders.ShaderEntries)
-		{
-			FIoStoreShaderCodeEntry& IoStoreShaderEntry = ShaderEntries.AddDefaulted_GetRef();
-			IoStoreShaderEntry.UncompressedSize = ShaderEntry.UncompressedSize;
-			IoStoreShaderEntry.CompressedSize = ShaderEntry.Size;
-			IoStoreShaderEntry.Frequency = ShaderEntry.Frequency;
-		}
-		OutLibraryAr << ShaderEntries;
-		OutLibraryAr << const_cast<FSerializedShaderArchive&>(SerializedShaders).ShaderIndices;
-	}
-
 	bool Finalize(FString OutputDir, const FString& MetaOutputDir, bool bSaveOnlyAssetInfo = false, TArray<FString>* OutputFilenames = nullptr)
 	{
 		check(LibraryName.Len() > 0);
@@ -3054,90 +3026,6 @@ bool FShaderLibraryCooker::SaveShaderLibraryChunk(int32 ChunkId, const TSet<FNam
 	return bResult;
 }
 
-bool FShaderLibraryCooker::ConvertToIoStoreShaderLibrary(
-	const TCHAR* FileName,
-	TTuple<FIoChunkId, FIoBuffer>& OutLibraryIoChunk,
-	TArray<TTuple<FIoChunkId, FIoBuffer, FSHAHash>>& OutCodeIoChunks,
-	TArray<TTuple<FSHAHash, TArray<FIoChunkId>>>& OutShaderMaps,
-	TArray<TTuple<FSHAHash, TSet<FName>>>& OutShaderMapAssetAssociations)
-{
-	TArray<FString> Components;
-	if (FPaths::GetBaseFilename(FileName).ParseIntoArray(Components, TEXT("-")) != 3)
-	{
-		UE_LOG(LogShaderLibrary, Error, TEXT("Invalid shader code library file name '%s'."), FileName);
-		return false;
-	}
-	FString LibraryName = Components[1];
-	FName FormatName(Components[2]);
-	
-	TUniquePtr<FArchive> LibraryAr(IFileManager::Get().CreateFileReader(FileName));
-	if (!LibraryAr)
-	{
-		UE_LOG(LogShaderLibrary, Error, TEXT("Missing shader code library file '%s'."), FileName);
-		return false;
-	}
-
-	uint32 Version;
-	(*LibraryAr) << Version;
-	if (Version != GShaderCodeArchiveVersion)
-	{
-		UE_LOG(LogShaderLibrary, Error, TEXT("Invalid shader code library chunk version in file '%s'."), FileName);
-		return false;
-	}
-
-	FSerializedShaderArchive SerializedShaders;
-	(*LibraryAr) << SerializedShaders;
-	int64 OffsetToShaderCode = LibraryAr->Tell();
-
-	FString AssetInfoFileName = GetShaderAssetInfoFilename(FPaths::GetPath(FileName), LibraryName, FormatName);
-	if (!SerializedShaders.LoadAssetInfo(AssetInfoFileName))
-	{
-		UE_LOG(LogShaderLibrary, Error, TEXT("Failed loading asset asset info file '%s'"), *AssetInfoFileName);
-		return false;
-	}
-
-	FLargeMemoryWriter IoStoreLibraryAr(0, true);
-	FEditorShaderCodeArchive::SaveIoStoreShaderCodeArchive(SerializedShaders, IoStoreLibraryAr);
-	OutLibraryIoChunk.Key = GetShaderCodeArchiveChunkId(LibraryName, FormatName);
-	int64 TotalSize = IoStoreLibraryAr.TotalSize();
-	OutLibraryIoChunk.Value = FIoBuffer(FIoBuffer::AssumeOwnership, IoStoreLibraryAr.ReleaseOwnership(), TotalSize);
-
-	int32 ShaderCount = SerializedShaders.ShaderEntries.Num();
-	OutCodeIoChunks.Reserve(ShaderCount);
-	for (int32 ShaderIndex = 0; ShaderIndex < ShaderCount; ++ShaderIndex)
-	{
-		const FSHAHash& ShaderHash = SerializedShaders.ShaderHashes[ShaderIndex];
-		const FShaderCodeEntry& ShaderEntry = SerializedShaders.ShaderEntries[ShaderIndex];
-		TTuple<FIoChunkId, FIoBuffer, FSHAHash>& OutCodeIoChunk = OutCodeIoChunks.Emplace_GetRef(GetShaderCodeChunkId(ShaderHash), ShaderEntry.Size, ShaderHash);
-		LibraryAr->Seek(OffsetToShaderCode + ShaderEntry.Offset);
-		LibraryAr->Serialize(OutCodeIoChunk.Get<1>().Data(), ShaderEntry.Size);
-	}
-
-	int32 ShaderMapCount = SerializedShaders.ShaderMapHashes.Num();
-	OutShaderMaps.Reserve(ShaderMapCount);
-	for (int32 ShaderMapIndex = 0; ShaderMapIndex < ShaderMapCount; ++ShaderMapIndex)
-	{
-		TTuple<FSHAHash, TArray<FIoChunkId>>& OutShaderMap = OutShaderMaps.AddDefaulted_GetRef();
-		OutShaderMap.Key = SerializedShaders.ShaderMapHashes[ShaderMapIndex];
-		const FShaderMapEntry& ShaderMapEntry = SerializedShaders.ShaderMapEntries[ShaderMapIndex];
-		int32 LookupIndexEnd = ShaderMapEntry.ShaderIndicesOffset + ShaderMapEntry.NumShaders;
-		OutShaderMap.Value.Reserve(ShaderMapEntry.NumShaders);
-		for (int32 ShaderLookupIndex = ShaderMapEntry.ShaderIndicesOffset; ShaderLookupIndex < LookupIndexEnd; ++ShaderLookupIndex)
-		{
-			int32 ShaderIndex = SerializedShaders.ShaderIndices[ShaderLookupIndex];
-			OutShaderMap.Value.Add(GetShaderCodeChunkId(SerializedShaders.ShaderHashes[ShaderIndex]));
-		}
-	}
-
-	OutShaderMapAssetAssociations.Reserve(SerializedShaders.ShaderCodeToAssets.Num());
-	for (auto& KV : SerializedShaders.ShaderCodeToAssets)
-	{
-		OutShaderMapAssetAssociations.Emplace(KV.Key, MoveTemp(KV.Value));
-	}
-
-	return true;
-}
-
 #endif// WITH_EDITOR
 
 void FShaderCodeLibrary::SafeAssignHash(FRHIShader* InShader, const FSHAHash& Hash)
diff --git a/Engine/Source/Runtime/RenderCore/Public/ShaderCodeArchive.h b/Engine/Source/Runtime/RenderCore/Public/ShaderCodeArchive.h
--- a/Engine/Source/Runtime/RenderCore/Public/ShaderCodeArchive.h
+++ b/Engine/Source/Runtime/RenderCore/Public/ShaderCodeArchive.h
@@ -9,9 +9,6 @@
 #include "Shader.h"
 #include "IO/IoDispatcher.h"
 
-FIoChunkId GetShaderCodeArchiveChunkId(const FString& LibraryName, FName FormatName);
-FIoChunkId GetShaderCodeChunkId(const FSHAHash& ShaderHash);
-
 struct FShaderMapEntry
 {
 	uint32 ShaderIndicesOffset = 0u;
@@ -297,6 +294,9 @@ struct FIoStoreShaderCodeEntry
 class FIoStoreShaderCodeArchive : public FRHIShaderLibrary
 {
 public:
+	RENDERCORE_API static FIoChunkId GetShaderCodeArchiveChunkId(const FString& LibraryName, FName FormatName);
+	RENDERCORE_API static FIoChunkId GetShaderCodeChunkId(const FSHAHash& ShaderHash);
+	RENDERCORE_API static void SaveIoStoreShaderCodeArchive(const FSerializedShaderArchive& SerializedShaders, FArchive& OutLibraryAr);
 	static FIoStoreShaderCodeArchive* Create(EShaderPlatform InPlatform, const FString& InLibraryName, FIoDispatcher& InIoDispatcher);
 
 	virtual ~FIoStoreShaderCodeArchive();
@@ -334,6 +334,8 @@ public:
 	virtual void Teardown() override;
 
 private:
+	static constexpr uint32 CurrentVersion = 1;
+
 	struct FShaderPreloadEntry
 	{
 		FGraphEventRef PreloadEvent;
diff --git a/Engine/Source/Runtime/RenderCore/Public/ShaderCodeLibrary.h b/Engine/Source/Runtime/RenderCore/Public/ShaderCodeLibrary.h
--- a/Engine/Source/Runtime/RenderCore/Public/ShaderCodeLibrary.h
+++ b/Engine/Source/Runtime/RenderCore/Public/ShaderCodeLibrary.h
@@ -304,12 +304,5 @@ struct RENDERCORE_API FShaderLibraryCooker
 
 	// Create a smaller 'patch' library that only contains data from 'NewMetaDataDir' not contained in any of 'OldMetaDataDirs'
 	static bool CreatePatchLibrary(TArray<FString> const& OldMetaDataDirs, FString const& NewMetaDataDir, FString const& OutDir, bool bNativeFormat, bool bNeedsDeterministicOrder);
-
-	static bool ConvertToIoStoreShaderLibrary(
-		const TCHAR* FileName,
-		TTuple<FIoChunkId, FIoBuffer>& OutLibraryIoChunk,
-		TArray<TTuple<FIoChunkId, FIoBuffer, FSHAHash>>& OutCodeIoChunks,
-		TArray<TTuple<FSHAHash, TArray<FIoChunkId>>>& OutShaderMaps,
-		TArray<TTuple<FSHAHash, TSet<FName>>>& OutShaderMapAssetAssociations);
 };
 #endif
