Trivial logging fix: 'packges' -> 'packages'

diff --git a/Engine/Source/Developer/IoStoreUtilities/Private/ZenStoreWriter.cpp b/Engine/Source/Developer/IoStoreUtilities/Private/ZenStoreWriter.cpp
--- a/Engine/Source/Developer/IoStoreUtilities/Private/ZenStoreWriter.cpp
+++ b/Engine/Source/Developer/IoStoreUtilities/Private/ZenStoreWriter.cpp
@@ -541,7 +541,7 @@ void FZenStoreWriter::BeginCook(const FCookInfo& Info)
 
 			if (FutureOplogStatus.Get().IsOk())
 			{
-				UE_LOG(LogZenStoreWriter, Display, TEXT("Fetched '%d' packges(s) from oplog '%s/%s'"), PackageStoreEntries.Num(), *ProjectId, *OplogId);
+				UE_LOG(LogZenStoreWriter, Display, TEXT("Fetched '%d' packages(s) from oplog '%s/%s'"), PackageStoreEntries.Num(), *ProjectId, *OplogId);
 			}
 			else
 			{
