ZenStoreHttpClient: Expose connection failure to the Oplog owner to handle.
#rnx

diff --git a/Engine/Source/Developer/IoStoreUtilities/Internal/ZenStoreHttpClient.h b/Engine/Source/Developer/IoStoreUtilities/Internal/ZenStoreHttpClient.h
--- a/Engine/Source/Developer/IoStoreUtilities/Internal/ZenStoreHttpClient.h
+++ b/Engine/Source/Developer/IoStoreUtilities/Internal/ZenStoreHttpClient.h
@@ -31,12 +31,9 @@ public:
 	FZenStoreHttpClient(const FStringView InHostName, uint16 InPort);
 	~FZenStoreHttpClient();
 
-	void Initialize(FStringView InProjectId, 
-					FStringView InOplogId, 
-					FStringView ServerRoot, 
-					FStringView EngineRoot, 
-					FStringView ProjectRoot);
-	void EstablishWritableOpLog(FStringView InProjectId, FStringView InOplogId, bool bFullBuild);
+	bool TryCreateProject(FStringView InProjectId, FStringView InOplogId, FStringView ServerRoot, 
+					FStringView EngineRoot, FStringView ProjectRoot);
+	bool TryCreateOplog(FStringView InProjectId, FStringView InOplogId, bool bFullBuild);
 
 	void InitializeReadOnly(FStringView InProjectId, FStringView InOplogId);
 
diff --git a/Engine/Source/Developer/IoStoreUtilities/Private/ZenStoreHttpClient.cpp b/Engine/Source/Developer/IoStoreUtilities/Private/ZenStoreHttpClient.cpp
--- a/Engine/Source/Developer/IoStoreUtilities/Private/ZenStoreHttpClient.cpp
+++ b/Engine/Source/Developer/IoStoreUtilities/Private/ZenStoreHttpClient.cpp
@@ -42,8 +42,8 @@ FZenStoreHttpClient::~FZenStoreHttpClient()
 {
 }
 
-void 
-FZenStoreHttpClient::Initialize(FStringView InProjectId, 
+bool
+FZenStoreHttpClient::TryCreateProject(FStringView InProjectId, 
 	FStringView InOplogId, 
 	FStringView ServerRoot,
 	FStringView EngineRoot,
@@ -111,11 +111,16 @@ FZenStoreHttpClient::Initialize(FStringView InProjectId,
 
 	bAllowRead = true;
 	bAllowEdit = true;
+
+	return bConnectionSucceeded;
 }
 
-void FZenStoreHttpClient::EstablishWritableOpLog(FStringView InProjectId, FStringView InOplogId, bool bFullBuild)
+bool FZenStoreHttpClient::TryCreateOplog(FStringView InProjectId, FStringView InOplogId, bool bFullBuild)
 {
-	check(IsConnected());
+	if (!IsConnected())
+	{
+		return false;
+	}
 
 	UE::Zen::FZenScopedRequestPtr Request(RequestPool.Get());
 
@@ -146,8 +151,9 @@ void FZenStoreHttpClient::EstablishWritableOpLog(FStringView InProjectId, FStrin
 		if (Res != Zen::FZenHttpRequest::Result::Success)
 		{
 			UE_LOG(LogZenStore, Error, TEXT("Zen oplog '%s'/'%s' creation FAILED"), *FString(InProjectId), *FString(InOplogId));
-
-			// TODO: how to recover / handle this?
+			// Demote the connection status back to not connected
+			bConnectionSucceeded = false;
+			return false;
 		}
 		else if (Request->GetResponseCode() == 201)
 		{
@@ -171,6 +177,7 @@ void FZenStoreHttpClient::EstablishWritableOpLog(FStringView InProjectId, FStrin
 	}
 
 	TempDirPath = FUTF8ToTCHAR(OplogInfo["tempdir"].AsString());
+	return true;
 }
 
 void FZenStoreHttpClient::InitializeReadOnly(FStringView InProjectId, FStringView InOplogId)
@@ -581,17 +588,15 @@ FZenStoreHttpClient::~FZenStoreHttpClient()
 {
 }
 
-void FZenStoreHttpClient::Initialize(
-	FStringView InProjectId,
-	FStringView InOplogId,
-	FStringView ServerRoot,
-	FStringView EngineRoot,
-	FStringView ProjectRoot)
+bool FZenStoreHttpClient::TryCreateProject(FStringView InProjectId, FStringView InOplogId, FStringView ServerRoot,
+	FStringView EngineRoot,	FStringView ProjectRoot)
 {
+	return false;
 }
 
-void FZenStoreHttpClient::EstablishWritableOpLog(FStringView InProjectId, FStringView InOplogId, bool bFullBuild)
+bool FZenStoreHttpClient::TryCreateOplog(FStringView InProjectId, FStringView InOplogId, bool bFullBuild)
 {
+	return false;
 }
 
 void FZenStoreHttpClient::InitializeReadOnly(FStringView InProjectId, FStringView InOplogId)
diff --git a/Engine/Source/Developer/IoStoreUtilities/Private/ZenStoreWriter.cpp b/Engine/Source/Developer/IoStoreUtilities/Private/ZenStoreWriter.cpp
--- a/Engine/Source/Developer/IoStoreUtilities/Private/ZenStoreWriter.cpp
+++ b/Engine/Source/Developer/IoStoreUtilities/Private/ZenStoreWriter.cpp
@@ -277,11 +277,7 @@ FZenStoreWriter::FZenStoreWriter(
 	FString AbsEngineDir = PlatformFile.ConvertToAbsolutePathForExternalAppForRead(*EngineDir);
 	FString AbsProjectDir = PlatformFile.ConvertToAbsolutePathForExternalAppForRead(*ProjectDir);
 
-	HttpClient->Initialize(	ProjectId, 
-							OplogId, 
-							AbsServerRoot, 
-							AbsEngineDir, 
-							AbsProjectDir);
+	HttpClient->TryCreateProject(ProjectId, OplogId, AbsServerRoot, AbsEngineDir, AbsProjectDir);
 
 	PackageStoreOptimizer->Initialize(InTargetPlatform);
 
@@ -444,11 +440,10 @@ void FZenStoreWriter::BeginCook(const FCookInfo& Info)
 
 	if (!bInitialized)
 	{
-		UE_CLOG(!HttpClient->IsConnected(), LogZenStoreWriter, Fatal, TEXT("Failed to connect to ZenServer"));
-
 		FString ProjectId = FApp::GetProjectName();
 		FString OplogId = TargetPlatform.PlatformName();
-		HttpClient->EstablishWritableOpLog(ProjectId, OplogId, Info.bFullBuild);
+		bool bOplogEstablished = HttpClient->TryCreateOplog(ProjectId, OplogId, Info.bFullBuild);
+		UE_CLOG(!bOplogEstablished, LogZenStoreWriter, Fatal, TEXT("Failed to establish oplog on the ZenServer"));
 
 		if (!Info.bFullBuild)
 		{
diff --git a/Engine/Source/Editor/UnrealEd/Private/TargetDomain/TargetDomainUtils.cpp b/Engine/Source/Editor/UnrealEd/Private/TargetDomain/TargetDomainUtils.cpp
--- a/Engine/Source/Editor/UnrealEd/Private/TargetDomain/TargetDomainUtils.cpp
+++ b/Engine/Source/Editor/UnrealEd/Private/TargetDomain/TargetDomainUtils.cpp
@@ -405,15 +405,8 @@ FEditorDomainOplog::FEditorDomainOplog()
 	FString AbsEngineDir = PlatformFile.ConvertToAbsolutePathForExternalAppForRead(*EngineDir);
 	FString AbsProjectDir = PlatformFile.ConvertToAbsolutePathForExternalAppForRead(*ProjectDir);
 
-	HttpClient.Initialize(ProjectId,
-		OplogId,
-		AbsServerRoot,
-		AbsEngineDir,
-		AbsProjectDir);
-	if (HttpClient.IsConnected())
-	{
-		HttpClient.EstablishWritableOpLog(ProjectId, OplogId, false /* bFullBuild */);
-	}
+	HttpClient.TryCreateProject(ProjectId, OplogId, AbsServerRoot, AbsEngineDir, AbsProjectDir);
+	HttpClient.TryCreateOplog(ProjectId, OplogId, false /* bFullBuild */);
 }
 
 void FEditorDomainOplog::InitializeRead()
