Modify 'List' command for investigating patch reordering issues.
Add LexToString for EIoChunkType.

diff --git a/Engine/Source/Developer/IoStoreUtilities/Private/IoStoreUtilities.cpp b/Engine/Source/Developer/IoStoreUtilities/Private/IoStoreUtilities.cpp
--- a/Engine/Source/Developer/IoStoreUtilities/Private/IoStoreUtilities.cpp
+++ b/Engine/Source/Developer/IoStoreUtilities/Private/IoStoreUtilities.cpp
@@ -49,7 +49,9 @@
 #include "Misc/StringBuilder.h"
 #include "Async/Future.h"
 #include "Algo/MaxElement.h"
+#include "Algo/Sort.h"
 #include "Algo/StableSort.h"
+#include "Algo/IsSorted.h"
 #include "PackageStoreOptimizer.h"
 #include "ShaderCodeArchive.h"
 #include "ZenStoreHttpClient.h"
@@ -2903,12 +2905,38 @@ int32 ListContainer(
 	const FString& ContainerPathOrWildcard,
 	const FString& CsvPath)
 {
+	struct FComputedChunkInfo
+	{
+		FIoChunkId Id;
+		FIoChunkHash Hash;
+		EIoChunkType ChunkType;
+		uint64 Size;
+		uint64 CompressedSize;
+		uint64 Offset;
+		uint64 OffsetOnDisk;
+		int32 NumCompressedBlocks;
+		bool bIsCompressed;
+	};
+
 	TArray<FString> ContainerFilePaths;
 
 	if (IFileManager::Get().FileExists(*ContainerPathOrWildcard))
 	{
 		ContainerFilePaths.Add(ContainerPathOrWildcard);
 	}
+	else if (IFileManager::Get().DirectoryExists(*ContainerPathOrWildcard))
+	{
+		FString Directory = ContainerPathOrWildcard;
+		FPaths::NormalizeDirectoryName(Directory);
+
+		TArray<FString> FoundContainerFiles;
+		IFileManager::Get().FindFiles(FoundContainerFiles, *(Directory / TEXT("*.utoc")), true, false);
+
+		for (const FString& Filename : FoundContainerFiles)
+		{
+			ContainerFilePaths.Emplace(Directory / Filename);
+		}
+	}
 	else
 	{
 		FString Directory = FPaths::GetPath(ContainerPathOrWildcard);
@@ -2931,7 +2959,10 @@ int32 ListContainer(
 
 	TArray<FString> CsvLines;
 	
-	CsvLines.Add(TEXT("PackageId, PackageName, Filename, ContainerName, Offset, Size, CompressedSize, Hash"));
+	TUniquePtr<FOutputDeviceFile> Out = MakeUnique<FOutputDeviceFile>(*CsvPath, true);
+	Out->SetSuppressEventTag(true);
+
+	Out->Log(TEXT("OrderInContainer, PackageId, PackageName, Filename, ContainerName, Offset, OffsetOnDisk, Size, CompressedSize, Hash, ChunkType"));
 
 	for (const FString& ContainerFilePath : ContainerFilePaths)
 	{
@@ -2966,41 +2997,73 @@ int32 ListContainer(
 			return true;
 		});
 
-		Reader->EnumerateChunks([&CsvLines, &ChunkFileNamesMap, &ContainerName](const FIoStoreTocChunkInfo& ChunkInfo)
+		uint64 CompressionBlockSize = Reader->GetCompressionBlockSize();
+		TArray<FIoStoreTocCompressedBlockInfo> CompressedBlocks;
+		Reader->EnumerateCompressedBlocks([&CompressedBlocks](const FIoStoreTocCompressedBlockInfo& Block) {
+			CompressedBlocks.Add(Block);
+			return true;
+			});
+
+
+		TArray<FComputedChunkInfo> Chunks;
+		Reader->EnumerateChunks([&Chunks, CompressionBlockSize, &CompressedBlocks](const FIoStoreTocChunkInfo& ChunkInfo) {
+
+			int32 FirstBlockIndex = int32(ChunkInfo.Offset / CompressionBlockSize);
+			int32 LastBlockIndex = int32((Align(ChunkInfo.Offset + ChunkInfo.Size, CompressionBlockSize) - 1) / CompressionBlockSize);
+			int32 NumCompressedBlocks = LastBlockIndex - FirstBlockIndex + 1;
+			uint64 OffsetOnDisk = CompressedBlocks[FirstBlockIndex].Offset;
+
+			FComputedChunkInfo ComputedInfo{
+				ChunkInfo.Id,
+				ChunkInfo.Hash,
+				ChunkInfo.ChunkType,
+				ChunkInfo.Size,
+				ChunkInfo.CompressedSize,
+				ChunkInfo.Offset,
+				OffsetOnDisk,
+				NumCompressedBlocks,
+				ChunkInfo.bIsCompressed
+			};
+
+			Chunks.Add(ComputedInfo);
+			return true;
+			});
+
+		auto SortKey = [](const FComputedChunkInfo& ChunkInfo) { return ChunkInfo.OffsetOnDisk; };
+		Algo::SortBy(Chunks, SortKey);
+
+		for(int32 Index=0; Index < Chunks.Num(); ++Index)
 		{
+			const FComputedChunkInfo& ChunkInfo = Chunks[Index];
 			FString PackageName;
+			FPackageId PackageId;
 			FString* FindFileName = ChunkFileNamesMap.Find(ChunkInfo.Id);
-			if (FindFileName)
+			if (FindFileName && FPackageName::TryConvertFilenameToLongPackageName(*FindFileName, PackageName, nullptr))
 			{
-				FPackageName::TryConvertFilenameToLongPackageName(*FindFileName, PackageName, nullptr);
+				PackageId = FPackageId::FromName(FName(*PackageName));
+			}
+			else
+			{
+				// For non-packages, print chunk id as a string
+				const uint32* IdParts = (const uint32*)&ChunkInfo.Id;
+				static_assert(sizeof(ChunkInfo.Id) == sizeof(uint32) * 3);
+				PackageName = FString::Printf(TEXT("%x%x%x"), IdParts[0], IdParts[1], IdParts[2]);
 			}
-			FPackageId PackageId = PackageName.Len() > 0
-				? FPackageId::FromName(FName(*PackageName))
-				: FPackageId();
 
-			CsvLines.Emplace(
-				FString::Printf(TEXT("0x%llX, %s, %s, %s, %lld, %lld, %lld, 0x%s"),
+			Out->Logf(TEXT("%d, 0x%llX, %s, %s, %s, %lld, %lld, %lld, %lld, 0x%s, %s"),
+					Index,
 					PackageId.ValueForDebugging(),
 					*PackageName,
 					FindFileName ? **FindFileName : TEXT(""),
 					*ContainerName,
 					ChunkInfo.Offset,
+					ChunkInfo.OffsetOnDisk,
 					ChunkInfo.Size,
 					ChunkInfo.CompressedSize,
-					*ChunkInfo.Hash.ToString()));
-
-			return true;
-		});
-	}
-
-	if (CsvLines.Num())
-	{
-		UE_LOG(LogIoStore, Display, TEXT("Saving '%d' file entries to '%s'"), CsvLines.Num(), *CsvPath);
-		FFileHelper::SaveStringArrayToFile(CsvLines, *CsvPath);
-	}
-	else
-	{
-		UE_LOG(LogIoStore, Warning, TEXT("No file entries to save from '%s'"), *ContainerPathOrWildcard);
+					*ChunkInfo.Hash.ToString(),
+					*LexToString(ChunkInfo.ChunkType)
+					);
+		}
 	}
 
 	return 0;
diff --git a/Engine/Source/Runtime/Core/Private/IO/IoStore.cpp b/Engine/Source/Runtime/Core/Private/IO/IoStore.cpp
--- a/Engine/Source/Runtime/Core/Private/IO/IoStore.cpp
+++ b/Engine/Source/Runtime/Core/Private/IO/IoStore.cpp
@@ -33,6 +33,35 @@ constexpr char FIoStoreTocHeader::TocMagicImg[];
 
 TRACE_DECLARE_INT_COUNTER(IoStoreAvailableCompressionBuffers, TEXT("IoStore/AvailableCompressionBuffers"));
 
+FString LexToString(const EIoChunkType Type)
+{
+	const TCHAR* Strings[] = {
+		TEXT("Invalid"),
+		TEXT("ExportBundleData"),
+		TEXT("BulkData"),
+		TEXT("OptionalBulkData"),
+		TEXT("MemoryMappedBulkData"),
+		TEXT("ScriptObjects"),
+		TEXT("ContainerHeader"),
+		TEXT("ExternalFile"),
+		TEXT("ShaderCodeLibrary"),
+		TEXT("ShaderCode"),
+		TEXT("PackageStoreEntry"),
+		TEXT("DerivedData"),
+		TEXT("EditorDerivedData"),
+	};
+	static_assert(UE_ARRAY_COUNT(Strings) == (SIZE_T)EIoChunkType::MAX);
+	uint8 Index = (uint8)Type;
+	if (Index < UE_ARRAY_COUNT(Strings))
+	{
+		return Strings[Index];
+	}
+	else
+	{
+		return Strings[0]; // return Invalid
+	}
+}
+
 template<typename ArrayType>
 bool WriteArray(IFileHandle* FileHandle, const ArrayType& Array)
 {
diff --git a/Engine/Source/Runtime/Core/Public/IO/IoDispatcher.h b/Engine/Source/Runtime/Core/Public/IO/IoDispatcher.h
--- a/Engine/Source/Runtime/Core/Public/IO/IoDispatcher.h
+++ b/Engine/Source/Runtime/Core/Public/IO/IoDispatcher.h
@@ -657,8 +657,12 @@ enum class EIoChunkType : uint8
 	PackageStoreEntry = 10,
 	DerivedData = 11,
 	EditorDerivedData = 12,
+	
+	MAX
 };
 
+CORE_API FString LexToString(const EIoChunkType Type);
+
 /**
  * Identifier to a chunk of data.
  */
