Add chunk id to output of iostore list command

diff --git a/Engine/Source/Developer/IoStoreUtilities/Private/IoStoreUtilities.cpp b/Engine/Source/Developer/IoStoreUtilities/Private/IoStoreUtilities.cpp
--- a/Engine/Source/Developer/IoStoreUtilities/Private/IoStoreUtilities.cpp
+++ b/Engine/Source/Developer/IoStoreUtilities/Private/IoStoreUtilities.cpp
@@ -2962,7 +2962,7 @@ int32 ListContainer(
 	TUniquePtr<FOutputDeviceFile> Out = MakeUnique<FOutputDeviceFile>(*CsvPath, true);
 	Out->SetSuppressEventTag(true);
 
-	Out->Log(TEXT("OrderInContainer, PackageId, PackageName, Filename, ContainerName, Offset, OffsetOnDisk, Size, CompressedSize, Hash, ChunkType"));
+	Out->Log(TEXT("OrderInContainer, ChunkId, PackageId, PackageName, Filename, ContainerName, Offset, OffsetOnDisk, Size, CompressedSize, Hash, ChunkType"));
 
 	for (const FString& ContainerFilePath : ContainerFilePaths)
 	{
@@ -3042,16 +3042,15 @@ int32 ListContainer(
 			{
 				PackageId = FPackageId::FromName(FName(*PackageName));
 			}
-			else
-			{
-				// For non-packages, print chunk id as a string
-				const uint32* IdParts = (const uint32*)&ChunkInfo.Id;
-				static_assert(sizeof(ChunkInfo.Id) == sizeof(uint32) * 3);
-				PackageName = FString::Printf(TEXT("%x%x%x"), IdParts[0], IdParts[1], IdParts[2]);
-			}
 
-			Out->Logf(TEXT("%d, 0x%llX, %s, %s, %s, %lld, %lld, %lld, %lld, 0x%s, %s"),
+
+			const uint32* IdParts = (const uint32*)&ChunkInfo.Id;
+			static_assert(sizeof(ChunkInfo.Id) == sizeof(uint32) * 3);
+			FString ChunkIdString = FString::Printf(TEXT("%x%x%x"), IdParts[0], IdParts[1], IdParts[2]);
+
+			Out->Logf(TEXT("%d, %s, 0x%llX, %s, %s, %s, %lld, %lld, %lld, %lld, 0x%s, %s"),
 					Index,
+					*ChunkIdString,
 					PackageId.ValueForDebugging(),
 					*PackageName,
 					FindFileName ? **FindFileName : TEXT(""),
