Separate the relative ordering of packages assigned by different order files from the priority at which we use those order files.

diff --git a/Engine/Source/Developer/DeveloperToolSettings/Classes/Settings/ProjectPackagingSettings.h b/Engine/Source/Developer/DeveloperToolSettings/Classes/Settings/ProjectPackagingSettings.h
--- a/Engine/Source/Developer/DeveloperToolSettings/Classes/Settings/ProjectPackagingSettings.h
+++ b/Engine/Source/Developer/DeveloperToolSettings/Classes/Settings/ProjectPackagingSettings.h
@@ -91,6 +91,29 @@ enum class EProjectPackagingBlueprintNativizationMethod : uint8
 	Exclusive
 };
 
+USTRUCT()
+struct FPakOrderFileSpec
+{
+	GENERATED_BODY()
+
+	UPROPERTY()
+	FString Pattern;
+
+	UPROPERTY()
+	int32 Priority;
+
+	FPakOrderFileSpec()
+		: Priority(0)
+	{
+	}
+
+	FPakOrderFileSpec(FString InPattern)
+		: Pattern(MoveTemp(InPattern))
+		, Priority(0)
+	{
+	}
+};
+
 USTRUCT()
 struct FProjectBuildSettings
 {
diff --git a/Engine/Source/Developer/IoStoreUtilities/Private/IoStoreUtilities.cpp b/Engine/Source/Developer/IoStoreUtilities/Private/IoStoreUtilities.cpp
--- a/Engine/Source/Developer/IoStoreUtilities/Private/IoStoreUtilities.cpp
+++ b/Engine/Source/Developer/IoStoreUtilities/Private/IoStoreUtilities.cpp
@@ -52,6 +52,7 @@
 #include "Misc/StringBuilder.h"
 #include "Async/Future.h"
 #include "Algo/MaxElement.h"
+#include "Algo/StableSort.h"
 #include "PackageStoreOptimizer.h"
 
 //PRAGMA_DISABLE_OPTIMIZATION
@@ -330,8 +331,16 @@ struct FContainerTargetFile
 
 struct FFileOrderMap
 {
-	TMap<FName, uint64> PackageNameToOrder;
+	TMap<FName, int64> PackageNameToOrder;
 	FString Name;
+	int32 Priority;
+	int32 Index;
+
+	FFileOrderMap(int32 InPriority, int32 InIndex)
+		: Priority(InPriority)
+		, Index(InIndex)
+	{
+	}
 };
 
 struct FIoStoreArguments
@@ -354,6 +363,7 @@ struct FIoStoreArguments
 	bool bSign = false;
 	bool bRemapPluginContentToGame = false;
 	bool bCreateDirectoryIndex = true;
+	bool bClusterByOrderFilePriority = false;
 
 	bool ShouldCreateContainers() const
 	{
@@ -431,14 +441,34 @@ static FIoChunkId CreateChunkId(FPackageId GlobalPackageId, uint16 ChunkIndex, E
 	return ChunkId;
 }
 
+// If bClusterByOrderFilePriority is false
+//		Order packages by the order of OrderMaps with their associated priority 
+//		e.g. Order map 1 with priority 0 (A, B, C) 
+//			 Order map 2 with priority 10 (B, D)
+//			Final order: (A, C, B, D)
+//		Then cluster packages in this order. 
+// If bClusterByOrderFilePriority is true
+//		Cluster packages first by OrderMaps in priority order, then concatenate clusters in the array order of OrderMaps.
+//		e.g. Order map 1 with priority 0 (A, B, C) 
+//			 Order map 2 with priority 10 (B, D)
+//			Cluster packages B, D, then A, C
+//			Then reassemble clusters A, C, B, D
 static void AssignPackagesDiskOrder(
 	const TArray<FLegacyCookedPackage*>& Packages,
 	const TArray<FFileOrderMap>& OrderMaps,
-		const FPackageIdMap& PackageIdMap)
+	const FPackageIdMap& PackageIdMap,
+	bool bClusterByOrderFilePriority
+	)
 {
 	struct FCluster
 	{
 		TArray<FLegacyCookedPackage*> Packages;
+		int32 OrderFileIndex; // Index in OrderMaps of the FFileOrderMap which contained Packages.Last()
+
+		FCluster(int32 InOrderFileIndex)
+			: OrderFileIndex(InOrderFileIndex)
+		{
+		}
 	};
 
 	TArray<FCluster*> Clusters;
@@ -448,20 +478,22 @@ static void AssignPackagesDiskOrder(
 	struct FPackageAndOrder
 	{
 		FLegacyCookedPackage* Package = nullptr;
-		uint64 Order = MAX_uint64;
-		const FFileOrderMap* BlameOrderMap = nullptr;
-
-		bool operator<(const FPackageAndOrder& Other) const
-		{
-			if (Order != Other.Order)
-			{
-				return Order < Other.Order;
-			}
-			// Fallback to reverse bundle order (so that packages are considered before their imports)
-			return Package->OptimizedPackage->GetLoadOrder() > Other.Package->OptimizedPackage->GetLoadOrder();
-		}
+		int64 LocalOrder = MAX_uint64;
+		const FFileOrderMap* OrderMap = nullptr;
 	};
 
+	// Order maps sorted by priority
+	TArray<const FFileOrderMap*> PriorityOrderMaps;
+	for (const FFileOrderMap& Map : OrderMaps)
+	{
+		PriorityOrderMaps.Add(&Map);
+	}
+	Algo::StableSortBy(PriorityOrderMaps, [](const FFileOrderMap* Map) { return Map->Priority; }, TGreater<int32>());
+
+	// Create a fallback order map to avoid null checks later
+	// Lowest priority, last index
+	FFileOrderMap FallbackOrderMap(MIN_int32, MAX_int32); 
+
 	TArray<FPackageAndOrder> SortedPackages;
 	SortedPackages.Reserve(Packages.Num());
 	for (FLegacyCookedPackage* Package : Packages)
@@ -472,35 +504,76 @@ static void AssignPackagesDiskOrder(
 		}
 		FPackageAndOrder& Entry = SortedPackages.AddDefaulted_GetRef();
 		Entry.Package = Package;
-		Entry.Order = MAX_uint64;
 
-		for (const FFileOrderMap& OrderMap : OrderMaps)
+		for (const FFileOrderMap* OrderMap : PriorityOrderMaps)
 		{
-			if (const uint64* Order = OrderMap.PackageNameToOrder.Find(Package->PackageName))
+			if (const int64* Order = OrderMap->PackageNameToOrder.Find(Package->PackageName))
 			{
-				Entry.Order = *Order;
-				Entry.BlameOrderMap = &OrderMap;
+				Entry.LocalOrder = *Order;
+				Entry.OrderMap = OrderMap;
 				break;
 			}
 		}
+
+		if (Entry.OrderMap == nullptr)
+		{
+			Entry.OrderMap = &FallbackOrderMap;
+			// Reverse the bundle load order for the fallback map (so that packages are considered before their imports)
+			Entry.LocalOrder = -int64(Package->OptimizedPackage->GetLoadOrder());
+		}
 	}
 	const FFileOrderMap* LastBlameOrderMap = nullptr;
 	int32 LastAssignedCount = 0;
-	Algo::Sort(SortedPackages);
+
+	if (bClusterByOrderFilePriority)
+	{
+		// Sort by priority of the order map
+		Algo::Sort(SortedPackages, [](const FPackageAndOrder& A, const FPackageAndOrder& B) {
+			// Packages in the same order map should be sorted by their local ordering
+			if (A.OrderMap == B.OrderMap)
+			{
+				return A.LocalOrder < B.LocalOrder;
+			}
+
+			// First priority, then index
+			if (A.OrderMap->Priority != B.OrderMap->Priority)
+			{
+				return A.OrderMap->Priority > B.OrderMap->Priority;
+			}
+
+			check(A.OrderMap->Index != B.OrderMap->Index);
+			return A.OrderMap->Index < B.OrderMap->Index;
+		});
+	}
+	else
+	{
+		// Sort by the order of the order map (...)
+		Algo::Sort(SortedPackages, [](const FPackageAndOrder& A, const FPackageAndOrder& B) {
+			// Packages in the same order map should be sorted by their local ordering
+			if (A.OrderMap == B.OrderMap)
+			{
+				return A.LocalOrder < B.LocalOrder;
+			}
+
+			// Blame order priority is not considered for the order in which we cluster packages, only for the order in which we assign packages to an order map
+			return A.OrderMap->Index < B.OrderMap->Index;
+		});
+	}
+
 	for (FPackageAndOrder& Entry : SortedPackages)
 	{
-		if (Entry.BlameOrderMap != LastBlameOrderMap)
+		if (Entry.OrderMap != LastBlameOrderMap)
 		{
 			if( LastBlameOrderMap != nullptr )
 			{
 				UE_LOG(LogIoStore, Display, TEXT("Ordered %d/%d packages using order file %s"), AssignedPackages.Num() - LastAssignedCount, Packages.Num(), *LastBlameOrderMap->Name);
 			}
 			LastAssignedCount = AssignedPackages.Num();
-			LastBlameOrderMap = Entry.BlameOrderMap;
+			LastBlameOrderMap = Entry.OrderMap;
 		}
 		if (!AssignedPackages.Contains(Entry.Package))
 		{
-			FCluster* Cluster = new FCluster();
+			FCluster* Cluster = new FCluster(Entry.OrderMap->Index);
 			Clusters.Add(Cluster);
 			ProcessStack.Push(Entry.Package);
 
@@ -533,6 +606,11 @@ static void AssignPackagesDiskOrder(
 
 	check(AssignedPackages.Num() == Packages.Num());
 	
+	if (bClusterByOrderFilePriority)
+	{
+		Algo::StableSortBy(Clusters, [](FCluster* Cluster) { return Cluster->OrderFileIndex; });
+	}
+	
 	for (FCluster* Cluster : Clusters)
 	{
 		Algo::Sort(Cluster->Packages, [](const FLegacyCookedPackage* A, const FLegacyCookedPackage* B)
@@ -556,11 +634,12 @@ static void CreateDiskLayout(
 	const TArray<FContainerTargetSpec*>& ContainerTargets,
 	const TArray<FLegacyCookedPackage*>& Packages,
 	const TArray<FFileOrderMap>& OrderMaps,
-	const FPackageIdMap& PackageIdMap)
+	const FPackageIdMap& PackageIdMap,
+	bool bClusterByOrderFilePriority)
 {
 	IOSTORE_CPU_SCOPE(CreateDiskLayout);
 
-	AssignPackagesDiskOrder(Packages, OrderMaps, PackageIdMap);
+	AssignPackagesDiskOrder(Packages, OrderMaps, PackageIdMap, bClusterByOrderFilePriority);
 
 	for (FContainerTargetSpec* ContainerTarget : ContainerTargets)
 	{
@@ -1574,7 +1653,7 @@ int32 CreateTarget(const FIoStoreArguments& Arguments, const FIoStoreWriterSetti
 	PackageStoreOptimizer.Flush(true);
 
 	UE_LOG(LogIoStore, Display, TEXT("Creating disk layout..."));
-	CreateDiskLayout(ContainerTargets, Packages, Arguments.OrderMaps, PackageIdMap);
+	CreateDiskLayout(ContainerTargets, Packages, Arguments.OrderMaps, PackageIdMap, Arguments.bClusterByOrderFilePriority);
 
 	for (FContainerTargetSpec* ContainerTarget : ContainerTargets)
 	{
@@ -3080,7 +3159,7 @@ static bool ParsePakResponseFile(const TCHAR* FilePath, TArray<FContainerSourceF
 	return true;
 }
 
-static bool ParsePakOrderFile(const TCHAR* FilePath, FFileOrderMap& Map, uint64& InOutNumEntries)
+static bool ParsePakOrderFile(const TCHAR* FilePath, FFileOrderMap& Map)
 {
 	TArray<FString> OrderFileContents;
 	if (!FFileHelper::LoadFileToStringArray(OrderFileContents, FilePath))
@@ -3090,8 +3169,8 @@ static bool ParsePakOrderFile(const TCHAR* FilePath, FFileOrderMap& Map, uint64&
 	}
 
 	Map.Name = FPaths::GetCleanFilename(FilePath);
-	uint64 LocalNumEntries = InOutNumEntries;
-	UE_LOG(LogIoStore, Display, TEXT("Order file %s (short name %s) starting at global index %llu"), FilePath, *Map.Name, LocalNumEntries);
+	UE_LOG(LogIoStore, Display, TEXT("Order file %s (short name %s) priority %d"), FilePath, *Map.Name, Map.Priority);
+	int64 NextOrder = 0;
 	for (const FString& OrderLine : OrderFileContents)
 	{
 		const TCHAR* OrderLinePtr = *OrderLine;
@@ -3117,12 +3196,11 @@ static bool ParsePakOrderFile(const TCHAR* FilePath, FFileOrderMap& Map, uint64&
 		FName PackageFName(MoveTemp(PackageName));
 		if (!Map.PackageNameToOrder.Contains(PackageFName))
 		{
-			Map.PackageNameToOrder.Emplace(PackageFName, LocalNumEntries++);
+			Map.PackageNameToOrder.Emplace(PackageFName, NextOrder++);
 		}
 	}
 
-	UE_LOG(LogIoStore, Display, TEXT("Order file %s (short name %s) contained %llu valid entries"), FilePath, *Map.Name, LocalNumEntries - InOutNumEntries);
-	InOutNumEntries = LocalNumEntries;
+	UE_LOG(LogIoStore, Display, TEXT("Order file %s (short name %s) contained %d valid entries"), FilePath, *Map.Name, Map.PackageNameToOrder.Num());
 	return true;
 }
 
@@ -3276,33 +3354,61 @@ int32 CreateIoStoreContainerFiles(const TCHAR* CmdLine)
 	}
 
 	uint64 OrderMapStartIndex = 0;
-	FString GameOrderFileStr;
-	if (FParse::Value(FCommandLine::Get(), TEXT("GameOrder="), GameOrderFileStr, false))
+	FString OrderFileStr;
+	if (FParse::Value(FCommandLine::Get(), TEXT("Order="), OrderFileStr, false))
 	{
-		TArray<FString> GameOrderFilePaths;
-		GameOrderFileStr.ParseIntoArray(GameOrderFilePaths, TEXT(","), true);
-		bool bMerge = false;
-		for (FString& PrimaryOrderFile : GameOrderFilePaths)
+		TArray<int32> OrderFilePriorities;
+		TArray<FString> OrderFilePaths;
+		OrderFileStr.ParseIntoArray(OrderFilePaths, TEXT(","), true);
+
+		FString LegacyParam;
+		if (FParse::Value(FCommandLine::Get(), TEXT("GameOrder="), LegacyParam, false))
 		{
-			FFileOrderMap OrderMap;
-			if (!ParsePakOrderFile(*PrimaryOrderFile, OrderMap, OrderMapStartIndex))
+			UE_LOG(LogIoStore, Warning, TEXT("-GameOrder= and -CookerOrder= are deprecated in favor of -Order"));
+			TArray<FString> LegacyPaths;
+			LegacyParam.ParseIntoArray(LegacyPaths, TEXT(","), true);
+			OrderFilePaths.Append(LegacyPaths);
+		}
+		if (FParse::Value(FCommandLine::Get(), TEXT("CookerOrder="), LegacyParam, false))
+		{
+			UE_LOG(LogIoStore, Warning, TEXT("-GameOrder= and -CookerOrder= are deprecated in favor of -Order"));
+			TArray<FString> LegacyPaths;
+			LegacyParam.ParseIntoArray(LegacyPaths, TEXT(","), true);
+			OrderFilePaths.Append(LegacyPaths);
+		}
+
+		FString OrderPriorityString;
+		if (FParse::Value(FCommandLine::Get(), TEXT("OrderPriority="), OrderPriorityString, false))
+		{
+			TArray<FString> PriorityStrings;
+			OrderPriorityString.ParseIntoArray(PriorityStrings, TEXT(","), true);
+			if (PriorityStrings.Num() != OrderFilePaths.Num())
 			{
+				UE_LOG(LogIoStore, Error, TEXT("Number of parameters to -Order= and -OrderPriority= do not match"));
 				return -1;
 			}
-			Arguments.OrderMaps.Add(OrderMap);
+
+			for (const FString& PriorityString : PriorityStrings)
+			{
+				int32 Priority = FCString::Atoi(*PriorityString);
+				OrderFilePriorities.Add(Priority);
+			}
+		}
+		else
+		{
+			OrderFilePriorities.AddZeroed(OrderFilePaths.Num());
 		}
-	}
 
-	FString CookerOrderFileStr;
-	if (FParse::Value(FCommandLine::Get(), TEXT("CookerOrder="), CookerOrderFileStr, false))
-	{
-		TArray<FString> CookerOrderFilePaths;
-		CookerOrderFileStr.ParseIntoArray(CookerOrderFilePaths, TEXT(","), true);
+		check(OrderFilePaths.Num() == OrderFilePriorities.Num());
+
 		bool bMerge = false;
-		for (FString& SecondOrderFile : CookerOrderFilePaths)
+		for (int32 i=0; i < OrderFilePaths.Num(); ++i)
 		{
-			FFileOrderMap OrderMap;
-			if (!ParsePakOrderFile(*SecondOrderFile, OrderMap, OrderMapStartIndex))
+			FString& OrderFile = OrderFilePaths[i];
+			int32 Priority = OrderFilePriorities[i];
+
+			FFileOrderMap OrderMap(Priority, i);
+			if (!ParsePakOrderFile(*OrderFile, OrderMap))
 			{
 				return -1;
 			}
@@ -3323,6 +3429,8 @@ int32 CreateIoStoreContainerFiles(const TCHAR* CmdLine)
 		}
 	}
 
+	Arguments.bClusterByOrderFilePriority = FParse::Param(CmdLine, TEXT("ClusterByOrderPriority"));
+
 	FIoStoreWriterSettings GeneralIoWriterSettings { DefaultCompressionMethod, DefaultCompressionBlockSize, false };
 	GeneralIoWriterSettings.bEnableCsvOutput = FParse::Param(CmdLine, TEXT("-csvoutput"));
 
diff --git a/Engine/Source/Developer/PakFileUtilities/Private/PakFileUtilities.cpp b/Engine/Source/Developer/PakFileUtilities/Private/PakFileUtilities.cpp
--- a/Engine/Source/Developer/PakFileUtilities/Private/PakFileUtilities.cpp
+++ b/Engine/Source/Developer/PakFileUtilities/Private/PakFileUtilities.cpp
@@ -27,7 +27,6 @@
 #include "Serialization/FileRegions.h"
 #include "Misc/ICompressionFormat.h"
 #include "Misc/KeyChainUtilities.h"
-#include "Algo/MaxElement.h"
 
 IMPLEMENT_MODULE(FDefaultModuleImpl, PakFileUtilities);
 
@@ -282,18 +281,18 @@ private:
 	int32 Index;
 };
 
-bool FPakOrderMap::ProcessOrderFile(const TCHAR* ResponseFile, bool bSecondaryOrderFile, bool bMergeOrder)
+bool FPakOrderMap::ProcessOrderFile(const TCHAR* ResponseFile, bool bSecondaryOrderFile, bool bMergeOrder, TOptional<uint64> InOffset)
 {
-	int32 OrderOffset = 0; 
+	uint64 OrderOffset = 0; 
 	int32 OpenOrderNumber = 0;
 
-	if (bSecondaryOrderFile || bMergeOrder)
+	if (InOffset.IsSet())
 	{
-		const auto* maxValue = Algo::MaxElementBy(OrderMap, [](const auto& data) { return data.Value; });
-		if (maxValue)
-		{
-			OrderOffset = maxValue->Value + 1;
-		}
+		OrderOffset = InOffset.GetValue();
+	}
+	else if (bSecondaryOrderFile || bMergeOrder)
+	{
+		OrderOffset = MaxIndex + 1;
 	}
 
 	if (bSecondaryOrderFile)
@@ -322,7 +321,7 @@ bool FPakOrderMap::ProcessOrderFile(const TCHAR* ResponseFile, bool bSecondaryOr
 
 			if (!FParse::Token(OrderLinePtr, Path, false))
 			{
-				UE_LOG(LogPakFile, Error, TEXT("Invlaid entry in the response file %s."), *Lines[EntryIndex]);
+				UE_LOG(LogPakFile, Error, TEXT("Invalid entry in the response file %s."), *Lines[EntryIndex]);
 				return false;
 			}
 
@@ -331,10 +330,25 @@ bool FPakOrderMap::ProcessOrderFile(const TCHAR* ResponseFile, bool bSecondaryOr
 				FString ReadNum = Lines[EntryIndex].RightChop(OpenOrderNumber + 1);
 				Lines[EntryIndex].LeftInline(OpenOrderNumber + 1, false);
 				ReadNum.TrimStartInline();
-				if (ReadNum.IsNumeric())
+				if(ReadNum.Len() == 0)
+				{
+					// If order files don't have explicit numbers just use the line number
+					OpenOrderNumber = EntryIndex;
+				}
+				else if (ReadNum.IsNumeric())
 				{
 					OpenOrderNumber = FCString::Atoi(*ReadNum);
 				}
+				else
+				{
+					UE_LOG(LogPakFile, Error, TEXT("Invalid entry in the response file %s, couldn't parse an order number after the path."), *Lines[EntryIndex]);
+					return false;
+				}
+			}
+			else
+			{
+				// If order files don't have explicit numbers just use the line number
+				OpenOrderNumber = EntryIndex;
 			}
 
 			FPaths::NormalizeFilename(Path);
@@ -346,6 +360,7 @@ bool FPakOrderMap::ProcessOrderFile(const TCHAR* ResponseFile, bool bSecondaryOr
 			}
 
 			OrderMap.Add(Path, OpenOrderNumber + OrderOffset);
+			MaxIndex = FMath::Max(MaxIndex, OpenOrderNumber + OrderOffset);
 		}
 
 		UE_LOG(LogPakFile, Display, TEXT("Finished loading pak order file %s."), ResponseFile);
@@ -358,6 +373,27 @@ bool FPakOrderMap::ProcessOrderFile(const TCHAR* ResponseFile, bool bSecondaryOr
 	}
 }
 
+void FPakOrderMap::MergeOrderMap(FPakOrderMap&& Other)
+{
+	for (TPair<FString, uint64>& OrderedFile : Other.OrderMap)
+	{
+		if (OrderMap.Contains(OrderedFile.Key) == false)
+		{
+			OrderMap.Add(MoveTemp(OrderedFile.Key), OrderedFile.Value);
+		}
+	}
+	Other.OrderMap.Empty(); // We moved strings out of this so empty it so nobody tries to use the keys
+
+	if (MaxIndex == MAX_uint64)
+	{
+		MaxIndex = Other.MaxIndex;
+	}
+	else
+	{
+		MaxIndex = FMath::Max(MaxIndex, Other.MaxIndex);
+	}
+}
+
 uint64 FPakOrderMap::GetFileOrder(const FString& Path, bool bAllowUexpUBulkFallback, bool* OutIsPrimary) const
 {
 	FString RegionStr;
diff --git a/Engine/Source/Developer/PakFileUtilities/Public/PakFileUtilities.h b/Engine/Source/Developer/PakFileUtilities/Public/PakFileUtilities.h
--- a/Engine/Source/Developer/PakFileUtilities/Public/PakFileUtilities.h
+++ b/Engine/Source/Developer/PakFileUtilities/Public/PakFileUtilities.h
@@ -20,12 +20,14 @@ class FPakOrderMap
 public:
 	FPakOrderMap()
 		: MaxPrimaryOrderIndex(MAX_uint64)
+		, MaxIndex(0)
 	{}
 
 	void Empty()
 	{
 		OrderMap.Empty();
 		MaxPrimaryOrderIndex = MAX_uint64;
+		MaxIndex = 0;
 	}
 
 	int32 Num() const
@@ -37,6 +39,7 @@ public:
 	void Add(const FString& Filename, uint64 Index)
 	{
 		OrderMap.Add(Filename, Index);
+		MaxIndex = FMath::Max(MaxIndex, Index);
 	}
 
 	/**
@@ -46,6 +49,7 @@ public:
 	void AddOffset(const FString& Filename, uint64 Offset)
 	{
 		OrderMap.Add(Filename, Offset);
+		MaxIndex = FMath::Max(MaxIndex, Offset);
 	}
 
 	/** Remaps all the current values in the OrderMap onto [0, NumEntries).  Useful to convert from Offset in Pak file bytes into an Index sorted by Offset */
@@ -66,19 +70,26 @@ public:
 			OrderMap[FilenameAndOffset.Key] = Index;
 			++Index;
 		}
+		MaxIndex = Index - 1;
 	}
 
-	bool PAKFILEUTILITIES_API ProcessOrderFile(const TCHAR* ResponseFile, bool bSecondaryOrderFile = false, bool bMergeOrder = false);
+	bool PAKFILEUTILITIES_API ProcessOrderFile(const TCHAR* ResponseFile, bool bSecondaryOrderFile = false, bool bMergeOrder = false, TOptional<uint64> InOffset = {});
+
+	// Merge another order map into this one where the files are not already ordered by this map. Steals the strings and empties the other order map.
+	void PAKFILEUTILITIES_API MergeOrderMap(FPakOrderMap&& Other);
 
 	uint64 PAKFILEUTILITIES_API GetFileOrder(const FString& Path, bool bAllowUexpUBulkFallback, bool* OutIsPrimary=nullptr) const;
 
 	void PAKFILEUTILITIES_API WriteOpenOrder(FArchive* Ar);
 
+	uint64 PAKFILEUTILITIES_API GetMaxIndex() { return MaxIndex; }
+
 private:
 	FString RemapLocalizationPathIfNeeded(const FString& PathLower, FString& OutRegion) const;
 
 	TMap<FString, uint64> OrderMap;
 	uint64 MaxPrimaryOrderIndex;
+	uint64 MaxIndex;
 };
 
 
diff --git a/Engine/Source/Editor/UnrealEd/Private/Commandlets/AssetRegistryGenerator.cpp b/Engine/Source/Editor/UnrealEd/Private/Commandlets/AssetRegistryGenerator.cpp
--- a/Engine/Source/Editor/UnrealEd/Private/Commandlets/AssetRegistryGenerator.cpp
+++ b/Engine/Source/Editor/UnrealEd/Private/Commandlets/AssetRegistryGenerator.cpp
@@ -28,6 +28,7 @@
 #include "Templates/UniquePtr.h"
 #include "Engine/AssetManager.h"
 #include "Misc/DataDrivenPlatformInfoRegistry.h"
+#include "Algo/StableSort.h"
 
 #include "PakFileUtilities.h"
 
@@ -382,7 +383,7 @@ bool FAssetRegistryGenerator::GenerateStreamingInstallManifest(int64 InExtraFlav
 	
 	bool bEnableGameOpenOrderSort = false;
 	bool bUseSecondaryOpenOrder = false;
-	TArray<FString> OrderFileSpecs;
+	TArray<FString> OrderFileSpecStrings;
 	{
 		FConfigFile PlatformIniFile;
 		FConfigCacheIni::LoadLocalIniFile(PlatformIniFile, TEXT("Game"), true, *TargetPlatform->IniPlatformName());
@@ -390,7 +391,7 @@ bool FAssetRegistryGenerator::GenerateStreamingInstallManifest(int64 InExtraFlav
 		PlatformIniFile.GetBool(TEXT("/Script/UnrealEd.ProjectPackagingSettings"), TEXT("bEnableAssetRegistryGameOpenOrderSort"), bEnableGameOpenOrderSort);
 		PlatformIniFile.GetBool(TEXT("/Script/UnrealEd.ProjectPackagingSettings"), TEXT("bPakUsesSecondaryOrder"), bUseSecondaryOpenOrder);
 
-		PlatformIniFile.GetArray(TEXT("/Script/UnrealEd.ProjectPackagingSettings"), TEXT("PakOrderFileSpecs"), OrderFileSpecs);
+		PlatformIniFile.GetArray(TEXT("/Script/UnrealEd.ProjectPackagingSettings"), TEXT("PakOrderFileSpecs"), OrderFileSpecStrings);
 	}
 
 	// if a game open order can be found then use that to sort the filenames
@@ -398,13 +399,29 @@ bool FAssetRegistryGenerator::GenerateStreamingInstallManifest(int64 InExtraFlav
 	bool bHaveGameOpenOrder = false;
 	if (bEnableGameOpenOrderSort)
 	{
-		if (OrderFileSpecs.Num() == 0)
+		TArray<FPakOrderFileSpec> OrderFileSpecs;
+		const FPakOrderFileSpec* SecondaryOrderSpec = nullptr;
+
+		if (OrderFileSpecStrings.Num() == 0)
 		{
-			OrderFileSpecs.Add(TEXT("GameOpenOrder*.log"));
+			OrderFileSpecs.Add(FPakOrderFileSpec(TEXT("GameOpenOrder*.log")));
 			if (bUseSecondaryOpenOrder)
 			{
-				OrderFileSpecs.Add(TEXT("CookerOpenOrder*.log"));
+				OrderFileSpecs.Add(FPakOrderFileSpec(TEXT("CookerOpenOrder*.log")));
+				SecondaryOrderSpec = &OrderFileSpecs.Last();
+			}
+		}
+		else
+		{
+			UScriptStruct* Struct = FPakOrderFileSpec::StaticStruct();
+			int32 Index = 0;
+			for (const FString& String : OrderFileSpecStrings)
+			{
+				FPakOrderFileSpec Spec;
+				Struct->ImportText(*String, &Spec, nullptr, PPF_Delimited, nullptr, Struct->GetName());
+				OrderFileSpecs.Add(Spec);
 			}
+
 		}
 
 		TArray<FString> DirsToSearch;
@@ -418,15 +435,23 @@ bool FAssetRegistryGenerator::GenerateStreamingInstallManifest(int64 InExtraFlav
 		{
 			DirsToSearch.Add(FPaths::Combine(FPaths::ProjectDir(), TEXT("Build"), Platform, TEXT("FileOpenOrder")));
 		}
-		
-		TArray<FString> OrderFiles;
-		for (const FString& Spec : OrderFileSpecs)
+
+		TArray<FPakOrderMap> OrderMaps; // Indexes match with OrderFileSpecs
+		OrderMaps.Reserve(OrderFileSpecs.Num());
+		uint64 StartIndex = 0;
+		for (const FPakOrderFileSpec& Spec : OrderFileSpecs)
 		{
+			// For each order map reserve it a contiguous integer range based on what indices were used by previous maps
+			// After building each map we'll then merge in priority order
+			UE_LOG(LogAssetRegistryGenerator, Log, TEXT("Order file spec %s starting at index %llu"), *Spec.Pattern, StartIndex);
+
+			FPakOrderMap& LocalOrderMap = OrderMaps.AddDefaulted_GetRef();
+
 			TArray<FString> FoundFiles;
 			for (const FString& Directory : DirsToSearch)
 			{
 				TArray<FString> LocalFoundFiles;
-				IFileManager::Get().FindFiles(LocalFoundFiles, *FPaths::Combine(Directory, Spec), true, false);
+				IFileManager::Get().FindFiles(LocalFoundFiles, *FPaths::Combine(Directory, Spec.Pattern), true, false);
 
 				for (const FString& Filename : LocalFoundFiles)
 				{
@@ -444,18 +469,46 @@ bool FAssetRegistryGenerator::GenerateStreamingInstallManifest(int64 InExtraFlav
 				return Order;
 			});
 
-			for (const FString& Found : FoundFiles)
+			if (FoundFiles.Num() > 0)
+			{
+				bHaveGameOpenOrder = bHaveGameOpenOrder || (&Spec != SecondaryOrderSpec);
+			}
+
+			for (int32 i=0; i < FoundFiles.Num(); ++i)
 			{
-				bHaveGameOpenOrder = bHaveGameOpenOrder || Found.Contains(TEXT("CookerOpenOrder")) == false;
+				const FString& Found = FoundFiles[i];
 				UE_LOG(LogAssetRegistryGenerator, Display, TEXT("Found order file %s"), *Found);
+				if (LocalOrderMap.Num() == 0) 
+				{
+					LocalOrderMap.ProcessOrderFile(*Found, &Spec == SecondaryOrderSpec, false, StartIndex);
+				}
+				else
+				{
+					LocalOrderMap.ProcessOrderFile(*Found, &Spec == SecondaryOrderSpec, true);
+				}
+			}
+
+			if (LocalOrderMap.Num() > 0)
+			{
+				check(LocalOrderMap.GetMaxIndex() >= StartIndex);
+				StartIndex = LocalOrderMap.GetMaxIndex() + 1;
 			}
+		}
 
-			OrderFiles.Append(FoundFiles);
+		TArray<int32> SpecIndicesByPriority;
+		for (int32 i=0; i < OrderFileSpecs.Num(); ++i)
+		{
+			SpecIndicesByPriority.Add(i);
 		}
+		Algo::StableSortBy(SpecIndicesByPriority, [&OrderFileSpecs](int32 i) { return  OrderFileSpecs[i].Priority; }, TGreater<int32>());
 
-		for (const FString& OrderFile : OrderFiles)
+		UE_LOG(LogAssetRegistryGenerator, Log, TEXT("Merging order maps in priority order"));
+		for (int32 SpecIndex : SpecIndicesByPriority)
 		{
-			OrderMap.ProcessOrderFile(*OrderFile, OrderFile.Contains(TEXT("CookerOpenOrder")), &OrderFile != &OrderFiles[0]);
+			const FPakOrderFileSpec& Spec = OrderFileSpecs[SpecIndex];
+			UE_LOG(LogAssetRegistryGenerator, Log, TEXT("Merging order file spec %d (%s) at priority (%d) "), SpecIndex, *Spec.Pattern, Spec.Priority);
+			FPakOrderMap& LocalOrderMap = OrderMaps[SpecIndex];
+			OrderMap.MergeOrderMap(MoveTemp(LocalOrderMap));
 		}
 	}
 
diff --git a/Engine/Source/Programs/AutomationTool/Scripts/CopyBuildToStagingDirectory.Automation.cs b/Engine/Source/Programs/AutomationTool/Scripts/CopyBuildToStagingDirectory.Automation.cs
--- a/Engine/Source/Programs/AutomationTool/Scripts/CopyBuildToStagingDirectory.Automation.cs
+++ b/Engine/Source/Programs/AutomationTool/Scripts/CopyBuildToStagingDirectory.Automation.cs
@@ -38,13 +38,14 @@ public partial class Project : CommandUtils
 			Editor
 		}
 
-		public OrderFile(FileReference FileRef, OrderFileType InType, int SpecIndex)
+		public OrderFile(FileReference FileRef, OrderFileType InType, int SpecIndex, int InPriority)
 		{
 			File = FileRef;
 			OrderSpecIndex = SpecIndex;
 
 			string FileName = FileRef.GetFileNameWithoutExtension();
 			OrderType = InType;
+			Priority = InPriority;
 
 			//Check if we have an order number
 			AppendOrder = -1;
@@ -77,6 +78,8 @@ public partial class Project : CommandUtils
 		// OrderType can be used for filtering
 		public int AppendOrder;
 		public int OrderSpecIndex;
+		// Priority at which this order file can claim files even if it appears later in the list than another order file
+		public int Priority;
 
 		public OrderFileType OrderType;
 
@@ -88,6 +91,7 @@ public partial class Project : CommandUtils
 		public string FileNamePattern;
 		public OrderFile.OrderFileType OrderType;
 		public bool Required;
+		public int Priority;
 	};
 
 	/// <returns>The path for the BuildPatchTool executable depending on host platform.</returns>
@@ -2462,9 +2466,31 @@ public partial class Project : CommandUtils
 		List<OrderFileSpec> OrderFileSpecs = new List<OrderFileSpec>();
 		if (PlatformGameConfig.TryGetValues("/Script/UnrealEd.ProjectPackagingSettings", "PakOrderFileSpecs", out IReadOnlyList<string> OrderFileConfigValues))
 		{
-			foreach (string FilePattern in OrderFileConfigValues)
+			foreach (string Config in OrderFileConfigValues)
 			{
-				OrderFileSpecs.Add(new OrderFileSpec { FileNamePattern = FilePattern, OrderType = OrderFile.OrderFileType.Custom, Required = true });
+				Dictionary<string, string> Fields;
+				if (!ConfigHierarchy.TryParse(Config, out Fields))
+				{
+					throw new AutomationException(String.Format("Failed to parse pak order file config \"{0}\"", Config));
+				}
+
+				string FilePattern;
+				if (!Fields.TryGetValue("Pattern", out FilePattern))
+				{
+					throw new AutomationException(String.Format("Missing field 'Pattern' in pak order file config \"{0}\"", Config));
+				}
+
+				// Priority missing is fine, we default to 0
+				int Priority = 0;
+				if( Fields.TryGetValue("Priority", out string PriorityString)) 
+				{
+					if (!int.TryParse(PriorityString, out Priority))
+					{
+						throw new AutomationException(String.Format("Failed to parse int from field 'Priority' in pak order file config \"{0}\"", Config));
+					}
+				}
+
+				OrderFileSpecs.Add(new OrderFileSpec { FileNamePattern = FilePattern, OrderType = OrderFile.OrderFileType.Custom, Required = FilePattern.Contains("*") == false, Priority = Priority });
 			}
 		}
 
@@ -2472,9 +2498,9 @@ public partial class Project : CommandUtils
 		{
 			// Default filespecs
 			OrderFileSpecs.AddRange(new OrderFileSpec[] {
-				new OrderFileSpec{ FileNamePattern = "GameOpenOrder*.log", OrderType = OrderFile.OrderFileType.Game, Required=false },
-				new OrderFileSpec{ FileNamePattern = "CookerOpenOrder*.log", OrderType = OrderFile.OrderFileType.Cooker, Required=false },
-				new OrderFileSpec{ FileNamePattern = "EditorOpenOrder.log", OrderType = OrderFile.OrderFileType.Editor, Required=false }
+				new OrderFileSpec{ FileNamePattern = "GameOpenOrder*.log", OrderType = OrderFile.OrderFileType.Game, Required=false, Priority = 0 },
+				new OrderFileSpec{ FileNamePattern = "CookerOpenOrder*.log", OrderType = OrderFile.OrderFileType.Cooker, Required=false, Priority = 0 },
+				new OrderFileSpec{ FileNamePattern = "EditorOpenOrder.log", OrderType = OrderFile.OrderFileType.Editor, Required=false, Priority = 0 }
 			});
 		}
 
@@ -2510,7 +2536,7 @@ public partial class Project : CommandUtils
 					foreach (var File in FileLocations)
 					{
 						bAnyFound = true;
-						OrderFile FileOrder = new OrderFile(File, Spec.OrderType, OrderFileIndex);
+						OrderFile FileOrder = new OrderFile(File, Spec.OrderType, OrderFileIndex, Spec.Priority);
 
 						//check if the file is alreay in the list
 						if (!OrderFiles.Any(o => o.File.FullName.Equals(File.FullName)))
@@ -2549,7 +2575,7 @@ public partial class Project : CommandUtils
 		LogInformation("Using {0} pak order files:", OrderFiles.Count);
 		foreach( OrderFile File in OrderFiles )
 		{
-			LogInformation("    {0}", File.File.ToString());
+			LogInformation("    {0} Priority {1}", File.File.ToString(), File.Priority);
 		}
 
 		List<Tuple<FileReference, StagedFileReference, string>> Outputs = new List<Tuple<FileReference, StagedFileReference, string>>();
@@ -2806,12 +2832,10 @@ public partial class Project : CommandUtils
 				}
 			}
 
-			List<OrderFile> PrimaryOrderFiles = OrderFiles.FindAll(x => (x.OrderType != OrderFile.OrderFileType.Cooker));
-			List<OrderFile> SecondaryOrderFiles = null;
-
+			List<OrderFile> FinalOrderFiles = OrderFiles.FindAll(x => (x.OrderType != OrderFile.OrderFileType.Cooker));
 			if (bUseSecondaryOrder && OrderFiles.Count >= 1)
 			{
-				SecondaryOrderFiles = OrderFiles.FindAll(x => x.OrderType == OrderFile.OrderFileType.Cooker);
+				FinalOrderFiles.AddRange(OrderFiles.FindAll(x => x.OrderType == OrderFile.OrderFileType.Cooker));
 			}
 
 			string AdditionalArgs =
@@ -2846,7 +2870,7 @@ public partial class Project : CommandUtils
 
 			AdditionalArgs += " " + Params.AdditionalIoStoreOptions;
 
-			RunIoStore(Params, SC, IoStoreCommandsFileName, PrimaryOrderFiles, SecondaryOrderFiles, AdditionalArgs);
+			RunIoStore(Params, SC, IoStoreCommandsFileName, FinalOrderFiles, AdditionalArgs);
 		}
 
 		// Do any additional processing on the command output
@@ -3006,7 +3030,7 @@ public partial class Project : CommandUtils
 		}
 	}
 
-	private static void RunIoStore(ProjectParams Params, DeploymentContext SC, string CommandsFileName, List<OrderFile> GameOpenOrderFileLocations, List<OrderFile> CookerOpenOrderFileLocations, string AdditionalArgs)
+	private static void RunIoStore(ProjectParams Params, DeploymentContext SC, string CommandsFileName, List<OrderFile> OrderFiles, string AdditionalArgs)
 	{
 		StagedFileReference GlobalContainerOutputRelativeLocation;
 		GlobalContainerOutputRelativeLocation = StagedFileReference.Combine(SC.RelativeProjectRootForStage, "Content", "Paks", "global.utoc");
@@ -3057,14 +3081,11 @@ public partial class Project : CommandUtils
 		}
 
 		CommandletParams += String.Format(" -CookedDirectory={0} -Commands={1}", MakePathSafeToUseWithCommandLine(SC.PlatformCookDir.ToString()), MakePathSafeToUseWithCommandLine(CommandsFileName));
-		if (GameOpenOrderFileLocations != null && GameOpenOrderFileLocations.Count() > 0)
+		if (OrderFiles != null && OrderFiles.Count() > 0)
 		{
-			CommandletParams += String.Format(" -GameOrder={0}", MakePathSafeToUseWithCommandLine(string.Join(",", GameOpenOrderFileLocations.Select(u => u.File.FullName).ToArray())));
-		}
-		if (CookerOpenOrderFileLocations != null && CookerOpenOrderFileLocations.Count() > 0)
-		{
-			CommandletParams += String.Format(" -CookerOrder={0}", MakePathSafeToUseWithCommandLine(string.Join(",", CookerOpenOrderFileLocations.Select(u => u.File.FullName).ToArray())));
-		}
+			CommandletParams += String.Format(" -Order={0}", MakePathSafeToUseWithCommandLine(string.Join(",", OrderFiles.Select(u => u.File.FullName).ToArray())));
+			CommandletParams += String.Format(" -OrderPriority={0}", MakePathSafeToUseWithCommandLine(string.Join(",", OrderFiles.Select(u => u.Priority).ToArray())));
+		} 
 		if (!string.IsNullOrWhiteSpace(AdditionalArgs))
 		{
 			CommandletParams += AdditionalArgs;
