Fix static analysis detecting a spurious null pointer dereference.

diff --git a/Engine/Source/Developer/IoStoreUtilities/Private/IoStoreUtilities.cpp b/Engine/Source/Developer/IoStoreUtilities/Private/IoStoreUtilities.cpp
--- a/Engine/Source/Developer/IoStoreUtilities/Private/IoStoreUtilities.cpp
+++ b/Engine/Source/Developer/IoStoreUtilities/Private/IoStoreUtilities.cpp
@@ -478,8 +478,16 @@ static void AssignPackagesDiskOrder(
 	struct FPackageAndOrder
 	{
 		FLegacyCookedPackage* Package = nullptr;
-		int64 LocalOrder = MAX_uint64;
-		const FFileOrderMap* OrderMap = nullptr;
+		int64 LocalOrder;
+		const FFileOrderMap* OrderMap;
+
+		FPackageAndOrder(FLegacyCookedPackage* InPackage, int64 InLocalOrder, const FFileOrderMap* InOrderMap)
+			: Package(InPackage)
+			, LocalOrder(InLocalOrder)
+			, OrderMap(InOrderMap)
+		{
+			check(OrderMap);
+		}
 	};
 
 	// Order maps sorted by priority
@@ -492,7 +500,7 @@ static void AssignPackagesDiskOrder(
 
 	// Create a fallback order map to avoid null checks later
 	// Lowest priority, last index
-	FFileOrderMap FallbackOrderMap(MIN_int32, MAX_int32); 
+	FFileOrderMap FallbackOrderMap(MIN_int32, MAX_int32);
 
 	TArray<FPackageAndOrder> SortedPackages;
 	SortedPackages.Reserve(Packages.Num());
@@ -502,25 +510,23 @@ static void AssignPackagesDiskOrder(
 		{
 			continue;
 		}
-		FPackageAndOrder& Entry = SortedPackages.AddDefaulted_GetRef();
-		Entry.Package = Package;
+
+		// Default to the fallback order map 
+		// Reverse the bundle load order for the fallback map (so that packages are considered before their imports)
+		const FFileOrderMap* UsedOrderMap = &FallbackOrderMap;
+		int64 LocalOrder = -int64(Package->OptimizedPackage->GetLoadOrder());
 
 		for (const FFileOrderMap* OrderMap : PriorityOrderMaps)
 		{
 			if (const int64* Order = OrderMap->PackageNameToOrder.Find(Package->PackageName))
 			{
-				Entry.LocalOrder = *Order;
-				Entry.OrderMap = OrderMap;
+				LocalOrder = *Order;
+				UsedOrderMap = OrderMap;
 				break;
 			}
 		}
 
-		if (Entry.OrderMap == nullptr)
-		{
-			Entry.OrderMap = &FallbackOrderMap;
-			// Reverse the bundle load order for the fallback map (so that packages are considered before their imports)
-			Entry.LocalOrder = -int64(Package->OptimizedPackage->GetLoadOrder());
-		}
+		SortedPackages.Emplace(Package, LocalOrder, UsedOrderMap);
 	}
 	const FFileOrderMap* LastBlameOrderMap = nullptr;
 	int32 LastAssignedCount = 0;
@@ -562,6 +568,7 @@ static void AssignPackagesDiskOrder(
 
 	for (FPackageAndOrder& Entry : SortedPackages)
 	{
+		checkSlow(Entry.OrderMap); // Without this, Entry.OrderMap != LastBlameOrderMap convinces static analysis that Entry.OrderMap may be null
 		if (Entry.OrderMap != LastBlameOrderMap)
 		{
 			if( LastBlameOrderMap != nullptr )
