Cluster by order file priority by default.

diff --git a/Engine/Source/Developer/IoStoreUtilities/Private/IoStoreUtilities.cpp b/Engine/Source/Developer/IoStoreUtilities/Private/IoStoreUtilities.cpp
--- a/Engine/Source/Developer/IoStoreUtilities/Private/IoStoreUtilities.cpp
+++ b/Engine/Source/Developer/IoStoreUtilities/Private/IoStoreUtilities.cpp
@@ -3436,7 +3436,7 @@ int32 CreateIoStoreContainerFiles(const TCHAR* CmdLine)
 		}
 	}
 
-	Arguments.bClusterByOrderFilePriority = FParse::Param(CmdLine, TEXT("ClusterByOrderPriority"));
+	Arguments.bClusterByOrderFilePriority = !FParse::Param(CmdLine, TEXT("DoNotClusterByOrderPriority"));
 
 	FIoStoreWriterSettings GeneralIoWriterSettings { DefaultCompressionMethod, DefaultCompressionBlockSize, false };
 	GeneralIoWriterSettings.bEnableCsvOutput = FParse::Param(CmdLine, TEXT("-csvoutput"));
