Non-Unity Fixes

diff --git a/Engine/Source/Developer/IoStoreUtilities/Private/ZenStoreHttpClient.cpp b/Engine/Source/Developer/IoStoreUtilities/Private/ZenStoreHttpClient.cpp
--- a/Engine/Source/Developer/IoStoreUtilities/Private/ZenStoreHttpClient.cpp
+++ b/Engine/Source/Developer/IoStoreUtilities/Private/ZenStoreHttpClient.cpp
@@ -6,13 +6,14 @@
 #include "Algo/IsSorted.h"
 #include "Algo/Sort.h"
 #include "Misc/ScopeRWLock.h"
+#include "Misc/App.h"
+#include "HAL/FileManager.h"
 
 #if PLATFORM_WINDOWS
 
 #include "Async/Async.h"
 #include "HAL/PlatformFileManager.h"
 #include "IO/IoHash.h"
-#include "Misc/App.h"
 #include "Misc/Paths.h"
 #include "Misc/ScopeLock.h"
 #include "Misc/SecureHash.h"
