Fixed package store manifest not working correctly for iterative cooks
#rnx

diff --git a/Engine/Source/Developer/IoStoreUtilities/Private/ZenStoreWriter.cpp b/Engine/Source/Developer/IoStoreUtilities/Private/ZenStoreWriter.cpp
--- a/Engine/Source/Developer/IoStoreUtilities/Private/ZenStoreWriter.cpp
+++ b/Engine/Source/Developer/IoStoreUtilities/Private/ZenStoreWriter.cpp
@@ -473,6 +473,8 @@ void FZenStoreWriter::Initialize(const FCookInfo& Info)
 
 void FZenStoreWriter::BeginCook()
 {
+	PackageStoreManifest.Load(*(MetadataDirectoryPath / TEXT("packagestore.manifest")));
+
 	if (CookMode == ICookedPackageWriter::FCookInfo::CookOnTheFlyMode)
 	{
 		FCbPackage Pkg;
diff --git a/Engine/Source/Editor/UnrealEd/Private/Cooker/LooseCookedPackageWriter.cpp b/Engine/Source/Editor/UnrealEd/Private/Cooker/LooseCookedPackageWriter.cpp
--- a/Engine/Source/Editor/UnrealEd/Private/Cooker/LooseCookedPackageWriter.cpp
+++ b/Engine/Source/Editor/UnrealEd/Private/Cooker/LooseCookedPackageWriter.cpp
@@ -366,6 +366,7 @@ void FLooseCookedPackageWriter::Initialize(const FCookInfo& Info)
 
 void FLooseCookedPackageWriter::BeginCook()
 {
+	PackageStoreManifest.Load(*(MetadataDirectoryPath / TEXT("packagestore.manifest")));
 }
 
 void FLooseCookedPackageWriter::EndCook()
