Minor optimizations to PackageStoreOptimizer when cooking with ZenStore

#rnx

diff --git a/Engine/Source/Developer/IoStoreUtilities/Internal/PackageStoreOptimizer.h b/Engine/Source/Developer/IoStoreUtilities/Internal/PackageStoreOptimizer.h
--- a/Engine/Source/Developer/IoStoreUtilities/Internal/PackageStoreOptimizer.h
+++ b/Engine/Source/Developer/IoStoreUtilities/Internal/PackageStoreOptimizer.h
@@ -219,6 +219,11 @@ private:
 			return FromExportBundleIndex == Other.FromExportBundleIndex &&
 				ToExportBundleIndex == Other.ToExportBundleIndex;
 		}
+
+		friend FORCEINLINE uint32 GetTypeHash(const FInternalArc& Arc)
+		{
+			return HashCombine(GetTypeHash(Arc.FromExportBundleIndex), GetTypeHash(Arc.ToExportBundleIndex));
+		}
 	};
 
 	struct FExternalArc
diff --git a/Engine/Source/Developer/IoStoreUtilities/Private/PackageStoreOptimizer.cpp b/Engine/Source/Developer/IoStoreUtilities/Private/PackageStoreOptimizer.cpp
--- a/Engine/Source/Developer/IoStoreUtilities/Private/PackageStoreOptimizer.cpp
+++ b/Engine/Source/Developer/IoStoreUtilities/Private/PackageStoreOptimizer.cpp
@@ -196,6 +196,7 @@ FPackageStoreOptimizer::FCookedHeaderData FPackageStoreOptimizer::LoadCookedHead
 
 		FNameEntrySerialized NameEntry(ENAME_LinkerConstructor);
 
+		CookedHeaderData.SummaryNames.Reserve(Summary.NameCount);
 		for (int32 I = 0; I < Summary.NameCount; ++I)
 		{
 			Ar << NameEntry;
@@ -438,7 +439,7 @@ void FPackageStoreOptimizer::ProcessImports(const FCookedHeaderData& CookedHeade
 			{
 				UE_LOG(LogPackageStoreOptimizer, Warning, TEXT("Package '%s' is referencing missing script import '%s'"), *Package->Name.ToString(), *UnresolvedImport.FullName);
 			}
-			Package->Imports[ImportIndex] = FPackageObjectIndex::FromScriptPath(UnresolvedImport.FullName);
+			Package->Imports[ImportIndex] = ScriptObjectIndex;
 		}
 		else if (!UnresolvedImport.bIsImportOfPackage)
 		{
@@ -1522,16 +1523,22 @@ void FPackageStoreOptimizer::FinalizePackage(FPackageStorePackage* Package)
 	for (FPackageStorePackage::FExportGraphNode& Node : Package->ExportGraphNodes)
 	{
 		check(Node.ExportBundleIndex >= 0);
+		TSet<FPackageStorePackage::FInternalArc> InternalArcsSet;
 		for (FPackageStorePackage::FExportGraphNode* InternalDependency : Node.InternalDependencies)
 		{
 			FPackageStorePackage::FInternalArc Arc;
 			check(InternalDependency->ExportBundleIndex >= 0);
 			Arc.FromExportBundleIndex = InternalDependency->ExportBundleIndex;
 			Arc.ToExportBundleIndex = Node.ExportBundleIndex;
-			if (Arc.FromExportBundleIndex != Arc.ToExportBundleIndex && !Package->GraphData.InternalArcs.Contains(Arc))
+			if (Arc.FromExportBundleIndex != Arc.ToExportBundleIndex)
 			{
-				Package->GraphData.InternalArcs.Add(Arc);
-				++TotalInternalBundleArcsCount;
+				bool bIsAlreadyInSet = false;
+				InternalArcsSet.Add(Arc, &bIsAlreadyInSet);
+				if (!bIsAlreadyInSet)
+				{
+					Package->GraphData.InternalArcs.Add(Arc);
+					++TotalInternalBundleArcsCount;
+				}
 			}
 		}
 
