Added support to PackageStoreOptimizer to support optional package when generated from legacy cooked header
Support for Zen store will be added separately

diff --git a/Engine/Source/Developer/IoStoreUtilities/Internal/PackageStoreOptimizer.h b/Engine/Source/Developer/IoStoreUtilities/Internal/PackageStoreOptimizer.h
--- a/Engine/Source/Developer/IoStoreUtilities/Internal/PackageStoreOptimizer.h
+++ b/Engine/Source/Developer/IoStoreUtilities/Internal/PackageStoreOptimizer.h
@@ -205,8 +205,10 @@ private:
 		FName FromPackageName;
 		int32 FromPackageNameLen = 0;
 		FPackageId FromPackageId;
+		FPackageId FromOptionalPackageId;
 		bool bIsScriptImport = false;
 		bool bIsImportOfPackage = false;
+		bool bIsImportOptional = false;
 	};
 
 	struct FInternalArc
@@ -294,6 +296,7 @@ private:
 	bool bPermanentMark = false;
 
 	bool bIsRedirected = false;
+	bool bIsOptional = false;
 
 	friend class FPackageStoreOptimizer;
 };
@@ -372,7 +375,7 @@ private:
 		TArray<FPackageId> ImportedPackageIds;
 		TArray<uint64> ImportedPublicExportHashes;
 		TArray<FNameEntryId> NameMap;
-		TArray<FPackageObjectIndex> Imports;
+		TArray<FPackageObjectIndex> Imports; // FH: Imports might need to have more info to be able to resolve export hash with import as outer
 		TArray<FExportMapEntry> Exports;
 		TArray<FExportBundleHeader> ExportBundleHeaders;
 		TArray<FExportBundleEntry> ExportBundleEntries;
@@ -387,11 +390,11 @@ private:
 	FCookedHeaderData LoadCookedHeader(const FIoBuffer& CookedHeaderBuffer) const;
 	FPackageStoreHeaderData LoadPackageStoreHeader(const FIoBuffer& PackageStoreHeaderBuffer, const FPackageStoreEntryResource& PackageStoreEntry) const;
 	void ResolveImport(FPackageStorePackage::FUnresolvedImport* Imports, const FObjectImport* ObjectImports, int32 LocalImportIndex) const;
-	void ProcessImports(const FCookedHeaderData& CookedHeaderData, FPackageStorePackage* Package) const;
+	void ProcessImports(const FCookedHeaderData& CookedHeaderData, FPackageStorePackage* Package, TArray<FPackageStorePackage::FUnresolvedImport>& UnresolvedImports) const;
 	void ProcessImports(const FPackageStoreHeaderData& PackageStoreHeaderData, FPackageStorePackage* Package) const;
-	void ResolveExport(FPackageStorePackage::FExport* Exports, const FObjectExport* ObjectExports, const int32 LocalExportIndex, const FName& PackageName) const;
+	void ResolveExport(FPackageStorePackage::FExport* Exports, const FObjectExport* ObjectExports, const int32 LocalExportIndex, const FName& PackageName, FPackageStorePackage::FUnresolvedImport* Imports = nullptr, const FObjectImport* ObjectImports = nullptr) const;
 	void ResolveExport(FPackageStorePackage::FExport* Exports, const int32 LocalExportIndex, const FName& PackageName) const;
-	void ProcessExports(const FCookedHeaderData& CookedHeaderData, FPackageStorePackage* Package) const;
+	void ProcessExports(const FCookedHeaderData& CookedHeaderData, FPackageStorePackage* Package, FPackageStorePackage::FUnresolvedImport* Imports = nullptr) const;
 	void ProcessExports(const FPackageStoreHeaderData& PackageStoreHeaderData, FPackageStorePackage* Package) const;
 	void ProcessPreloadDependencies(const FCookedHeaderData& CookedHeaderData, FPackageStorePackage* Package) const;
 	void ProcessPreloadDependencies(const FPackageStoreHeaderData& PackageStoreHeaderData, FPackageStorePackage* Package) const;
diff --git a/Engine/Source/Developer/IoStoreUtilities/Private/PackageStoreOptimizer.cpp b/Engine/Source/Developer/IoStoreUtilities/Private/PackageStoreOptimizer.cpp
--- a/Engine/Source/Developer/IoStoreUtilities/Private/PackageStoreOptimizer.cpp
+++ b/Engine/Source/Developer/IoStoreUtilities/Private/PackageStoreOptimizer.cpp
@@ -119,10 +119,23 @@ FPackageStorePackage* FPackageStoreOptimizer::CreateMissingPackage(const FName&
 FPackageStorePackage* FPackageStoreOptimizer::CreatePackageFromCookedHeader(const FName& Name, const FIoBuffer& CookedHeaderBuffer) const
 {
 	FPackageStorePackage* Package = new FPackageStorePackage();
-	Package->Name = Name;
 	Package->Id = FPackageId::FromName(Name);
 
-	Package->SourceName = *RemapLocalizationPathIfNeeded(Name.ToString(), &Package->Region);
+	// The package id should be generated from the original name
+	// However strip any '.' from the passed in name if any to generate the PackageName,
+	// <PackageName>.o is currently used to identify an optional package (optional exports chunk)
+	FString NameStr = Name.ToString();
+	FName PackageName = Name;
+	int32 Index = NameStr.Find(FPackagePath::GetOptionalSegmentExtensionModifier());
+	if (Index != INDEX_NONE)
+	{
+		NameStr.LeftInline(Index, false);
+		PackageName = *NameStr;
+		Package->bIsOptional = true;
+	}
+
+	Package->Name = PackageName;
+	Package->SourceName = *RemapLocalizationPathIfNeeded(NameStr, &Package->Region);
 
 	FCookedHeaderData CookedHeaderData = LoadCookedHeader(CookedHeaderBuffer);
 	if (!CookedHeaderData.Summary.bUnversioned)
@@ -139,8 +152,10 @@ FPackageStorePackage* FPackageStoreOptimizer::CreatePackageFromCookedHeader(cons
 	{
 		Package->NameMapBuilder.AddName(CookedHeaderData.SummaryNames[I]);
 	}
-	ProcessImports(CookedHeaderData, Package);
-	ProcessExports(CookedHeaderData, Package);
+
+	TArray<FPackageStorePackage::FUnresolvedImport> Imports;
+	ProcessImports(CookedHeaderData, Package, Imports);
+	ProcessExports(CookedHeaderData, Package, Imports.GetData());
 	ProcessPreloadDependencies(CookedHeaderData, Package);
 	
 	CreateExportBundles(Package);
@@ -151,10 +166,19 @@ FPackageStorePackage* FPackageStoreOptimizer::CreatePackageFromCookedHeader(cons
 FPackageStorePackage* FPackageStoreOptimizer::CreatePackageFromPackageStoreHeader(const FName& Name, const FIoBuffer& Buffer, const FPackageStoreEntryResource& PackageStoreEntry) const
 {
 	FPackageStorePackage* Package = new FPackageStorePackage();
-	Package->Name = Name;
 	Package->Id = FPackageId::FromName(Name);
-	
-	Package->SourceName = *RemapLocalizationPathIfNeeded(Name.ToString(), &Package->Region);
+
+	// The package id should be generated from the original name
+	// Support for optional package is not implemented when loading from the package store yet however
+	FString NameStr = Name.ToString();
+	int32 Index = NameStr.Find(FPackagePath::GetOptionalSegmentExtensionModifier());
+	if (Index != INDEX_NONE)
+	{
+		unimplemented();
+	}
+
+	Package->Name = Name;
+	Package->SourceName = *RemapLocalizationPathIfNeeded(NameStr, &Package->Region);
 
 	FPackageStoreHeaderData PackageStoreHeaderData = LoadPackageStoreHeader(Buffer, PackageStoreEntry);
 	if (PackageStoreHeaderData.VersioningInfo.IsSet())
@@ -379,6 +403,7 @@ void FPackageStoreOptimizer::ResolveImport(FPackageStorePackage::FUnresolvedImpo
 			PackageName.AppendString(Import->FullName);
 			Import->FullName.ToLowerInline();
 			Import->FromPackageId = FPackageId::FromName(PackageName);
+			Import->FromOptionalPackageId = FPackageId::FromName(PackageName, true);
 			Import->FromPackageName = PackageName;
 			Import->FromPackageNameLen = Import->FullName.Len();
 			Import->bIsScriptImport = Import->FullName.StartsWith(TEXT("/Script/"));
@@ -396,6 +421,8 @@ void FPackageStoreOptimizer::ResolveImport(FPackageStorePackage::FUnresolvedImpo
 			ObjectImport->ObjectName.AppendString(Import->FullName);
 			Import->FullName.ToLowerInline();
 			Import->FromPackageId = OuterImport->FromPackageId;
+			Import->FromOptionalPackageId = OuterImport->FromOptionalPackageId;
+			Import->bIsImportOptional = ObjectImport->bImportOptional;
 			Import->FromPackageName = OuterImport->FromPackageName;
 			Import->FromPackageNameLen = OuterImport->FromPackageNameLen;
 		}
@@ -409,10 +436,9 @@ uint64 FPackageStoreOptimizer::GetPublicExportHash(FStringView PackageRelativeEx
 	return CityHash64(reinterpret_cast<const char*>(PackageRelativeExportPath.GetData() + 1), (PackageRelativeExportPath.Len() - 1) * sizeof(TCHAR));
 }
 
-void FPackageStoreOptimizer::ProcessImports(const FCookedHeaderData& CookedHeaderData, FPackageStorePackage* Package) const
+void FPackageStoreOptimizer::ProcessImports(const FCookedHeaderData& CookedHeaderData, FPackageStorePackage* Package, TArray<FPackageStorePackage::FUnresolvedImport>& UnresolvedImports) const
 {
 	int32 ImportCount = CookedHeaderData.ObjectImports.Num();
-	TArray<FPackageStorePackage::FUnresolvedImport> UnresolvedImports;
 	UnresolvedImports.SetNum(ImportCount);
 	Package->Imports.SetNum(ImportCount);
 
@@ -421,9 +447,17 @@ void FPackageStoreOptimizer::ProcessImports(const FCookedHeaderData& CookedHeade
 	{
 		ResolveImport(UnresolvedImports.GetData(), CookedHeaderData.ObjectImports.GetData(), ImportIndex);
 		FPackageStorePackage::FUnresolvedImport& UnresolvedImport = UnresolvedImports[ImportIndex];
-		if (!UnresolvedImport.bIsScriptImport && UnresolvedImport.bIsImportOfPackage)
+		if (!UnresolvedImport.bIsScriptImport )
 		{
-			ImportedPackageIds.Add(UnresolvedImport.FromPackageId);
+			if (UnresolvedImport.bIsImportOfPackage)
+			{
+				ImportedPackageIds.Add(UnresolvedImport.FromPackageId);
+			}
+			// Also add the associated optional package id to the list of imported package ids if the import is optional
+			else if (UnresolvedImport.bIsImportOptional)
+			{
+				ImportedPackageIds.Add(UnresolvedImport.FromOptionalPackageId);
+			}
 		}
 	}
 	Package->ImportedPackageIds = ImportedPackageIds.Array();
@@ -446,7 +480,8 @@ void FPackageStoreOptimizer::ProcessImports(const FCookedHeaderData& CookedHeade
 			bool bFoundPackageIndex = false;
 			for (uint32 PackageIndex = 0, PackageCount = static_cast<uint32>(Package->ImportedPackageIds.Num()); PackageIndex < PackageCount; ++PackageIndex)
 			{
-				if (UnresolvedImport.FromPackageId == Package->ImportedPackageIds[PackageIndex])
+				FPackageId FromPackageId = !UnresolvedImport.bIsImportOptional ? UnresolvedImport.FromPackageId : UnresolvedImport.FromOptionalPackageId;
+				if (FromPackageId == Package->ImportedPackageIds[PackageIndex])
 				{
 					FStringView PackageRelativeName = FStringView(UnresolvedImport.FullName).RightChop(UnresolvedImport.FromPackageNameLen);
 					check(PackageRelativeName.Len());
@@ -474,7 +509,9 @@ void FPackageStoreOptimizer::ResolveExport(
 	FPackageStorePackage::FExport* Exports,
 	const FObjectExport* ObjectExports,
 	const int32 LocalExportIndex,
-	const FName& PackageName) const
+	const FName& PackageName,
+	FPackageStorePackage::FUnresolvedImport* Imports,
+	const FObjectImport* ObjectImports) const
 {
 	FPackageStorePackage::FExport* Export = Exports + LocalExportIndex;
 	if (Export->FullName.Len() == 0)
@@ -491,13 +528,23 @@ void FPackageStoreOptimizer::ResolveExport(
 		}
 		else
 		{
-			check(ObjectExport->OuterIndex.IsExport());
+			FString* OuterName = nullptr;
+			if (ObjectExport->OuterIndex.IsExport())
+			{
+				int32 OuterExportIndex = ObjectExport->OuterIndex.ToExport();
+				ResolveExport(Exports, ObjectExports, OuterExportIndex, PackageName);
+				OuterName = &Exports[OuterExportIndex].FullName;
+			}
+			else
+			{
+				check(Imports && ObjectImports);
+				int32 OuterImportIndex = ObjectExport->OuterIndex.ToImport();
+				ResolveImport(Imports, ObjectImports, OuterImportIndex);
+				OuterName = &Imports[OuterImportIndex].FullName;
 
-			int32 OuterExportIndex = ObjectExport->OuterIndex.ToExport();
-			ResolveExport(Exports, ObjectExports, OuterExportIndex, PackageName);
-			FString& OuterName = Exports[OuterExportIndex].FullName;
-			check(OuterName.Len() > 0);
-			Export->FullName.Append(OuterName);
+			}
+			check(OuterName && OuterName->Len() > 0);
+			Export->FullName.Append(*OuterName);
 			Export->FullName.AppendChar(TEXT('/'));
 			ObjectExport->ObjectName.AppendString(Export->FullName);
 			Export->FullName.ToLowerInline();
@@ -534,7 +581,7 @@ void FPackageStoreOptimizer::ResolveExport(FPackageStorePackage::FExport* Export
 	}
 }
 
-void FPackageStoreOptimizer::ProcessExports(const FCookedHeaderData& CookedHeaderData, FPackageStorePackage* Package) const
+void FPackageStoreOptimizer::ProcessExports(const FCookedHeaderData& CookedHeaderData, FPackageStorePackage* Package, FPackageStorePackage::FUnresolvedImport* Imports) const
 {
 	int32 ExportCount = CookedHeaderData.ObjectExports.Num();
 	Package->Exports.SetNum(ExportCount);
@@ -569,8 +616,8 @@ void FPackageStoreOptimizer::ProcessExports(const FCookedHeaderData& CookedHeade
 		Export.SerialSize = ObjectExport.SerialSize;
 		Export.bNotForClient = ObjectExport.bNotForClient;
 		Export.bNotForServer = ObjectExport.bNotForServer;
-		Export.bIsPublic = (Export.ObjectFlags & RF_Public) > 0;
-		ResolveExport(Package->Exports.GetData(), CookedHeaderData.ObjectExports.GetData(), ExportIndex, Package->Name);
+		Export.bIsPublic = (Export.ObjectFlags & RF_Public) > 0 || ObjectExport.bGeneratePublicHash;
+		ResolveExport(Package->Exports.GetData(), CookedHeaderData.ObjectExports.GetData(), ExportIndex, Package->Name, Imports, CookedHeaderData.ObjectImports.GetData());
 		if (Export.bIsPublic)
 		{
 			check(Export.FullName.Len() > 0);
@@ -1782,6 +1829,7 @@ FPackageStoreEntryResource FPackageStoreOptimizer::CreatePackageStoreEntry(const
 {
 	FPackageStoreEntryResource Result;
 	Result.Flags = Package->bIsRedirected ? EPackageStoreEntryFlags::Redirected : EPackageStoreEntryFlags::None;
+	Result.Flags |= Package->bIsOptional ? EPackageStoreEntryFlags::Optional : EPackageStoreEntryFlags::None;
 	Result.PackageName = Package->Name;
 	Result.SourcePackageName = Package->SourceName;
 	Result.Region = FName(*Package->Region);
diff --git a/Engine/Source/Runtime/Core/Private/IO/PackageId.cpp b/Engine/Source/Runtime/Core/Private/IO/PackageId.cpp
--- a/Engine/Source/Runtime/Core/Private/IO/PackageId.cpp
+++ b/Engine/Source/Runtime/Core/Private/IO/PackageId.cpp
@@ -3,10 +3,18 @@
 #include "IO/PackageId.h"
 #include "Serialization/StructuredArchive.h"
 
-FPackageId FPackageId::FromName(const FName& Name)
+FPackageId FPackageId::FromName(const FName& Name, bool bAsOptional)
 {
-	TCHAR NameStr[FName::StringBufferSize];
+	TCHAR NameStr[FName::StringBufferSize + 2];
 	uint32 NameLen = Name.ToString(NameStr);
+
+	if (bAsOptional)
+	{
+		NameStr[NameLen++] = '.';
+		NameStr[NameLen++] = 'o';
+		NameStr[NameLen] = '\0';
+	}
+
 	for (uint32 I = 0; I < NameLen; ++I)
 	{
 		NameStr[I] = TChar<TCHAR>::ToLower(NameStr[I]);
diff --git a/Engine/Source/Runtime/Core/Public/IO/PackageId.h b/Engine/Source/Runtime/Core/Public/IO/PackageId.h
--- a/Engine/Source/Runtime/Core/Public/IO/PackageId.h
+++ b/Engine/Source/Runtime/Core/Public/IO/PackageId.h
@@ -17,7 +17,7 @@ class FPackageId
 public:
 	FPackageId() = default;
 
-	CORE_API static FPackageId FromName(const FName& Name);
+	CORE_API static FPackageId FromName(const FName& Name, bool bAsOptional = false);
 
 	static FPackageId FromValue(const uint64 Value)
 	{
diff --git a/Engine/Source/Runtime/Core/Public/IO/PackageStore.h b/Engine/Source/Runtime/Core/Public/IO/PackageStore.h
--- a/Engine/Source/Runtime/Core/Public/IO/PackageStore.h
+++ b/Engine/Source/Runtime/Core/Public/IO/PackageStore.h
@@ -54,6 +54,7 @@ enum class EPackageStoreEntryFlags : uint32
 {
 	None		= 0,
 	Redirected	= 0x01,
+	Optional	= 0x02,
 };
 ENUM_CLASS_FLAGS(EPackageStoreEntryFlags);
 
@@ -84,13 +85,13 @@ struct FPackageStoreEntryResource
 	/** Returns the package ID. */
 	FPackageId GetPackageId() const
 	{
-		return FPackageId::FromName(PackageName);
+		return FPackageId::FromName(PackageName, IsOptional());
 	}
 
 	/** Returns the source package ID. */
 	FPackageId GetSourcePackageId() const
 	{
-		return SourcePackageName.IsNone() ? FPackageId() : FPackageId::FromName(SourcePackageName);
+		return SourcePackageName.IsNone() ? FPackageId() : FPackageId::FromName(SourcePackageName, IsOptional());
 	}
 
 	FName GetSourcePackageName() const
@@ -104,6 +105,12 @@ struct FPackageStoreEntryResource
 		return EnumHasAnyFlags(Flags, EPackageStoreEntryFlags::Redirected); 
 	}
 
+	/** Returns whether this package is optional. */
+	bool IsOptional() const
+	{
+		return EnumHasAnyFlags(Flags, EPackageStoreEntryFlags::Optional);
+	}
+
 	CORE_API friend FArchive& operator<<(FArchive& Ar, FPackageStoreEntryResource& PackageStoreEntry);
 	
 	CORE_API friend FCbWriter& operator<<(FCbWriter& Writer, const FPackageStoreEntryResource& PackageStoreEntry);
