IOStoreUtilities checks for the existance of a development asset registry before trying to load it

diff --git a/Engine/Source/Developer/IoStoreUtilities/Private/IoStoreUtilities.cpp b/Engine/Source/Developer/IoStoreUtilities/Private/IoStoreUtilities.cpp
--- a/Engine/Source/Developer/IoStoreUtilities/Private/IoStoreUtilities.cpp
+++ b/Engine/Source/Developer/IoStoreUtilities/Private/IoStoreUtilities.cpp
@@ -5541,7 +5541,7 @@ int32 CreateIoStoreContainerFiles(const TCHAR* CmdLine)
 			UE_LOG(LogIoStore, Display, TEXT("Based on release version path: '%s'"), *BasedOnReleaseVersionPath);
 			FString DevelopmentAssetRegistryPath = FPaths::Combine(BasedOnReleaseVersionPath, TEXT("Metadata"), GetDevelopmentAssetRegistryFilename());
 			FArrayReader SerializedAssetData;
-			if (FFileHelper::LoadFileToArray(SerializedAssetData, *DevelopmentAssetRegistryPath))
+			if (FPaths::FileExists(*DevelopmentAssetRegistryPath) && FFileHelper::LoadFileToArray(SerializedAssetData, *DevelopmentAssetRegistryPath))
 			{
 				FAssetRegistryState ReleaseAssetRegistry;
 				FAssetRegistrySerializationOptions Options;
