Fix crash in iostore describe command when reading uncompressed data.
Fix broken import lookup logic.

#rnx

diff --git a/Engine/Source/Developer/IoStoreUtilities/Private/IoStoreUtilities.cpp b/Engine/Source/Developer/IoStoreUtilities/Private/IoStoreUtilities.cpp
--- a/Engine/Source/Developer/IoStoreUtilities/Private/IoStoreUtilities.cpp
+++ b/Engine/Source/Developer/IoStoreUtilities/Private/IoStoreUtilities.cpp
@@ -3370,14 +3370,14 @@ int32 Describe(
 		FPackageDesc* Package = nullptr;
 		FName Name;
 		FName FullName;
-		uint32 ExportHash;
+		uint64 PublicExportHash;
 		FPackageObjectIndex OuterIndex;
 		FPackageObjectIndex ClassIndex;
 		FPackageObjectIndex SuperIndex;
 		FPackageObjectIndex TemplateIndex;
 		uint64 SerialOffset = 0;
 		uint64 SerialSize = 0;
-		FSHAHash Hash;
+		FSHAHash ExportHash;
 	};
 
 	struct FExportBundleEntryDesc
@@ -3693,7 +3693,7 @@ int32 Describe(
 			ExportDesc.ClassIndex = ExportMapEntry.ClassIndex;
 			ExportDesc.SuperIndex = ExportMapEntry.SuperIndex;
 			ExportDesc.TemplateIndex = ExportMapEntry.TemplateIndex;
-			ExportDesc.ExportHash = ExportMapEntry.PublicExportHash;
+			ExportDesc.PublicExportHash = ExportMapEntry.PublicExportHash;
 			ExportDesc.SerialSize = ExportMapEntry.CookedSerialSize;
 		}
 
@@ -3721,7 +3721,7 @@ int32 Describe(
 					if (bIncludeExportHashes)
 					{
 						check(EntryDesc.Export->SerialOffset + EntryDesc.Export->SerialSize <= IoBuffer.ValueOrDie().DataSize());
-						FSHA1::HashBuffer(IoBuffer.ValueOrDie().Data() + EntryDesc.Export->SerialOffset, EntryDesc.Export->SerialSize, EntryDesc.Export->Hash.Hash);
+						FSHA1::HashBuffer(IoBuffer.ValueOrDie().Data() + EntryDesc.Export->SerialOffset, EntryDesc.Export->SerialSize, EntryDesc.Export->ExportHash.Hash);
 					}
 				}
 				++BundleEntry;
@@ -3738,9 +3738,9 @@ int32 Describe(
 		{
 			for (FExportDesc& ExportDesc : PackageDesc->Exports)
 			{
-				if (ExportDesc.ExportHash)
+				if (ExportDesc.PublicExportHash)
 				{
-					FPublicExportKey Key = FPublicExportKey::MakeKey(PackageDesc->PackageId, ExportDesc.ExportHash);
+					FPublicExportKey Key = FPublicExportKey::MakeKey(PackageDesc->PackageId, ExportDesc.PublicExportHash);
 					ExportByKeyMap.Add(Key, &ExportDesc);
 				}
 			}
@@ -4041,13 +4041,14 @@ int32 Describe(
 				OutputOverride->Logf(ELogVerbosity::Display, TEXT("\t\t       ClassIndex: %s"), *PackageObjectIndexToString(PackageDesc, Export.ClassIndex, true));
 				OutputOverride->Logf(ELogVerbosity::Display, TEXT("\t\t       SuperIndex: %s"), *PackageObjectIndexToString(PackageDesc, Export.SuperIndex, true));
 				OutputOverride->Logf(ELogVerbosity::Display, TEXT("\t\t    TemplateIndex: %s"), *PackageObjectIndexToString(PackageDesc, Export.TemplateIndex, true));
-				OutputOverride->Logf(ELogVerbosity::Display, TEXT("\t\t       ExportHash: %d"), Export.ExportHash);
-				OutputOverride->Logf(ELogVerbosity::Display, TEXT("\t\t           Offset: %lld"), Export.SerialOffset);
-				OutputOverride->Logf(ELogVerbosity::Display, TEXT("\t\t             Size: %lld"), Export.SerialSize);
+				OutputOverride->Logf(ELogVerbosity::Display, TEXT("\t\t PublicExportHash: %llu"), Export.PublicExportHash);
 				if (bIncludeExportHashes)
 				{
-					OutputOverride->Logf(ELogVerbosity::Display, TEXT("\t\t             Hash: %s"), *Export.Hash.ToString());
+					OutputOverride->Logf(ELogVerbosity::Display, TEXT("\t\t   ExportHash: %s"), *Export.ExportHash.ToString());
 				}
+				OutputOverride->Logf(ELogVerbosity::Display, TEXT("\t\t           Offset: %lld"), Export.SerialOffset);
+				OutputOverride->Logf(ELogVerbosity::Display, TEXT("\t\t             Size: %lld"), Export.SerialSize);
+
 			}
 
 			OutputOverride->Logf(ELogVerbosity::Display, TEXT("--------------------------------------------"));
diff --git a/Engine/Source/Runtime/Core/Private/IO/IoStore.cpp b/Engine/Source/Runtime/Core/Private/IO/IoStore.cpp
--- a/Engine/Source/Runtime/Core/Private/IO/IoStore.cpp
+++ b/Engine/Source/Runtime/Core/Private/IO/IoStore.cpp
@@ -2137,6 +2137,7 @@ public:
 		uint64 UncompressedDestinationOffset = 0;
 		uint64 OffsetInBlock = ResolvedOffset % CompressionBlockSize;
 		uint64 RemainingSize = ResolvedSize;
+		TArray<uint8> TempBuffer;
 		for (int32 BlockIndex = FirstBlockIndex; BlockIndex <= LastBlockIndex; ++BlockIndex)
 		{
 			const FIoStoreTocCompressedBlockEntry& CompressionBlock = TocResource.CompressionBlocks[BlockIndex];
@@ -2165,15 +2166,14 @@ public:
 			const uint32 UncompressedSize = CompressionBlock.GetUncompressedSize();
 			if (CompressionMethod.IsNone())
 			{
-				check(UncompressedDestination + UncompressedSize - OffsetInBlock <= UncompressedBuffer.Data() + UncompressedBuffer.DataSize());
-				FMemory::Memcpy(UncompressedDestination, CompressedBuffer.GetData() + OffsetInBlock, UncompressedSize - OffsetInBlock);
+				uint64 CopySize = FMath::Min<uint64>(UncompressedSize - OffsetInBlock, RemainingSize);
+				FMemory::Memcpy(UncompressedDestination, CompressedBuffer.GetData() + OffsetInBlock, CopySize);
 			}
 			else
 			{
 				bool bUncompressed;
 				if (OffsetInBlock || RemainingSize < UncompressedSize)
 				{
-					TArray<uint8> TempBuffer;
 					TempBuffer.SetNumUninitialized(UncompressedSize);
 					bUncompressed = FCompression::UncompressMemory(CompressionMethod, TempBuffer.GetData(), UncompressedSize, CompressedBuffer.GetData(), CompressionBlock.GetCompressedSize());
 					uint64 CopySize = FMath::Min<uint64>(UncompressedSize - OffsetInBlock, RemainingSize);
