ZenCook:
Removed several unnecessary files from the file system manifest
Added IoStoreCmdlet utility command for generating the file system manifest
#rnx

diff --git a/Engine/Source/Developer/IoStoreUtilities/IoStoreUtilities.Build.cs b/Engine/Source/Developer/IoStoreUtilities/IoStoreUtilities.Build.cs
--- a/Engine/Source/Developer/IoStoreUtilities/IoStoreUtilities.Build.cs
+++ b/Engine/Source/Developer/IoStoreUtilities/IoStoreUtilities.Build.cs
@@ -27,5 +27,7 @@ public class IoStoreUtilities : ModuleRules
 		PrivateDependencyModuleNames.Add("PakFile");
         PrivateDependencyModuleNames.Add("Json");
         PrivateDependencyModuleNames.Add("RSA");
+        PrivateDependencyModuleNames.Add("DeveloperToolSettings");
+        PrivateDependencyModuleNames.Add("SandboxFile");
 	}
 }
diff --git a/Engine/Source/Developer/IoStoreUtilities/Private/IoStoreUtilities.cpp b/Engine/Source/Developer/IoStoreUtilities/Private/IoStoreUtilities.cpp
--- a/Engine/Source/Developer/IoStoreUtilities/Private/IoStoreUtilities.cpp
+++ b/Engine/Source/Developer/IoStoreUtilities/Private/IoStoreUtilities.cpp
@@ -60,6 +60,8 @@
 #include "IO/IoContainerHeader.h"
 #include "ProfilingDebugging/CountersTrace.h"
 #include "IO/IoStore.h"
+#include "ZenFileSystemManifest.h"
+#include "IPlatformFileSandboxWrapper.h"
 
 //PRAGMA_DISABLE_OPTIMIZATION
 
@@ -4737,6 +4739,25 @@ int32 Staged2Zen(const FString& BuildPath, const FKeyChain& KeyChain, const FStr
 	return 0;
 }
 
+int32 GenerateZenFileSystemManifest(ITargetPlatform* TargetPlatform)
+{
+	FString OutputDirectory = FPaths::Combine(*FPaths::ProjectDir(), TEXT("Saved"), TEXT("Cooked"), TEXT("[Platform]"));
+	OutputDirectory = FPaths::ConvertRelativePathToFull(OutputDirectory);
+	FPaths::NormalizeDirectoryName(OutputDirectory);
+	TUniquePtr<FSandboxPlatformFile> LocalSandboxFile = FSandboxPlatformFile::Create(false);
+	LocalSandboxFile->Initialize(&FPlatformFileManager::Get().GetPlatformFile(), *FString::Printf(TEXT("-sandbox=\"%s\""), *OutputDirectory));
+	const FString RootPathSandbox = LocalSandboxFile->ConvertToAbsolutePathForExternalAppForWrite(*FPaths::RootDir());
+	FString MetadataPathSandbox = LocalSandboxFile->ConvertToAbsolutePathForExternalAppForWrite(*(FPaths::ProjectDir() / TEXT("Metadata")));
+	const FString PlatformString = TargetPlatform->PlatformName();
+	const FString ResolvedRootPath = RootPathSandbox.Replace(TEXT("[Platform]"), *PlatformString);
+	const FString ResolvedMetadataPath = MetadataPathSandbox.Replace(TEXT("[Platform]"), *PlatformString);
+
+	FZenFileSystemManifest ZenFileSystemManifest(*TargetPlatform, ResolvedRootPath);
+	ZenFileSystemManifest.Generate();
+	ZenFileSystemManifest.Save(*FPaths::Combine(ResolvedMetadataPath, TEXT("zenfs.manifest")));
+	return 0;
+}
+
 bool ExtractFilesFromIoStoreContainer(
 	const TCHAR* InContainerFilename,
 	const TCHAR* InDestPath,
@@ -5409,6 +5430,19 @@ int32 CreateIoStoreContainerFiles(const TCHAR* CmdLine)
 
 	LoadKeyChain(FCommandLine::Get(), Arguments.KeyChain);
 	
+	ITargetPlatform* TargetPlatform = nullptr;
+	FString TargetPlatformName;
+	if (FParse::Value(FCommandLine::Get(), TEXT("TargetPlatform="), TargetPlatformName))
+	{
+		ITargetPlatformManagerModule& TPM = GetTargetPlatformManagerRef();
+		TargetPlatform = TPM.FindTargetPlatform(TargetPlatformName);
+		if (!TargetPlatform)
+		{
+			UE_LOG(LogIoStore, Error, TEXT("Invalid TargetPlatform: '%s'"), *TargetPlatformName);
+			return 1;
+		}
+	}
+
 	FString ArgumentValue;
 	if (FParse::Value(FCommandLine::Get(), TEXT("List="), ArgumentValue))
 	{
@@ -5478,19 +5512,6 @@ int32 CreateIoStoreContainerFiles(const TCHAR* CmdLine)
 	}
 	else if (FParse::Param(FCommandLine::Get(), TEXT("Staged2Zen")))
 	{
-		ITargetPlatform* TargetPlatform = nullptr;
-		FString TargetPlatformName;
-		if (FParse::Value(FCommandLine::Get(), TEXT("TargetPlatform="), TargetPlatformName))
-		{
-			ITargetPlatformManagerModule& TPM = GetTargetPlatformManagerRef();
-			TargetPlatform = TPM.FindTargetPlatform(TargetPlatformName);
-			if (!TargetPlatform)
-			{
-				UE_LOG(LogIoStore, Error, TEXT("Invalid TargetPlatform: '%s'"), *TargetPlatformName);
-				return 1;
-			}
-		}
-
 		FString BuildPath;
 		FString ProjectName;
 		if (!FParse::Value(FCommandLine::Get(), TEXT("BuildPath="), BuildPath) ||
@@ -5498,6 +5519,7 @@ int32 CreateIoStoreContainerFiles(const TCHAR* CmdLine)
 			!TargetPlatform)
 		{
 			UE_LOG(LogIoStore, Error, TEXT("Incorrect arguments. Expected: -Staged2Zen -BuildPath=<Path> -ProjectName=<ProjectName> -TargetPlatform=<Platform>"));
+			return -1;
 		}
 		return Staged2Zen(BuildPath, Arguments.KeyChain, ProjectName, TargetPlatform);
 	}
@@ -5519,6 +5541,15 @@ int32 CreateIoStoreContainerFiles(const TCHAR* CmdLine)
 
 		return CreateContentPatch(Arguments, WriterSettings);
 	}
+	else if (FParse::Param(FCommandLine::Get(), TEXT("GenerateZenFileSystemManifest")))
+	{
+		if (!TargetPlatform)
+		{
+			UE_LOG(LogIoStore, Error, TEXT("Incorrect arguments. Expected: -GenerateZenFileSystemManifest -TargetPlatform=<Platform>"));
+			return -11;
+		}
+		return GenerateZenFileSystemManifest(TargetPlatform);
+	}
 	else if (FParse::Value(FCommandLine::Get(), TEXT("CreateDLCContainer="), Arguments.DLCPluginPath))
 	{
 		if (!ParseContainerGenerationArguments(Arguments, WriterSettings))
diff --git a/Engine/Source/Developer/IoStoreUtilities/Private/ZenFileSystemManifest.cpp b/Engine/Source/Developer/IoStoreUtilities/Private/ZenFileSystemManifest.cpp
--- a/Engine/Source/Developer/IoStoreUtilities/Private/ZenFileSystemManifest.cpp
+++ b/Engine/Source/Developer/IoStoreUtilities/Private/ZenFileSystemManifest.cpp
@@ -7,6 +7,10 @@
 #include "Misc/FileHelper.h"
 #include "Misc/Paths.h"
 #include "Misc/DataDrivenPlatformInfoRegistry.h"
+#include "Misc/PathViews.h"
+#include "Settings/ProjectPackagingSettings.h"
+
+DEFINE_LOG_CATEGORY_STATIC(LogZenFileSystemManifest, Display, All);
 
 const FZenFileSystemManifest::FManifestEntry FZenFileSystemManifest::InvalidEntry = FManifestEntry();
 
@@ -20,49 +24,104 @@ FZenFileSystemManifest::FZenFileSystemManifest(const ITargetPlatform& InTargetPl
 	FPaths::NormalizeDirectoryName(ServerRoot);
 }
 
-int32 FZenFileSystemManifest::Generate()
+void FZenFileSystemManifest::GetExtensionDirs(TArray<FString>& OutExtensionDirs, const TCHAR* BaseDir, const TCHAR* SubDir, const TArray<FString>& PlatformDirectoryNames)
 {
-	TRACE_CPUPROFILER_EVENT_SCOPE(GenerateStorageServerFileSystemManifest);
+	auto AddIfDirectoryExists = [&OutExtensionDirs](FString&& Dir)
+	{
+		if (FPaths::DirectoryExists(Dir))
+		{
+			OutExtensionDirs.Emplace(MoveTemp(Dir));
+		}
+	};
+	
+	AddIfDirectoryExists(FPaths::Combine(BaseDir, SubDir));
 
-	struct FFileFilter
+	FString PlatformExtensionBaseDir = FPaths::Combine(BaseDir, TEXT("Platforms"));
+	for (const FString& PlatformDirectoryName : PlatformDirectoryNames)
 	{
-		bool IncludeDirectory(const TCHAR* Path) const
+		AddIfDirectoryExists(FPaths::Combine(PlatformExtensionBaseDir, PlatformDirectoryName, SubDir));
+	}
+
+	FString RestrictedBaseDir = FPaths::Combine(BaseDir, TEXT("Restricted"));
+	IFileManager::Get().IterateDirectory(*RestrictedBaseDir, [&OutExtensionDirs, SubDir, &PlatformDirectoryNames](const TCHAR* FilenameOrDirectory, bool bIsDirectory) -> bool
 		{
-			if (FCString::Stristr(Path, TEXT("/Intermediate")))
+			if (bIsDirectory)
 			{
-				return false;
+				GetExtensionDirs(OutExtensionDirs, FilenameOrDirectory, SubDir, PlatformDirectoryNames);
 			}
-			if (FCString::Stristr(Path, TEXT("/Binaries")))
-			{
-				return false;
-			}
-			if (FCString::Stristr(Path, TEXT("/Source")))
+			return true;
+		});
+}
+
+int32 FZenFileSystemManifest::Generate()
+{
+	TRACE_CPUPROFILER_EVENT_SCOPE(GenerateStorageServerFileSystemManifest);
+
+	class FFileFilter
+	{
+	public:
+		FFileFilter& ExcludeDirectory(const TCHAR* Name)
+		{
+			DirectoryExclusionFilter.Emplace(Name);
+			return *this;
+		}
+
+		FFileFilter& ExcludeExtension(const TCHAR* Extension)
+		{
+			ExtensionExclusionFilter.Emplace(Extension);
+			return *this;
+		}
+
+		FFileFilter& IncludeExtension(const TCHAR* Extension)
+		{
+			ExtensionInclusionFilter.Emplace(Extension);
+			return *this;
+		}
+
+		bool FilterDirectory(FStringView Name) const
+		{
+			for (const FString& ExcludedDirectory : DirectoryExclusionFilter)
 			{
-				return false;
+				if (Name == ExcludedDirectory)
+				{
+					return false;
+				}
 			}
 			return true;
 		}
 
-		bool IncludeFile(const TCHAR* Path) const
+		bool FilterFile(FStringView Extension)
 		{
-			if (FCString::Stristr(Path, TEXT("/Content/")))
+			if (!ExtensionExclusionFilter.IsEmpty())
 			{
-				const TCHAR* Extension = FCString::Strrchr(Path, '.');
-				if (Extension)
+				for (const FString& ExcludedExtension : ExtensionExclusionFilter)
 				{
-					if (!FCString::Stricmp(Extension, TEXT(".uasset")) ||
-						!FCString::Stricmp(Extension, TEXT(".umap")) ||
-						!FCString::Stricmp(Extension, TEXT(".uexp")) ||
-						!FCString::Stricmp(Extension, TEXT(".ubulk")))
+					if (Extension == ExcludedExtension)
 					{
 						return false;
 					}
 				}
 			}
+			if (!ExtensionInclusionFilter.IsEmpty())
+			{
+				for (const FString& IncludedExtension : ExtensionInclusionFilter)
+				{
+					if (Extension == IncludedExtension)
+					{
+						return true;
+					}
+				}
+				return false;
+			}
 			return true;
 		}
-	} Filter;
-	
+
+	private:
+		TArray<FString> DirectoryExclusionFilter;
+		TArray<FString> ExtensionExclusionFilter;
+		TArray<FString> ExtensionInclusionFilter;
+	};
+
 	IPlatformFile& PlatformFile = FPlatformFileManager::Get().GetPlatformFile();
 
 	FString RootDir = FPaths::RootDir();
@@ -71,94 +130,190 @@ int32 FZenFileSystemManifest::Generate()
 	FString ProjectDir = FPaths::ProjectDir();
 	FPaths::NormalizeDirectoryName(ProjectDir);
 
+	FFileFilter BaseFilter = FFileFilter()
+		.ExcludeDirectory(TEXT("Binaries"))
+		.ExcludeDirectory(TEXT("Intermediate"))
+		.ExcludeDirectory(TEXT("Saved"))
+		.ExcludeDirectory(TEXT("Source"));
 	auto AddFilesFromDirectory =
-		[this, &PlatformFile, &RootDir]
-		(const FString& ClientDirectory, const FString& LocalDirectory, bool bIncludeSubdirs, FFileFilter* Filter = nullptr)
+		[this, &PlatformFile, &RootDir, &BaseFilter]
+		(const FString& ClientDirectory, const FString& LocalDirectory, bool bIncludeSubdirs, FFileFilter* AdditionalFilter = nullptr)
 		{
+			//TRACE_BOOKMARK(TEXT("AddFilesFromDirectory: %s"), *ClientDirectory);
 			FString ServerRelativeDirectory = LocalDirectory;
 			FPaths::MakePathRelativeTo(ServerRelativeDirectory, *RootDir);
 			ServerRelativeDirectory = TEXT("/") + ServerRelativeDirectory;
 
 			TArray<FString> DirectoriesToVisit;
 			auto VisitorFunc =
-				[this, &DirectoriesToVisit, &RootDir, &ClientDirectory, &LocalDirectory, &ServerRelativeDirectory, bIncludeSubdirs, Filter]
+				[this, &DirectoriesToVisit, &RootDir, &ClientDirectory, &LocalDirectory, &ServerRelativeDirectory, bIncludeSubdirs, &BaseFilter, AdditionalFilter]
 				(const TCHAR* InnerFileNameOrDirectory, bool bIsDirectory)
 				{
-					if (Filter)
+					if (bIsDirectory)
 					{
-						if ((bIsDirectory && !Filter->IncludeDirectory(InnerFileNameOrDirectory)) ||
-							(!bIsDirectory && !Filter->IncludeFile(InnerFileNameOrDirectory)))
+						if (!bIncludeSubdirs)
 						{
 							return true;
 						}
-					}
-					if (bIsDirectory)
-					{
-						if (bIncludeSubdirs)
+						FStringView DirectoryName = FPathViews::GetPathLeaf(InnerFileNameOrDirectory);
+						if (!BaseFilter.FilterDirectory(DirectoryName))
 						{
-							DirectoriesToVisit.Add(InnerFileNameOrDirectory);
+							return true;
 						}
+						if (AdditionalFilter && !AdditionalFilter->FilterDirectory(DirectoryName))
+						{
+							return true;
+						}
+						DirectoriesToVisit.Add(InnerFileNameOrDirectory);
+						return true;
 					}
-					else
+					FStringView Extension = FPathViews::GetExtension(InnerFileNameOrDirectory);
+					if (!BaseFilter.FilterFile(Extension))
 					{
-						FStringView RelativePath = InnerFileNameOrDirectory;
-						RelativePath.RightChopInline(LocalDirectory.Len() + 1);
-						FString ClientPath = FPaths::Combine(ClientDirectory, RelativePath.GetData());
-						const FIoChunkId FileChunkId = CreateExternalFileChunkId(ClientPath);
-
-						AddManifestEntry(
-							FileChunkId,
-							FPaths::Combine(ServerRelativeDirectory, RelativePath.GetData()),
-							MoveTemp(ClientPath));
+						return true;
 					}
-
+					if (AdditionalFilter && !AdditionalFilter->FilterFile(Extension))
+					{
+						return true;
+					}
+					//TRACE_CPUPROFILER_EVENT_SCOPE(AddManifestEntry);
+					FStringView RelativePath = InnerFileNameOrDirectory;
+					RelativePath.RightChopInline(LocalDirectory.Len() + 1);
+					FString ClientPath = FPaths::Combine(ClientDirectory, RelativePath.GetData());
+					const FIoChunkId FileChunkId = CreateExternalFileChunkId(ClientPath);
+					AddManifestEntry(
+						FileChunkId,
+						FPaths::Combine(ServerRelativeDirectory, RelativePath.GetData()),
+						MoveTemp(ClientPath));
 					return true;
 				};
 
 			DirectoriesToVisit.Push(LocalDirectory);
 			while (!DirectoriesToVisit.IsEmpty())
 			{
-				PlatformFile.IterateDirectory(*DirectoriesToVisit.Pop(), VisitorFunc);
+				PlatformFile.IterateDirectory(*DirectoriesToVisit.Pop(false), VisitorFunc);
+			}
+		};
+
+	const UProjectPackagingSettings* const PackagingSettings = GetDefault<UProjectPackagingSettings>();
+
+	TArray<FString> PlatformDirectoryNames;
+	const FDataDrivenPlatformInfo& PlatformInfo = FDataDrivenPlatformInfoRegistry::GetPlatformInfo(TargetPlatform.IniPlatformName());
+	PlatformDirectoryNames.Reserve(PlatformInfo.IniParentChain.Num() + PlatformInfo.AdditionalRestrictedFolders.Num() + 1);
+	PlatformDirectoryNames.Add(TargetPlatform.IniPlatformName());
+	for (const FString& PlatformName : PlatformInfo.AdditionalRestrictedFolders)
+	{
+		PlatformDirectoryNames.AddUnique(PlatformName);
+	}
+	for (const FString& PlatformName : PlatformInfo.IniParentChain)
+	{
+		PlatformDirectoryNames.AddUnique(PlatformName);
+	}
+	auto AddFilesFromExtensionDirectories = [&AddFilesFromDirectory, &EngineDir, &ProjectDir, &PlatformDirectoryNames](const TCHAR* ExtensionSubDir, FFileFilter* AdditionalFilter = nullptr)
+		{
+			TArray<FString> ExtensionDirs;
+			GetExtensionDirs(ExtensionDirs, *EngineDir, ExtensionSubDir, PlatformDirectoryNames);
+			for (const FString& Dir : ExtensionDirs)
+			{
+				//AddFilesFromDirectory(Dir.Replace(*EngineDir, TEXT("/{engine}")), Dir, true, AdditionalFilter);
+				AddFilesFromDirectory(Dir.Replace(*EngineDir, TEXT("/{engine}")), Dir, true, AdditionalFilter);
+			}
+			ExtensionDirs.Reset();
+			GetExtensionDirs(ExtensionDirs, *ProjectDir, ExtensionSubDir, PlatformDirectoryNames);
+			for (const FString& Dir : ExtensionDirs)
+			{
+				AddFilesFromDirectory(Dir.Replace(*ProjectDir, TEXT("/{project}")), Dir, true, AdditionalFilter);
 			}
 		};
 
 	const int32 PreviousEntryCount = NumEntries();
 
-	AddFilesFromDirectory(TEXT("/{engine}"), FPaths::Combine(CookDirectory, TEXT("Engine")), true);
-	AddFilesFromDirectory(TEXT("/{project}"), FPaths::Combine(CookDirectory, FApp::GetProjectName()), true);
+	FFileFilter CookedFilter = FFileFilter()
+		.ExcludeDirectory(TEXT("Metadata"))
+		.ExcludeExtension(TEXT("json"))
+		.ExcludeExtension(TEXT("uasset"))
+		.ExcludeExtension(TEXT("ubulk"))
+		.ExcludeExtension(TEXT("uexp"))
+		.ExcludeExtension(TEXT("umap"))
+		.ExcludeExtension(TEXT("uregs"));
+	AddFilesFromDirectory(TEXT("/{engine}"), FPaths::Combine(CookDirectory, TEXT("Engine")), true, &CookedFilter);
+	AddFilesFromDirectory(TEXT("/{project}"), FPaths::Combine(CookDirectory, FApp::GetProjectName()), true, &CookedFilter);
+	
 	AddFilesFromDirectory(TEXT("/{project}"), ProjectDir, false);
-	AddFilesFromDirectory(TEXT("/{engine}/Config"), FPaths::Combine(EngineDir, TEXT("Config")), true);
-	AddFilesFromDirectory(TEXT("/{project}/Config"), FPaths::Combine(ProjectDir, TEXT("Config")), true);
-	AddFilesFromDirectory(TEXT("/{engine}/Plugins"), FPaths::Combine(EngineDir, TEXT("Plugins")), true, &Filter);
-	AddFilesFromDirectory(TEXT("/{project}/Plugins"), FPaths::Combine(ProjectDir, TEXT("Plugins")), true, &Filter);
-	AddFilesFromDirectory(TEXT("/{engine}/Restricted"), FPaths::Combine(EngineDir, TEXT("Restricted")), true, &Filter);
-	AddFilesFromDirectory(TEXT("/{project}/Restricted"), FPaths::Combine(ProjectDir, TEXT("Restricted")), true, &Filter);
-	AddFilesFromDirectory(TEXT("/{engine}/Content/Internationalization"), FPaths::Combine(EngineDir, TEXT("Content"), TEXT("Internationalization")), true);
-	AddFilesFromDirectory(TEXT("/{engine}/Content/Localization"), FPaths::Combine(EngineDir, TEXT("Content"), TEXT("Localization")), true);
-	AddFilesFromDirectory(TEXT("/{project}/Content/Localization"), FPaths::Combine(ProjectDir, TEXT("Content"), TEXT("Localization")), true);
-	AddFilesFromDirectory(TEXT("/{engine}/Content/Slate"), FPaths::Combine(EngineDir, TEXT("Content"), TEXT("Slate")), true);
-	AddFilesFromDirectory(TEXT("/{project}/Content/Slate"), FPaths::Combine(ProjectDir, TEXT("Content"), TEXT("Slate")), true);
-	AddFilesFromDirectory(TEXT("/{engine}/Content/Movies"), FPaths::Combine(EngineDir, TEXT("Content"), TEXT("Movies")), true);
-	AddFilesFromDirectory(TEXT("/{project}/Content/Movies"), FPaths::Combine(ProjectDir, TEXT("Content"), TEXT("Movies")), true);
-
-	const FDataDrivenPlatformInfo& Info = FDataDrivenPlatformInfoRegistry::GetPlatformInfo(TargetPlatform.IniPlatformName());
-
-	TArray<FString> PlatformDirectories;
-	PlatformDirectories.Reserve(Info.IniParentChain.Num() + Info.AdditionalRestrictedFolders.Num() + 1);
-	PlatformDirectories.Add(TargetPlatform.IniPlatformName());
-	for (const FString& PlatformName : Info.AdditionalRestrictedFolders)
-	{
-		PlatformDirectories.AddUnique(PlatformName);
+	
+	FFileFilter ConfigFilter = FFileFilter()
+		.IncludeExtension(TEXT("ini"));
+	AddFilesFromExtensionDirectories(TEXT("Config"), &ConfigFilter);
+
+	FFileFilter PluginFilter = FFileFilter()
+		.IncludeExtension(TEXT("ini"))
+		.IncludeExtension(TEXT("locmeta"))
+		.IncludeExtension(TEXT("locres"))
+		.IncludeExtension(TEXT("uplugin"));
+	AddFilesFromExtensionDirectories(TEXT("Plugins"), &PluginFilter);
+
+	FString InternationalizationPresetAsString = UEnum::GetValueAsString(PackagingSettings->InternationalizationPreset);
+	const TCHAR* InternationalizationPresetPath = FCString::Strrchr(*InternationalizationPresetAsString, ':');
+	if (InternationalizationPresetPath)
+	{
+		++InternationalizationPresetPath;
 	}
-	for (const FString& PlatformName : Info.IniParentChain)
+	else
 	{
-		PlatformDirectories.AddUnique(PlatformName);
+		UE_LOG(LogZenFileSystemManifest, Warning, TEXT("Failed reading internationalization preset setting, defaulting to English"));
+		InternationalizationPresetPath = TEXT("English");
 	}
+	const TCHAR* ICUDataVersion = TEXT("icudt64l"); // TODO: Could this go into datadriven platform info? But it's basically always this.
+	AddFilesFromDirectory(*FPaths::Combine(TEXT("/{engine}"), TEXT("Content"), TEXT("Internationalization"), ICUDataVersion), FPaths::Combine(EngineDir, TEXT("Content"), TEXT("Internationalization"), InternationalizationPresetPath, ICUDataVersion), true);
+	
+	FFileFilter LocalizationFilter = FFileFilter()
+		.IncludeExtension(TEXT("locmeta"))
+		.IncludeExtension(TEXT("locres"));
+	AddFilesFromExtensionDirectories(TEXT("Content/Localization"), &LocalizationFilter);
+
+	FFileFilter ContentFilter = FFileFilter()
+		.ExcludeExtension(TEXT("uasset"))
+		.ExcludeExtension(TEXT("ubulk"))
+		.ExcludeExtension(TEXT("uexp"))
+		.ExcludeExtension(TEXT("umap"));
+	AddFilesFromDirectory(TEXT("/{engine}/Content/Slate"), FPaths::Combine(EngineDir, TEXT("Content"), TEXT("Slate")), true, &ContentFilter);
+	AddFilesFromDirectory(TEXT("/{project}/Content/Slate"), FPaths::Combine(ProjectDir, TEXT("Content"), TEXT("Slate")), true, &ContentFilter);
+	AddFilesFromDirectory(TEXT("/{engine}/Content/Movies"), FPaths::Combine(EngineDir, TEXT("Content"), TEXT("Movies")), true, &ContentFilter);
+	AddFilesFromDirectory(TEXT("/{project}/Content/Movies"), FPaths::Combine(ProjectDir, TEXT("Content"), TEXT("Movies")), true, &ContentFilter);
+	
+	FFileFilter ShaderCacheFilter = FFileFilter()
+		.IncludeExtension(TEXT("ushadercache"))
+		.IncludeExtension(TEXT("upipelinecache"));
+	AddFilesFromDirectory(TEXT("/{project}/Content"), FPaths::Combine(ProjectDir, TEXT("Content")), false, &ShaderCacheFilter);
+	AddFilesFromDirectory(FPaths::Combine(TEXT("/{project}"), TEXT("Content"), TEXT("PipelineCaches"), TargetPlatform.IniPlatformName()), FPaths::Combine(ProjectDir, TEXT("Content"), TEXT("PipelineCaches"), TargetPlatform.IniPlatformName()), false, &ShaderCacheFilter);
 
-	for (const FString& PlatformDirectory : PlatformDirectories)
+	auto AddAdditionalFilesFromConfig = [&AddFilesFromDirectory , &ContentFilter, &ProjectDir, &EngineDir](const FString& RelativeDirToStage)
+	{
+		FString AbsoluteDirToStage = FPaths::ConvertRelativePathToFull(FPaths::Combine(ProjectDir, TEXT("Content"), RelativeDirToStage));
+		FPaths::NormalizeDirectoryName(AbsoluteDirToStage);
+		FString AbsoluteEngineDir = FPaths::ConvertRelativePathToFull(EngineDir);
+		FString AbsoluteProjectDir = FPaths::ConvertRelativePathToFull(ProjectDir);
+		FStringView RelativeToKnownRootView;
+		if (FPathViews::TryMakeChildPathRelativeTo(AbsoluteDirToStage, AbsoluteProjectDir, RelativeToKnownRootView))
+		{
+			AddFilesFromDirectory(FPaths::Combine(TEXT("/{project}"), RelativeToKnownRootView.GetData()), FPaths::Combine(ProjectDir, RelativeToKnownRootView.GetData()), true, &ContentFilter);
+		}
+		else if (FPathViews::TryMakeChildPathRelativeTo(AbsoluteDirToStage, AbsoluteEngineDir, RelativeToKnownRootView))
+		{
+			AddFilesFromDirectory(FPaths::Combine(TEXT("/{engine}"), RelativeToKnownRootView.GetData()), FPaths::Combine(EngineDir, RelativeToKnownRootView.GetData()), true, &ContentFilter);
+		}
+		else
+		{
+			UE_LOG(LogZenFileSystemManifest, Warning, TEXT("Ignoring additional folder to stage that is not relative to the engine or project directory: %s"), *RelativeDirToStage);
+		}
+	};
+	for (const FDirectoryPath& AdditionalFolderToStage : PackagingSettings->DirectoriesToAlwaysStageAsUFS)
+	{
+		AddAdditionalFilesFromConfig(AdditionalFolderToStage.Path);
+	}
+	for (const FDirectoryPath& AdditionalFolderToStage : PackagingSettings->DirectoriesToAlwaysStageAsUFSServer)
 	{
-		AddFilesFromDirectory(FPaths::Combine(TEXT("/{engine}/Platforms"), PlatformDirectory), FPaths::Combine(EngineDir, TEXT("Platforms"), PlatformDirectory), true, &Filter);
-		AddFilesFromDirectory(FPaths::Combine(TEXT("/{project}/Platforms"), PlatformDirectory), FPaths::Combine(ProjectDir, TEXT("Platforms"), PlatformDirectory), true, &Filter);
+		AddAdditionalFilesFromConfig(AdditionalFolderToStage.Path);
 	}
 
 	const int32 CurrentEntryCount = NumEntries();
diff --git a/Engine/Source/Developer/IoStoreUtilities/Private/ZenFileSystemManifest.h b/Engine/Source/Developer/IoStoreUtilities/Private/ZenFileSystemManifest.h
--- a/Engine/Source/Developer/IoStoreUtilities/Private/ZenFileSystemManifest.h
+++ b/Engine/Source/Developer/IoStoreUtilities/Private/ZenFileSystemManifest.h
@@ -40,6 +40,8 @@ public:
 	}
 
 private:
+	static void GetExtensionDirs(TArray<FString>& OutExtensionDirs, const TCHAR* BaseDir, const TCHAR* SubDir, const TArray<FString>& PlatformDirectoryNames);
+
 	const ITargetPlatform& TargetPlatform;
 	FString CookDirectory;
 	FString ServerRoot;
