- Enable Zen on Mac and Linux.
- Updated zenserver binaries to 4ae4f1f

#rnx

diff --git a/Engine/Source/Developer/IoStoreUtilities/Private/ZenStoreHttpClient.cpp b/Engine/Source/Developer/IoStoreUtilities/Private/ZenStoreHttpClient.cpp
--- a/Engine/Source/Developer/IoStoreUtilities/Private/ZenStoreHttpClient.cpp
+++ b/Engine/Source/Developer/IoStoreUtilities/Private/ZenStoreHttpClient.cpp
@@ -9,7 +9,7 @@
 #include "Misc/App.h"
 #include "HAL/FileManager.h"
 
-#if PLATFORM_WINDOWS
+#if PLATFORM_WINDOWS || PLATFORM_UNIX || PLATFORM_MAC
 
 #include "Async/Async.h"
 #include "HAL/PlatformFileManager.h"
@@ -604,7 +604,7 @@ FZenStoreHttpClient::EndBuildPass(FCbPackage OpEntry)
 
 } // UE
 
-#else // not PLATFORM_WINDOWS, dummy implementation stub for now
+#else // not desktop platform
 
 namespace UE {
 namespace Zen {
@@ -690,7 +690,7 @@ TFuture<TIoStatusOr<FCbObject>> FZenStoreHttpClient::GetFiles()
 
 }
 
-#endif // PLATFORM_WINDOWS
+#endif // desktop platform
 
 namespace UE
 {
diff --git a/Engine/Source/Developer/Zen/Public/ZenGlobals.h b/Engine/Source/Developer/Zen/Public/ZenGlobals.h
--- a/Engine/Source/Developer/Zen/Public/ZenGlobals.h
+++ b/Engine/Source/Developer/Zen/Public/ZenGlobals.h
@@ -5,7 +5,7 @@
 #include "HAL/Platform.h"
 
 #ifndef UE_WITH_ZEN
-#	if PLATFORM_WINDOWS
+#	if PLATFORM_WINDOWS || PLATFORM_UNIX || PLATFORM_MAC
 #		define UE_WITH_ZEN 1
 #	else
 #		define UE_WITH_ZEN 0
