Adding support for writing back metadata to an asset registry during iostore staging. Initially this is size information about the chunks.

diff --git a/Engine/Config/BaseGame.ini b/Engine/Config/BaseGame.ini
--- a/Engine/Config/BaseGame.ini
+++ b/Engine/Config/BaseGame.ini
@@ -109,6 +109,7 @@ bSharedMaterialNativeLibraries=True
 bShareMaterialShaderCode=True
 bSkipMovies=False
 bPakUsesSecondaryOrder=True
+WriteBackMetadataToAssetRegistry=Disabled
 
 +EarlyDownloaderPakFileFiles=...\Content\Internationalization\...\*.icu
 +EarlyDownloaderPakFileFiles=...\Content\Internationalization\...\*.brk
diff --git a/Engine/Source/Developer/DeveloperToolSettings/Classes/Settings/ProjectPackagingSettings.h b/Engine/Source/Developer/DeveloperToolSettings/Classes/Settings/ProjectPackagingSettings.h
--- a/Engine/Source/Developer/DeveloperToolSettings/Classes/Settings/ProjectPackagingSettings.h
+++ b/Engine/Source/Developer/DeveloperToolSettings/Classes/Settings/ProjectPackagingSettings.h
@@ -91,6 +91,28 @@ enum class EProjectPackagingBlueprintNativizationMethod : uint8
 	Exclusive
 };
 
+/**
+* The list of possible registry writebacks. During staging, iostore can
+* optionally write back data that is only available during the staging process
+* so that asset registry tools can associate this data with their respective
+* assets.
+* 
+* Note that this is used in UnrealPak and thus can't use StaticEnum<>, so if you
+* add any types here, be sure to add the parsing of the strings to IoStoreUtilities.cpp.
+*/
+UENUM()
+enum class EAssetRegistryWritebackMethod : uint8
+{
+	/** Do not write-back staging metadata to the asset registry */
+	Disabled,
+
+	/** The development asset registry from the source cooked directory will be re-used. */
+	OriginalFile,
+
+	/** A duplicate asset registry will be created with the metadata added to it, adjacent to the cooked development asset registry. */
+	AdjacentFile	
+};
+
 USTRUCT()
 struct FPakOrderFileSpec
 {
@@ -301,6 +323,12 @@ public:
 	UPROPERTY(config, EditAnywhere, Category = Packaging)
 	FDirectoryPath HttpChunkInstallDataDirectory;
 
+	/**
+	* Whether to write staging metadata back to the asset registry. This metadata contains information such as
+	* the actual compressed chunk sizes of the assets as well as some bulk data diff blame support information.
+	*/
+	UPROPERTY(config, EditAnywhere, Category = Packaging, AdvancedDisplay)
+	EAssetRegistryWritebackMethod WriteBackMetadataToAssetRegistry;
 
 	/**
 	 * Create compressed cooked packages (decreased deployment size)
diff --git a/Engine/Source/Developer/IoStoreUtilities/Private/IoStoreUtilities.cpp b/Engine/Source/Developer/IoStoreUtilities/Private/IoStoreUtilities.cpp
--- a/Engine/Source/Developer/IoStoreUtilities/Private/IoStoreUtilities.cpp
+++ b/Engine/Source/Developer/IoStoreUtilities/Private/IoStoreUtilities.cpp
@@ -27,6 +27,8 @@
 #include "Serialization/MemoryReader.h"
 #include "Serialization/AsyncLoading2.h"
 #include "Serialization/ArrayReader.h"
+#include "Serialization/ArrayWriter.h"
+#include "Settings/ProjectPackagingSettings.h" // for EAssetRegistryWritebackMethod
 #include "IO/PackageStore.h"
 #include "UObject/Class.h"
 #include "UObject/NameBatchSerialization.h"
@@ -626,6 +628,7 @@ struct FIoStoreArguments
 	bool bCreateDirectoryIndex = true;
 	bool bClusterByOrderFilePriority = false;
 	bool bFileRegions = false;
+	EAssetRegistryWritebackMethod WriteBackMetadataToAssetRegistry = EAssetRegistryWritebackMethod::Disabled;
 
 	bool IsDLC() const
 	{
@@ -2580,6 +2583,181 @@ private:
 	static constexpr uint64 BufferMemoryLimit = 2ull << 30;
 };
 
+bool DoAssetRegistryWriteback(EAssetRegistryWritebackMethod InMethod, const FString& InCookedDir, TArray<TSharedPtr<IIoStoreWriter>>& InIoStoreWriters)
+{
+	TRACE_CPUPROFILER_EVENT_SCOPE(UpdateAssetRegistryWithSizeInfo);
+	UE_LOG(LogIoStore, Display, TEXT("Adding staging metadata to asset registry..."));
+
+	// The overwhemling majority of time for the asset registry writeback is loading and saving.
+	FString AssetRegistryFileName;
+	FAssetRegistryState AssetRegistry;
+	{
+		TRACE_CPUPROFILER_EVENT_SCOPE(LoadingAssetRegistry);
+
+		// Look for the development registry. Should be in \\GameName\\Metadata\\DevelopmentAssetRegistry.bin, but we don't know what "GameName" is.
+		TArray<FString> PossibleAssetRegistryFiles;
+		IFileManager::Get().FindFilesRecursive(PossibleAssetRegistryFiles, *InCookedDir, GetDevelopmentAssetRegistryFilename(), true, false);
+
+		if (PossibleAssetRegistryFiles.Num() > 1)
+		{
+			UE_LOG(LogIoStore, Warning, TEXT("Found multiple possible development asset registries:"));
+			for (FString& Filename : PossibleAssetRegistryFiles)
+			{
+				UE_LOG(LogIoStore, Warning, TEXT("    %s"), *Filename);
+			}
+		}
+
+		if (PossibleAssetRegistryFiles.Num() == 0)
+		{
+			UE_LOG(LogIoStore, Error, TEXT("No development asset registry file found!"));
+			return false;
+		}
+		else
+		{
+			AssetRegistryFileName = MoveTemp(PossibleAssetRegistryFiles[0]);
+			UE_LOG(LogIoStore, Display, TEXT("Using input asset registry: %s"), *AssetRegistryFileName);
+
+			TUniquePtr<FArchive> FileReader(IFileManager::Get().CreateFileReader(*AssetRegistryFileName));
+			if (FileReader)
+			{
+				TArray64<uint8> Data;
+				Data.SetNumUninitialized(FileReader->TotalSize());
+				FileReader->Serialize(Data.GetData(), Data.Num());
+				check(!FileReader->IsError());
+
+				FAssetRegistrySerializationOptions Options(UE::AssetRegistry::ESerializationTarget::ForDevelopment);
+				FLargeMemoryReader MemoryReader(Data.GetData(), Data.Num());
+				if (AssetRegistry.Serialize(MemoryReader, Options) == false)
+				{
+					UE_LOG(LogIoStore, Error, TEXT("Failed to load asset registry %s"), *AssetRegistryFileName);
+					return false;
+				}
+			}
+			else
+			{
+				UE_LOG(LogIoStore, Error, TEXT("Failed to open asset registry %s"), *AssetRegistryFileName);
+				return false;
+			}
+		}
+	}
+
+	// Create a map off the package id to all of its chunks. 2 inline allocation
+	// is for the export data and the bulk data. For a major test project, 2 covers
+	// 89% of packages, 1 covers 72%.
+	TMap<FPackageId, TArray<FIoStoreTocChunkInfo, TInlineAllocator<2>> > PackageToChunks;
+	for (TSharedPtr<IIoStoreWriter> IoStoreWriter : InIoStoreWriters)
+	{
+		IoStoreWriter->EnumerateChunks([&](const FIoStoreTocChunkInfo& ChunkInfo)
+		{
+			FPackageId PackageId = FPackageId::FromValue(*(int64*)(ChunkInfo.Id.GetData()));
+			PackageToChunks.FindOrAdd(PackageId).Add(ChunkInfo);
+			return true;
+		});
+	}
+
+
+	// Iterate all of the assets in the asset registry and find the corresponding
+	// package. Then add tags to the asset registry to record the chunk info.
+	int32 AssetCount = 0;
+	int32 AddedCount = 0;
+	AssetRegistry.EnumerateAllAssets(TSet<FName>(), [&](const FAssetData& AssetData)
+	{
+		AssetCount++;
+		FPackageId PackageId = FPackageId::FromName(AssetData.PackageName);
+
+		if (PackageToChunks.Contains(PackageId) == false)
+		{
+			// This happens reasonably often.
+			UE_LOG(LogIoStore, Verbose, TEXT("Couldn't find package %s in containers!"), *AssetData.PackageName.ToString());
+			return true;
+		}
+
+		int32 ChunkCount = 0;
+		int64 Size = 0;
+		int64 CompressedSize = 0;
+		for (FIoStoreTocChunkInfo& ChunkInfo : PackageToChunks[PackageId])
+		{
+			ChunkCount++;
+			Size += ChunkInfo.Size;
+			CompressedSize += ChunkInfo.CompressedSize;
+		}
+
+		FAssetDataTagMap TagsAndValues;
+		TagsAndValues.Add("Stage_ChunkCount", LexToString(ChunkCount));
+		TagsAndValues.Add("Stage_ChunkSize", LexToString(Size));
+		TagsAndValues.Add("Stage_ChunkCompressedSize", LexToString(CompressedSize));
+		AssetRegistry.AddTagsToAssetData(FSoftObjectPath(AssetData.ObjectPath), MoveTemp(TagsAndValues));
+		AddedCount++;
+		return true;
+	});
+	UE_LOG(LogIoStore, Verbose, TEXT("Added information to %d assets (of %d possible)"), AddedCount, AssetCount);
+
+	{
+		TRACE_CPUPROFILER_EVENT_SCOPE(SavingAssetRegistry);
+		FLargeMemoryWriter SerializedAssetRegistry;
+		if (AssetRegistry.Save(SerializedAssetRegistry, FAssetRegistrySerializationOptions(UE::AssetRegistry::ESerializationTarget::ForDevelopment)) == false)
+		{
+			UE_LOG(LogIoStore, Error, TEXT("Failed to serialize asset registry to memory."));
+			return false;
+		}
+
+		FString OutputFileName;
+		switch (InMethod)
+		{
+		case EAssetRegistryWritebackMethod::OriginalFile: 
+			{
+				// Write to an adjacent file and move after
+				OutputFileName = AssetRegistryFileName.Replace(TEXT(".bin"), TEXT(".bin.temp"));
+				break;
+			}
+		case EAssetRegistryWritebackMethod::AdjacentFile:
+			{
+				OutputFileName = AssetRegistryFileName.Replace(TEXT(".bin"), TEXT("Staged.bin"));
+				break;
+			}
+		default:
+			{
+				UE_LOG(LogIoStore, Error, TEXT("Invalid asset registry writeback method (should already be handled!) (%d)"), InMethod);
+				return false;
+			}
+		}
+
+		{
+			TUniquePtr<FArchive> Writer = TUniquePtr<FArchive>(IFileManager::Get().CreateFileWriter(*OutputFileName));
+			if (!Writer)
+			{
+				UE_LOG(LogIoStore, Error, TEXT("Failed to open destination asset registry. (%s)"), *OutputFileName);
+				return false;
+			}
+
+			Writer->Serialize(SerializedAssetRegistry.GetData(), SerializedAssetRegistry.TotalSize());
+
+			// Always explicitly close to catch errors from flush/close
+			Writer->Close();
+
+			if (Writer->IsError() || Writer->IsCriticalError())
+			{
+				UE_LOG(LogIoStore, Error, TEXT("Failed to write asset registry to disk. (%s)"), *OutputFileName);
+				return false;
+			}
+
+			if (InMethod == EAssetRegistryWritebackMethod::OriginalFile)
+			{
+				// Move our temp file over the original asset registry.
+				if (IFileManager::Get().Move(*AssetRegistryFileName, *OutputFileName) == false)
+				{
+					// Error already logged by FileManager
+				}
+				OutputFileName = AssetRegistryFileName; // for logging below.
+			}
+		}
+
+		UE_LOG(LogIoStore, Display, TEXT("Staging metadata written back to asset registry: %s"), *OutputFileName);
+	}
+
+	return true;
+}
+
 int32 CreateTarget(const FIoStoreArguments& Arguments, const FIoStoreWriterSettings& GeneralIoWriterSettings)
 {
 	TGuardValue<int32> GuardAllowUnversionedContentInEditor(GAllowUnversionedContentInEditor, 1);
@@ -2835,6 +3013,11 @@ int32 CreateTarget(const FIoStoreArguments& Arguments, const FIoStoreWriterSetti
 		UE_LOG(LogIoStore, Display, TEXT("Compression DDC hits: %llu/%llu (%.2f%%)"), Progress.CompressionDDCHitCount, TotalDDCAttempts, DDCHitRate);
 	}
 
+	if (Arguments.WriteBackMetadataToAssetRegistry != EAssetRegistryWritebackMethod::Disabled)
+	{
+		DoAssetRegistryWriteback(Arguments.WriteBackMetadataToAssetRegistry, Arguments.CookedDir, IoStoreWriters);
+	}
+
 	TArray<FIoStoreWriterResult> IoStoreWriterResults;
 	IoStoreWriterResults.Reserve(IoStoreWriters.Num());
 	for (TSharedPtr<IIoStoreWriter> IoStoreWriter : IoStoreWriters)
@@ -5618,6 +5801,31 @@ int32 CreateIoStoreContainerFiles(const TCHAR* CmdLine)
 		return -1;
 	}
 
+	// Whether or not to write compressed asset sizes back to the asset registry.
+
+	FString WriteBackMetadataToAssetRegistry;
+	if (FParse::Value(FCommandLine::Get(), TEXT("WriteBackMetadataToAssetRegistry="), WriteBackMetadataToAssetRegistry))
+	{
+		// StaticEnum not available in UnrealPak, so manual conversion:
+		if (WriteBackMetadataToAssetRegistry.Equals(TEXT("AdjacentFile"), ESearchCase::IgnoreCase))
+		{
+			Arguments.WriteBackMetadataToAssetRegistry = EAssetRegistryWritebackMethod::AdjacentFile;
+		}
+		else if (WriteBackMetadataToAssetRegistry.Equals(TEXT("OriginalFile"), ESearchCase::IgnoreCase))
+		{
+			Arguments.WriteBackMetadataToAssetRegistry = EAssetRegistryWritebackMethod::OriginalFile;
+		}
+		else if (WriteBackMetadataToAssetRegistry.Equals(TEXT("Disabled"), ESearchCase::IgnoreCase))
+		{
+			Arguments.WriteBackMetadataToAssetRegistry = EAssetRegistryWritebackMethod::Disabled;
+		}
+		else
+		{
+			UE_LOG(LogIoStore, Error, TEXT("Invalid WriteBackMetdataToAssetRegistry value: %s - check setting in ProjectSettings -> Packaging"), *WriteBackMetadataToAssetRegistry);
+			return -1;
+		}
+	}
+
 	FString PackageStoreManifestFilename;
 	if (FParse::Value(FCommandLine::Get(), TEXT("PackageStoreManifest="), PackageStoreManifestFilename))
 	{
diff --git a/Engine/Source/Programs/AutomationTool/Scripts/CopyBuildToStagingDirectory.Automation.cs b/Engine/Source/Programs/AutomationTool/Scripts/CopyBuildToStagingDirectory.Automation.cs
--- a/Engine/Source/Programs/AutomationTool/Scripts/CopyBuildToStagingDirectory.Automation.cs
+++ b/Engine/Source/Programs/AutomationTool/Scripts/CopyBuildToStagingDirectory.Automation.cs
@@ -3104,6 +3104,13 @@ namespace AutomationScripts
 					}
 				}
 
+				string WriteBackMetadataToAssetRegistry;
+				if (PlatformGameConfig.GetString("/Script/UnrealEd.ProjectPackagingSettings", "WriteBackMetadataToAssetRegistry", out WriteBackMetadataToAssetRegistry))
+				{
+					AdditionalArgs += " -WriteBackMetadataToAssetRegistry=" + WriteBackMetadataToAssetRegistry;
+				}
+
+
 				AdditionalArgs += " " + Params.AdditionalIoStoreOptions;
 
 				RunIoStore(Params, SC, IoStoreCommandsFileName, FinalOrderFiles, AdditionalArgs);
diff --git a/Engine/Source/Runtime/Core/Internal/IO/IoStore.h b/Engine/Source/Runtime/Core/Internal/IO/IoStore.h
--- a/Engine/Source/Runtime/Core/Internal/IO/IoStore.h
+++ b/Engine/Source/Runtime/Core/Internal/IO/IoStore.h
@@ -128,6 +128,7 @@ ENUM_CLASS_FLAGS(FIoStoreTocEntryMetaFlags);
  */
 struct FIoStoreTocEntryMeta
 {
+	// Source data hash (i.e. not the on disk data)
 	FIoChunkHash ChunkHash;
 	FIoStoreTocEntryMetaFlags Flags;
 };
