Add a way to do the asset registry writeback from container files after staging.

diff --git a/Engine/Source/Developer/IoStoreUtilities/Private/IoStoreUtilities.cpp b/Engine/Source/Developer/IoStoreUtilities/Private/IoStoreUtilities.cpp
--- a/Engine/Source/Developer/IoStoreUtilities/Private/IoStoreUtilities.cpp
+++ b/Engine/Source/Developer/IoStoreUtilities/Private/IoStoreUtilities.cpp
@@ -1337,7 +1337,7 @@ static void ParsePackageAssetsFromPackageStore(FCookedPackageStore& PackageStore
 	}, EParallelForFlags::Unbalanced);
 }
 
-TUniquePtr<FIoStoreReader> CreateIoStoreReader(const TCHAR* Path, const FKeyChain& KeyChain)
+static TUniquePtr<FIoStoreReader> CreateIoStoreReader(const TCHAR* Path, const FKeyChain& KeyChain)
 {
 	TUniquePtr<FIoStoreReader> IoStoreReader(new FIoStoreReader());
 
@@ -2583,12 +2583,164 @@ private:
 	static constexpr uint64 BufferMemoryLimit = 2ull << 30;
 };
 
-bool DoAssetRegistryWriteback(EAssetRegistryWritebackMethod InMethod, const FString& InCookedDir, TArray<TSharedPtr<IIoStoreWriter>>& InIoStoreWriters)
+static void AddChunkInfoToAssetRegistry(TMap<FPackageId, TArray<FIoStoreTocChunkInfo, TInlineAllocator<2>>>& PackageToChunks, FAssetRegistryState& AssetRegistry)
 {
+	// Iterate all of the assets in the asset registry and find the corresponding
+	// package. Then add tags to the asset registry to record the chunk info.
+	int32 AssetCount = 0;
+	int32 AddedCount = 0;
+	AssetRegistry.EnumerateAllAssets(TSet<FName>(), [&](const FAssetData& AssetData)
+	{
+		AssetCount++;
+		FPackageId PackageId = FPackageId::FromName(AssetData.PackageName);
+
+		if (PackageToChunks.Contains(PackageId) == false)
+		{
+			// This happens reasonably often.
+			UE_LOG(LogIoStore, Verbose, TEXT("Couldn't find package %s in containers!"), *AssetData.PackageName.ToString());
+			return true;
+		}
+
+		int32 ChunkCount = 0;
+		int64 Size = 0;
+		int64 CompressedSize = 0;
+		for (FIoStoreTocChunkInfo& ChunkInfo : PackageToChunks[PackageId])
+		{
+			ChunkCount++;
+			Size += ChunkInfo.Size;
+			CompressedSize += ChunkInfo.CompressedSize;
+		}
+
+		FAssetDataTagMap TagsAndValues;
+		TagsAndValues.Add("Stage_ChunkCount", LexToString(ChunkCount));
+		TagsAndValues.Add("Stage_ChunkSize", LexToString(Size));
+		TagsAndValues.Add("Stage_ChunkCompressedSize", LexToString(CompressedSize));
+		AssetRegistry.AddTagsToAssetData(FSoftObjectPath(AssetData.ObjectPath), MoveTemp(TagsAndValues));
+		AddedCount++;
+		return true;
+	});
+	UE_LOG(LogIoStore, Verbose, TEXT("Added information to %d assets (of %d possible)"), AddedCount, AssetCount);
+}
+
+static bool LoadAssetRegistry(const FString& InAssetRegistryFileName, FAssetRegistryState& OutAssetRegistry)
+{
+	TUniquePtr<FArchive> FileReader(IFileManager::Get().CreateFileReader(*InAssetRegistryFileName));
+	if (FileReader)
+	{
+		TArray64<uint8> Data;
+		Data.SetNumUninitialized(FileReader->TotalSize());
+		FileReader->Serialize(Data.GetData(), Data.Num());
+		check(!FileReader->IsError());
+
+		FAssetRegistrySerializationOptions Options(UE::AssetRegistry::ESerializationTarget::ForDevelopment);
+		FLargeMemoryReader MemoryReader(Data.GetData(), Data.Num());
+		if (OutAssetRegistry.Serialize(MemoryReader, Options) == false)
+		{
+			UE_LOG(LogIoStore, Error, TEXT("Failed to load asset registry %s"), *InAssetRegistryFileName);
+			return false;
+		}
+		return true;
+	}
+	else
+	{
+		UE_LOG(LogIoStore, Error, TEXT("Failed to read asset registry %s"), *InAssetRegistryFileName);
+		return false;
+	}
+}
+
+static bool SaveAssetRegistry(const FString& InAssetRegistryFileName, FAssetRegistryState& InAssetRegistry, bool InSaveTempAndRename)
+{
+	TRACE_CPUPROFILER_EVENT_SCOPE(SavingAssetRegistry);
+	FLargeMemoryWriter SerializedAssetRegistry;
+	if (InAssetRegistry.Save(SerializedAssetRegistry, FAssetRegistrySerializationOptions(UE::AssetRegistry::ESerializationTarget::ForDevelopment)) == false)
+	{
+		UE_LOG(LogIoStore, Error, TEXT("Failed to serialize asset registry to memory."));
+		return false;
+	}
+
+	FString OutputFileName = InSaveTempAndRename ? (InAssetRegistryFileName + TEXT(".temp")) : InAssetRegistryFileName;
+
+	TUniquePtr<FArchive> Writer = TUniquePtr<FArchive>(IFileManager::Get().CreateFileWriter(*OutputFileName));
+	if (!Writer)
+	{
+		UE_LOG(LogIoStore, Error, TEXT("Failed to open destination asset registry. (%s)"), *OutputFileName);
+		return false;
+	}
+
+	Writer->Serialize(SerializedAssetRegistry.GetData(), SerializedAssetRegistry.TotalSize());
+
+	// Always explicitly close to catch errors from flush/close
+	Writer->Close();
+
+	if (Writer->IsError() || Writer->IsCriticalError())
+	{
+		UE_LOG(LogIoStore, Error, TEXT("Failed to write asset registry to disk. (%s)"), *OutputFileName);
+		return false;
+	}
+
+	if (InSaveTempAndRename)
+	{
+		// Move our temp file over the original asset registry.
+		if (IFileManager::Get().Move(*InAssetRegistryFileName, *OutputFileName) == false)
+		{
+			// Error already logged by FileManager
+			return false;
+		}
+	}
+
+	UE_LOG(LogIoStore, Display, TEXT("Saved asset registry to disk. (%s)"), *InAssetRegistryFileName);
+
+	return true;
+}
+
+int32 DoAssetRegistryWritebackAfterStage(const FString& InAssetRegistryFileName, FString&& InContainerDirectory, const FKeyChain& InKeyChain)
+{
+	// This version called after the containers are already created, when you
+	// have a bunch of containers on disk and you want to add chunk info back to
+	// an asset registry.
+
+	FAssetRegistryState AssetRegistry;
+	if (LoadAssetRegistry(InAssetRegistryFileName, AssetRegistry) == false)
+	{
+		return 1; // already logged
+	}
+
+	// Grab all containers in the directory
+	FPaths::NormalizeDirectoryName(InContainerDirectory);
+	TArray<FString> FoundContainerFiles;
+	IFileManager::Get().FindFiles(FoundContainerFiles, *(InContainerDirectory / TEXT("*.utoc")), true, false);
+	
+	// Grab all the package infos.
+	TMap<FPackageId, TArray<FIoStoreTocChunkInfo, TInlineAllocator<2>>> PackageToChunks;
+	for (const FString& Filename : FoundContainerFiles)
+	{
+		TUniquePtr<FIoStoreReader> Reader = CreateIoStoreReader(*(InContainerDirectory / Filename), InKeyChain);
+		if (Reader.IsValid() == false)
+		{
+			return 1; // already logged.
+		}
+
+		Reader->EnumerateChunks([&](const FIoStoreTocChunkInfo& ChunkInfo)
+		{
+			FPackageId PackageId = FPackageId::FromValue(*(int64*)(ChunkInfo.Id.GetData()));
+			PackageToChunks.FindOrAdd(PackageId).Add(ChunkInfo);
+			return true;
+		});
+	}
+
+	AddChunkInfoToAssetRegistry(PackageToChunks, AssetRegistry);
+
+	return SaveAssetRegistry(InAssetRegistryFileName, AssetRegistry, true) ? 0 : 1;
+}
+
+bool DoAssetRegistryWritebackDuringStage(EAssetRegistryWritebackMethod InMethod, const FString& InCookedDir, TArray<TSharedPtr<IIoStoreWriter>>& InIoStoreWriters)
+{
+	// This version called during container creation.
+
 	TRACE_CPUPROFILER_EVENT_SCOPE(UpdateAssetRegistryWithSizeInfo);
 	UE_LOG(LogIoStore, Display, TEXT("Adding staging metadata to asset registry..."));
 
-	// The overwhemling majority of time for the asset registry writeback is loading and saving.
+	// The overwhelming majority of time for the asset registry writeback is loading and saving.
 	FString AssetRegistryFileName;
 	FAssetRegistryState AssetRegistry;
 	{
@@ -2617,26 +2769,9 @@ bool DoAssetRegistryWriteback(EAssetRegistryWritebackMethod InMethod, const FStr
 			AssetRegistryFileName = MoveTemp(PossibleAssetRegistryFiles[0]);
 			UE_LOG(LogIoStore, Display, TEXT("Using input asset registry: %s"), *AssetRegistryFileName);
 
-			TUniquePtr<FArchive> FileReader(IFileManager::Get().CreateFileReader(*AssetRegistryFileName));
-			if (FileReader)
+			if (LoadAssetRegistry(AssetRegistryFileName, AssetRegistry) == false)
 			{
-				TArray64<uint8> Data;
-				Data.SetNumUninitialized(FileReader->TotalSize());
-				FileReader->Serialize(Data.GetData(), Data.Num());
-				check(!FileReader->IsError());
-
-				FAssetRegistrySerializationOptions Options(UE::AssetRegistry::ESerializationTarget::ForDevelopment);
-				FLargeMemoryReader MemoryReader(Data.GetData(), Data.Num());
-				if (AssetRegistry.Serialize(MemoryReader, Options) == false)
-				{
-					UE_LOG(LogIoStore, Error, TEXT("Failed to load asset registry %s"), *AssetRegistryFileName);
-					return false;
-				}
-			}
-			else
-			{
-				UE_LOG(LogIoStore, Error, TEXT("Failed to open asset registry %s"), *AssetRegistryFileName);
-				return false;
+				return false; // already logged
 			}
 		}
 	}
@@ -2644,7 +2779,7 @@ bool DoAssetRegistryWriteback(EAssetRegistryWritebackMethod InMethod, const FStr
 	// Create a map off the package id to all of its chunks. 2 inline allocation
 	// is for the export data and the bulk data. For a major test project, 2 covers
 	// 89% of packages, 1 covers 72%.
-	TMap<FPackageId, TArray<FIoStoreTocChunkInfo, TInlineAllocator<2>> > PackageToChunks;
+	TMap<FPackageId, TArray<FIoStoreTocChunkInfo, TInlineAllocator<2>>> PackageToChunks;
 	for (TSharedPtr<IIoStoreWriter> IoStoreWriter : InIoStoreWriters)
 	{
 		IoStoreWriter->EnumerateChunks([&](const FIoStoreTocChunkInfo& ChunkInfo)
@@ -2655,106 +2790,35 @@ bool DoAssetRegistryWriteback(EAssetRegistryWritebackMethod InMethod, const FStr
 		});
 	}
 
+	AddChunkInfoToAssetRegistry(PackageToChunks, AssetRegistry);
 
-	// Iterate all of the assets in the asset registry and find the corresponding
-	// package. Then add tags to the asset registry to record the chunk info.
-	int32 AssetCount = 0;
-	int32 AddedCount = 0;
-	AssetRegistry.EnumerateAllAssets(TSet<FName>(), [&](const FAssetData& AssetData)
-	{
-		AssetCount++;
-		FPackageId PackageId = FPackageId::FromName(AssetData.PackageName);
-
-		if (PackageToChunks.Contains(PackageId) == false)
-		{
-			// This happens reasonably often.
-			UE_LOG(LogIoStore, Verbose, TEXT("Couldn't find package %s in containers!"), *AssetData.PackageName.ToString());
-			return true;
-		}
-
-		int32 ChunkCount = 0;
-		int64 Size = 0;
-		int64 CompressedSize = 0;
-		for (FIoStoreTocChunkInfo& ChunkInfo : PackageToChunks[PackageId])
-		{
-			ChunkCount++;
-			Size += ChunkInfo.Size;
-			CompressedSize += ChunkInfo.CompressedSize;
-		}
-
-		FAssetDataTagMap TagsAndValues;
-		TagsAndValues.Add("Stage_ChunkCount", LexToString(ChunkCount));
-		TagsAndValues.Add("Stage_ChunkSize", LexToString(Size));
-		TagsAndValues.Add("Stage_ChunkCompressedSize", LexToString(CompressedSize));
-		AssetRegistry.AddTagsToAssetData(FSoftObjectPath(AssetData.ObjectPath), MoveTemp(TagsAndValues));
-		AddedCount++;
-		return true;
-	});
-	UE_LOG(LogIoStore, Verbose, TEXT("Added information to %d assets (of %d possible)"), AddedCount, AssetCount);
-
+	FString OutputFileName;
+	switch (InMethod)
 	{
-		TRACE_CPUPROFILER_EVENT_SCOPE(SavingAssetRegistry);
-		FLargeMemoryWriter SerializedAssetRegistry;
-		if (AssetRegistry.Save(SerializedAssetRegistry, FAssetRegistrySerializationOptions(UE::AssetRegistry::ESerializationTarget::ForDevelopment)) == false)
+	case EAssetRegistryWritebackMethod::OriginalFile:
 		{
-			UE_LOG(LogIoStore, Error, TEXT("Failed to serialize asset registry to memory."));
-			return false;
-		}
-
-		FString OutputFileName;
-		switch (InMethod)
-		{
-		case EAssetRegistryWritebackMethod::OriginalFile: 
-			{
-				// Write to an adjacent file and move after
-				OutputFileName = AssetRegistryFileName.Replace(TEXT(".bin"), TEXT(".bin.temp"));
-				break;
-			}
-		case EAssetRegistryWritebackMethod::AdjacentFile:
+			// Write to an adjacent file and move after
+			if (SaveAssetRegistry(AssetRegistryFileName, AssetRegistry, true) == false)
 			{
-				OutputFileName = AssetRegistryFileName.Replace(TEXT(".bin"), TEXT("Staged.bin"));
-				break;
-			}
-		default:
-			{
-				UE_LOG(LogIoStore, Error, TEXT("Invalid asset registry writeback method (should already be handled!) (%d)"), InMethod);
 				return false;
 			}
+			break;
 		}
-
+	case EAssetRegistryWritebackMethod::AdjacentFile:
 		{
-			TUniquePtr<FArchive> Writer = TUniquePtr<FArchive>(IFileManager::Get().CreateFileWriter(*OutputFileName));
-			if (!Writer)
+			if (SaveAssetRegistry(AssetRegistryFileName.Replace(TEXT(".bin"), TEXT("Staged.bin")), AssetRegistry, true) == false)
 			{
-				UE_LOG(LogIoStore, Error, TEXT("Failed to open destination asset registry. (%s)"), *OutputFileName);
 				return false;
 			}
-
-			Writer->Serialize(SerializedAssetRegistry.GetData(), SerializedAssetRegistry.TotalSize());
-
-			// Always explicitly close to catch errors from flush/close
-			Writer->Close();
-
-			if (Writer->IsError() || Writer->IsCriticalError())
-			{
-				UE_LOG(LogIoStore, Error, TEXT("Failed to write asset registry to disk. (%s)"), *OutputFileName);
-				return false;
-			}
-
-			if (InMethod == EAssetRegistryWritebackMethod::OriginalFile)
-			{
-				// Move our temp file over the original asset registry.
-				if (IFileManager::Get().Move(*AssetRegistryFileName, *OutputFileName) == false)
-				{
-					// Error already logged by FileManager
-				}
-				OutputFileName = AssetRegistryFileName; // for logging below.
-			}
+			break;
+		}
+	default:
+		{
+			UE_LOG(LogIoStore, Error, TEXT("Invalid asset registry writeback method (should already be handled!) (%d)"), InMethod);
+			return false;
 		}
-
-		UE_LOG(LogIoStore, Display, TEXT("Staging metadata written back to asset registry: %s"), *OutputFileName);
 	}
-
+	
 	return true;
 }
 
@@ -3015,7 +3079,7 @@ int32 CreateTarget(const FIoStoreArguments& Arguments, const FIoStoreWriterSetti
 
 	if (Arguments.WriteBackMetadataToAssetRegistry != EAssetRegistryWritebackMethod::Disabled)
 	{
-		DoAssetRegistryWriteback(Arguments.WriteBackMetadataToAssetRegistry, Arguments.CookedDir, IoStoreWriters);
+		DoAssetRegistryWritebackDuringStage(Arguments.WriteBackMetadataToAssetRegistry, Arguments.CookedDir, IoStoreWriters);
 	}
 
 	TArray<FIoStoreWriterResult> IoStoreWriterResults;
@@ -5638,6 +5702,21 @@ int32 CreateIoStoreContainerFiles(const TCHAR* CmdLine)
 
 		return ListContainer(Arguments.KeyChain, ContainerPathOrWildcard, CsvPath);
 	}
+	else if (FParse::Value(FCommandLine::Get(), TEXT("AssetRegistryWriteback="), ArgumentValue))
+	{
+		//
+		// Opens a given directory of containers and a given asset registry, and adds chunk size information
+		// for an asset's package to its asset tags in the asset registry. This can also be done during the staging
+		// process with -WriteBackMetadataToAssetRegistry (below).
+		//
+		FString AssetRegistryFileName = MoveTemp(ArgumentValue);
+		FString PathToContainers;
+		if (!FParse::Value(FCommandLine::Get(), TEXT("ContainerDirectory="), PathToContainers))
+		{
+			UE_LOG(LogIoStore, Error, TEXT("Asset registry writeback requires -ContainerDirectory=Path/To/Containers"));
+		}
+		return DoAssetRegistryWritebackAfterStage(AssetRegistryFileName, MoveTemp(PathToContainers), Arguments.KeyChain);
+	}
 	else if (FParse::Value(FCommandLine::Get(), TEXT("Describe="), ArgumentValue))
 	{
 		FString ContainerPathOrWildcard = ArgumentValue;
