IoStoreShaderCodeLibrary: Fixed issue with library names containing dashes

diff --git a/Engine/Source/Developer/IoStoreUtilities/Private/IoStoreUtilities.cpp b/Engine/Source/Developer/IoStoreUtilities/Private/IoStoreUtilities.cpp
--- a/Engine/Source/Developer/IoStoreUtilities/Private/IoStoreUtilities.cpp
+++ b/Engine/Source/Developer/IoStoreUtilities/Private/IoStoreUtilities.cpp
@@ -64,6 +64,7 @@
 #include "IO/IoStore.h"
 #include "ZenFileSystemManifest.h"
 #include "IPlatformFileSandboxWrapper.h"
+#include "Misc/PathViews.h"
 
 //PRAGMA_DISABLE_OPTIMIZATION
 
@@ -1433,15 +1434,23 @@ bool ConvertToIoStoreShaderLibrary(
 	TArray<TTuple<FSHAHash, TArray<FIoChunkId>>>& OutShaderMaps,
 	TArray<TTuple<FSHAHash, TSet<FName>>>& OutShaderMapAssetAssociations)
 {
-	TArray<FString> Components;
-	if (FPaths::GetBaseFilename(FileName).ParseIntoArray(Components, TEXT("-")) != 3)
+	// ShaderArchive-MyProject-PCD3D.ushaderbytecode
+	FStringView BaseFileNameView = FPathViews::GetBaseFilename(FileName);
+	int32 FormatStartIndex = -1;
+	if (!BaseFileNameView.FindLastChar('-', FormatStartIndex))
 	{
 		UE_LOG(LogIoStore, Error, TEXT("Invalid shader code library file name '%s'."), FileName);
 		return false;
 	}
-	FString LibraryName = Components[1];
-	FName FormatName(Components[2]);
-
+	int32 LibraryNameStartIndex = -1;
+	if (!BaseFileNameView.FindChar('-', LibraryNameStartIndex) || FormatStartIndex - LibraryNameStartIndex <= 1)
+	{
+		UE_LOG(LogIoStore, Error, TEXT("Invalid shader code library file name '%s'."), FileName);
+		return false;
+	}
+	FString LibraryName(BaseFileNameView.Mid(LibraryNameStartIndex + 1, FormatStartIndex - LibraryNameStartIndex - 1));
+	FName FormatName(BaseFileNameView.RightChop(FormatStartIndex + 1));
+	
 	TUniquePtr<FArchive> LibraryAr(IFileManager::Get().CreateFileReader(FileName));
 	if (!LibraryAr)
 	{
