I/O store - fix memory stomp when creating container file(s)

#rnx

diff --git a/Engine/Source/Developer/IoStoreUtilities/Private/IoStoreUtilities.cpp b/Engine/Source/Developer/IoStoreUtilities/Private/IoStoreUtilities.cpp
--- a/Engine/Source/Developer/IoStoreUtilities/Private/IoStoreUtilities.cpp
+++ b/Engine/Source/Developer/IoStoreUtilities/Private/IoStoreUtilities.cpp
@@ -417,8 +417,9 @@ public:
 				PackageIdToEntry.Add(ContainerHeader.PackageIds[Idx++], &Entry);
 			}
 		}
-
-		for (const FPackageStoreManifest::FPackageInfo& PackageInfo : PackageStoreManifest.GetPackages())
+		
+		ManifestPackageInfos = PackageStoreManifest.GetPackages();
+		for (const FPackageStoreManifest::FPackageInfo& PackageInfo : ManifestPackageInfos)
 		{
 			FName PackageName = FName(PackageInfo.PackageName);
 			PackageNameToPackageInfoMap.Add(PackageName, &PackageInfo);
@@ -588,6 +589,7 @@ private:
 	
 	TUniquePtr<IDataSource> DataSource;
 	FPackageStoreManifest PackageStoreManifest;
+	TArray<FPackageStoreManifest::FPackageInfo> ManifestPackageInfos;
 	TArray<FPackageStoreEntryResource> PackageEntries;
 	TMap<FPackageId, const FPackageStoreEntryResource*> PackageIdToEntry;
 	TMap<FString, FIoChunkId> FilenameToChunkIdMap;
