Updating iostore asset registry writeback to use new cleaner functions and use the chunk listing on the asset registry.

diff --git a/Engine/Source/Developer/IoStoreUtilities/Private/IoStoreUtilities.cpp b/Engine/Source/Developer/IoStoreUtilities/Private/IoStoreUtilities.cpp
--- a/Engine/Source/Developer/IoStoreUtilities/Private/IoStoreUtilities.cpp
+++ b/Engine/Source/Developer/IoStoreUtilities/Private/IoStoreUtilities.cpp
@@ -2594,28 +2594,46 @@ private:
 	static constexpr uint64 BufferMemoryLimit = 2ull << 30;
 };
 
-static void AddChunkInfoToAssetRegistry(TMap<FPackageId, TArray<FIoStoreTocChunkInfo, TInlineAllocator<2>>>& PackageToChunks, FAssetRegistryState& AssetRegistry)
+static void AddChunkInfoToAssetRegistry(TMap<FPackageId, TArray<FIoStoreTocChunkInfo, TInlineAllocator<2>>>&& PackageToChunks, FAssetRegistryState& AssetRegistry, uint64 TotalCompressedSize)
 {
-	// Iterate all of the assets in the asset registry and find the corresponding
-	// package. Then add tags to the asset registry to record the chunk info.
-	int32 AssetCount = 0;
-	int32 AddedCount = 0;
-	AssetRegistry.EnumerateAllAssets(TSet<FName>(), [&](const FAssetData& AssetData)
+	//
+	// The asset registry has the chunks associate with each package, so we can just iterate the
+	// packages, look up the chunk info, and then save the tags.
+	//
+	// The complicated thing is (as usual), trying to determine which asset gets the blame for the
+	// data. We use the GetMostImportantAsset function for this.
+	//
+	const TMap<FName, const FAssetPackageData*> AssetPackageMap = AssetRegistry.GetAssetPackageDataMap();
+
+	uint64 AssetsCompressedSize = 0;
+	uint64 UpdatedAssetCount = 0;
+
+	for (const TPair<FName, const FAssetPackageData*>& AssetPackage : AssetPackageMap)
 	{
-		AssetCount++;
-		FPackageId PackageId = FPackageId::FromName(AssetData.PackageName);
+		if (AssetPackage.Value->DiskSize == -1)
+		{
+			// No data on disk!
+			continue;
+		}
 
-		if (PackageToChunks.Contains(PackageId) == false)
+		const FAssetData* AssetData = UE::AssetRegistry::GetMostImportantAsset(AssetRegistry.GetAssetsByPackageName(AssetPackage.Key), false);
+		if (AssetData == nullptr)
 		{
-			// This happens reasonably often.
-			UE_LOG(LogIoStore, Verbose, TEXT("Couldn't find package %s in containers!"), *AssetData.PackageName.ToString());
-			return true;
+			// e.g. /Script packages.
+			continue;
+		}
+
+		const TArray<FIoStoreTocChunkInfo, TInlineAllocator<2>>* PackageChunks = PackageToChunks.Find(FPackageId::FromName(AssetPackage.Key));
+		if (PackageChunks == nullptr)
+		{
+			UE_LOG(LogIoStore, Warning, TEXT("Package data for %s had no chunks %d -- %lld "), *AssetPackage.Key.ToString(), AssetPackage.Value->ChunkHashes.Num(), AssetPackage.Value->DiskSize);
+			continue;
 		}
 
 		int32 ChunkCount = 0;
 		int64 Size = 0;
 		int64 CompressedSize = 0;
-		for (FIoStoreTocChunkInfo& ChunkInfo : PackageToChunks[PackageId])
+		for (const FIoStoreTocChunkInfo& ChunkInfo : *PackageChunks)
 		{
 			ChunkCount++;
 			Size += ChunkInfo.Size;
@@ -2626,37 +2644,50 @@ static void AddChunkInfoToAssetRegistry(TMap<FPackageId, TArray<FIoStoreTocChunk
 		TagsAndValues.Add("Stage_ChunkCount", LexToString(ChunkCount));
 		TagsAndValues.Add("Stage_ChunkSize", LexToString(Size));
 		TagsAndValues.Add("Stage_ChunkCompressedSize", LexToString(CompressedSize));
-		AssetRegistry.AddTagsToAssetData(FSoftObjectPath(AssetData.ObjectPath), MoveTemp(TagsAndValues));
-		AddedCount++;
-		return true;
-	});
-	UE_LOG(LogIoStore, Verbose, TEXT("Added information to %d assets (of %d possible)"), AddedCount, AssetCount);
-}
+		AssetRegistry.AddTagsToAssetData(FSoftObjectPath(AssetData->ObjectPath), MoveTemp(TagsAndValues));
 
-static bool LoadAssetRegistry(const FString& InAssetRegistryFileName, FAssetRegistryState& OutAssetRegistry)
-{
-	TUniquePtr<FArchive> FileReader(IFileManager::Get().CreateFileReader(*InAssetRegistryFileName));
-	if (FileReader)
-	{
-		TArray64<uint8> Data;
-		Data.SetNumUninitialized(FileReader->TotalSize());
-		FileReader->Serialize(Data.GetData(), Data.Num());
-		check(!FileReader->IsError());
+		// We assign a package's chunks to a single asset, remove it from the list so that
+		// at the end we can track how many chunks don't get assigned.
+		PackageToChunks.Remove(FPackageId::FromName(AssetPackage.Key));
 
-		FAssetRegistrySerializationOptions Options(UE::AssetRegistry::ESerializationTarget::ForDevelopment);
-		FLargeMemoryReader MemoryReader(Data.GetData(), Data.Num());
-		if (OutAssetRegistry.Serialize(MemoryReader, Options) == false)
+		UpdatedAssetCount++;
+		AssetsCompressedSize += CompressedSize;
+	}
+	
+	// PackageToChunks now has chunks that we never assigned to an asset, and so aren't accounted for.
+	uint64 RemainingByType[(uint8)EIoChunkType::MAX] = {};
+	for (auto PackageChunks : PackageToChunks)
+	{
+		for (FIoStoreTocChunkInfo& Info : PackageChunks.Value)
 		{
-			UE_LOG(LogIoStore, Error, TEXT("Failed to load asset registry %s"), *InAssetRegistryFileName);
-			return false;
+			RemainingByType[(uint8)Info.ChunkType] += Info.CompressedSize;
 		}
-		return true;
 	}
-	else
+
+	double PercentAssets = 1.0f;
+	if (TotalCompressedSize != 0)
 	{
-		UE_LOG(LogIoStore, Error, TEXT("Failed to read asset registry %s"), *InAssetRegistryFileName);
-		return false;
+		PercentAssets = AssetsCompressedSize / (double)TotalCompressedSize;
 	}
+
+	UE_LOG(LogIoStore, Display, TEXT("Added chunk metadata to %s assets."), *FText::AsNumber(UpdatedAssetCount).ToString());
+	UE_LOG(LogIoStore, Display, TEXT("Assets represent %s bytes of %s chunk bytes (%.1f%%)"), *FText::AsNumber(AssetsCompressedSize).ToString(), *FText::AsNumber(TotalCompressedSize).ToString(), 100 * PercentAssets);
+	UE_LOG(LogIoStore, Display, TEXT("Remaining data by chunk type:"));
+	for (uint8 TypeIndex = 0; TypeIndex < (uint8)EIoChunkType::MAX; TypeIndex++)
+	{
+		if (RemainingByType[TypeIndex] != 0)
+		{
+			UE_LOG(LogIoStore, Display, TEXT("    %-24s%s"), *LexToString((EIoChunkType)TypeIndex), *FText::AsNumber(RemainingByType[TypeIndex]).ToString());
+		}
+	}
+}
+
+static bool LoadAssetRegistry(const FString& InAssetRegistryFileName, FAssetRegistryState& OutAssetRegistry)
+{
+	FAssetRegistryVersion::Type Version;
+	FAssetRegistryLoadOptions Options(UE::AssetRegistry::ESerializationTarget::ForDevelopment);
+	bool bSucceeded = FAssetRegistryState::LoadFromDisk(*InAssetRegistryFileName, Options, OutAssetRegistry, &Version);
+	return bSucceeded;
 }
 
 static bool SaveAssetRegistry(const FString& InAssetRegistryFileName, FAssetRegistryState& InAssetRegistry, bool InSaveTempAndRename)
@@ -2721,6 +2752,8 @@ int32 DoAssetRegistryWritebackAfterStage(const FString& InAssetRegistryFileName,
 	TArray<FString> FoundContainerFiles;
 	IFileManager::Get().FindFiles(FoundContainerFiles, *(InContainerDirectory / TEXT("*.utoc")), true, false);
 	
+	uint64 TotalCompressedSize = 0;
+	
 	// Grab all the package infos.
 	TMap<FPackageId, TArray<FIoStoreTocChunkInfo, TInlineAllocator<2>>> PackageToChunks;
 	for (const FString& Filename : FoundContainerFiles)
@@ -2730,16 +2763,17 @@ int32 DoAssetRegistryWritebackAfterStage(const FString& InAssetRegistryFileName,
 		{
 			return 1; // already logged.
 		}
-
+		
 		Reader->EnumerateChunks([&](const FIoStoreTocChunkInfo& ChunkInfo)
 		{
 			FPackageId PackageId = FPackageId::FromValue(*(int64*)(ChunkInfo.Id.GetData()));
 			PackageToChunks.FindOrAdd(PackageId).Add(ChunkInfo);
+			TotalCompressedSize += ChunkInfo.CompressedSize;
 			return true;
 		});
 	}
 
-	AddChunkInfoToAssetRegistry(PackageToChunks, AssetRegistry);
+	AddChunkInfoToAssetRegistry(MoveTemp(PackageToChunks), AssetRegistry, TotalCompressedSize);
 
 	return SaveAssetRegistry(InAssetRegistryFileName, AssetRegistry, true) ? 0 : 1;
 }
@@ -2790,6 +2824,7 @@ bool DoAssetRegistryWritebackDuringStage(EAssetRegistryWritebackMethod InMethod,
 	// Create a map off the package id to all of its chunks. 2 inline allocation
 	// is for the export data and the bulk data. For a major test project, 2 covers
 	// 89% of packages, 1 covers 72%.
+	uint64 TotalCompressedSize = 0;
 	TMap<FPackageId, TArray<FIoStoreTocChunkInfo, TInlineAllocator<2>>> PackageToChunks;
 	for (TSharedPtr<IIoStoreWriter> IoStoreWriter : InIoStoreWriters)
 	{
@@ -2797,11 +2832,12 @@ bool DoAssetRegistryWritebackDuringStage(EAssetRegistryWritebackMethod InMethod,
 		{
 			FPackageId PackageId = FPackageId::FromValue(*(int64*)(ChunkInfo.Id.GetData()));
 			PackageToChunks.FindOrAdd(PackageId).Add(ChunkInfo);
+			TotalCompressedSize += ChunkInfo.CompressedSize;
 			return true;
 		});
 	}
 
-	AddChunkInfoToAssetRegistry(PackageToChunks, AssetRegistry);
+	AddChunkInfoToAssetRegistry(MoveTemp(PackageToChunks), AssetRegistry, TotalCompressedSize);
 
 	FString OutputFileName;
 	switch (InMethod)
@@ -3379,8 +3415,6 @@ int32 ListContainer(
 		return -1;
 	}
 
-	TArray<FString> CsvLines;
-	
 	TUniquePtr<FOutputDeviceFile> Out = MakeUnique<FOutputDeviceFile>(*CsvPath, true);
 	Out->SetSuppressEventTag(true);
 
