Zen - cook to custom project/oplog ID with -ZenStoreProject=<ProjectId> and -ZenStorePlatform=<OplogId>

#rnx

diff --git a/Engine/Source/Developer/IoStoreUtilities/Internal/ZenStoreHttpClient.h b/Engine/Source/Developer/IoStoreUtilities/Internal/ZenStoreHttpClient.h
--- a/Engine/Source/Developer/IoStoreUtilities/Internal/ZenStoreHttpClient.h
+++ b/Engine/Source/Developer/IoStoreUtilities/Internal/ZenStoreHttpClient.h
@@ -64,6 +64,7 @@ public:
 	static const UTF8CHAR* FindOrAddAttachmentId(FUtf8StringView AttachmentText);
 	static const UTF8CHAR* FindAttachmentId(FUtf8StringView AttachmentText);
 
+	static FString GetProjectId();
 	static FString GenerateDefaultProjectId();
 private:
 	TIoStatusOr<FIoBuffer> ReadOpLogUri(FStringBuilderBase& ChunkUri, uint64 Offset = 0, uint64 Size = ~0ull);
diff --git a/Engine/Source/Developer/IoStoreUtilities/Private/ZenStoreHttpClient.cpp b/Engine/Source/Developer/IoStoreUtilities/Private/ZenStoreHttpClient.cpp
--- a/Engine/Source/Developer/IoStoreUtilities/Private/ZenStoreHttpClient.cpp
+++ b/Engine/Source/Developer/IoStoreUtilities/Private/ZenStoreHttpClient.cpp
@@ -7,6 +7,7 @@
 #include "Algo/Sort.h"
 #include "Misc/ScopeRWLock.h"
 #include "Misc/App.h"
+#include "Misc/CommandLine.h"
 #include "HAL/FileManager.h"
 
 #if PLATFORM_WINDOWS || PLATFORM_UNIX || PLATFORM_MAC
@@ -781,6 +782,17 @@ static FString GetProjectPathId()
 	return FString::Printf(TEXT("%s.%.8s"), FApp::GetProjectName(), *HashString);
 }
 
+FString FZenStoreHttpClient::GetProjectId()
+{
+	FString ProjectId;
+	if (FParse::Value(FCommandLine::Get(), TEXT("-ZenStoreProject="), ProjectId) == false)
+	{
+		ProjectId = GetProjectPathId();
+	}
+
+	return ProjectId;
+}
+
 FString FZenStoreHttpClient::GenerateDefaultProjectId()
 {
 	return GetProjectPathId();
diff --git a/Engine/Source/Developer/IoStoreUtilities/Private/ZenStoreWriter.cpp b/Engine/Source/Developer/IoStoreUtilities/Private/ZenStoreWriter.cpp
--- a/Engine/Source/Developer/IoStoreUtilities/Private/ZenStoreWriter.cpp
+++ b/Engine/Source/Developer/IoStoreUtilities/Private/ZenStoreWriter.cpp
@@ -18,6 +18,7 @@
 #include "HAL/PlatformFileManager.h"
 #include "IO/IoDispatcher.h"
 #include "Misc/App.h"
+#include "Misc/CommandLine.h"
 #include "Misc/FileHelper.h"
 #include "Misc/StringBuilder.h"
 #include "Interfaces/ITargetPlatform.h"
@@ -168,9 +169,13 @@ FZenStoreWriter::FZenStoreWriter(
 {
 	StaticInit();
 
-	FString ProjectId = FZenStoreHttpClient::GenerateDefaultProjectId();
-	FString OplogId = InTargetPlatform->PlatformName();
-
+	ProjectId = FZenStoreHttpClient::GetProjectId();
+	
+	if (FParse::Value(FCommandLine::Get(), TEXT("-ZenStorePlatform="), OplogId) == false)
+	{
+		OplogId = InTargetPlatform->PlatformName();
+	}
+	
 	HttpClient = MakeUnique<UE::FZenStoreHttpClient>();
 
 	FString RootDir = FPaths::RootDir();
@@ -367,8 +372,6 @@ void FZenStoreWriter::Initialize(const FCookInfo& Info)
 			IFileManager::Get().DeleteDirectory(*OutputPath, bRequireExists, bTree);
 		}
 
-		FString ProjectId = FZenStoreHttpClient::GenerateDefaultProjectId();
-		FString OplogId = TargetPlatform.PlatformName();
 		bool bOplogEstablished = HttpClient->TryCreateOplog(ProjectId, OplogId, Info.bFullBuild);
 		UE_CLOG(!bOplogEstablished, LogZenStoreWriter, Fatal, TEXT("Failed to establish oplog on the ZenServer"));
 
diff --git a/Engine/Source/Developer/IoStoreUtilities/Public/ZenStoreWriter.h b/Engine/Source/Developer/IoStoreUtilities/Public/ZenStoreWriter.h
--- a/Engine/Source/Developer/IoStoreUtilities/Public/ZenStoreWriter.h
+++ b/Engine/Source/Developer/IoStoreUtilities/Public/ZenStoreWriter.h
@@ -153,6 +153,8 @@ private:
 
 	const ITargetPlatform&				TargetPlatform;
 	const FName							TargetPlatformFName;
+	FString								ProjectId;
+	FString								OplogId;
 	FString								OutputPath;
 	FString								MetadataDirectoryPath;
 	FIoContainerId						ContainerId = FIoContainerId::FromName(TEXT("global"));
diff --git a/Engine/Source/Editor/UnrealEd/Private/TargetDomain/TargetDomainUtils.cpp b/Engine/Source/Editor/UnrealEd/Private/TargetDomain/TargetDomainUtils.cpp
--- a/Engine/Source/Editor/UnrealEd/Private/TargetDomain/TargetDomainUtils.cpp
+++ b/Engine/Source/Editor/UnrealEd/Private/TargetDomain/TargetDomainUtils.cpp
@@ -417,7 +417,7 @@ FEditorDomainOplog::FEditorDomainOplog()
 {
 	StaticInit();
 
-	FString ProjectId = FZenStoreHttpClient::GenerateDefaultProjectId();
+	FString ProjectId = FZenStoreHttpClient::GetProjectId();
 	FString OplogId = TEXT("EditorDomain");
 
 	FString RootDir = FPaths::RootDir();
