IoStore:
Preparation for optional package data support
Some cleanup of the cooked file stat map
Fixed package bulkdata size not being calculated correctly when staging from Zen
#rnx

diff --git a/Engine/Source/Developer/IoStoreUtilities/Private/IoStoreUtilities.cpp b/Engine/Source/Developer/IoStoreUtilities/Private/IoStoreUtilities.cpp
--- a/Engine/Source/Developer/IoStoreUtilities/Private/IoStoreUtilities.cpp
+++ b/Engine/Source/Developer/IoStoreUtilities/Private/IoStoreUtilities.cpp
@@ -295,17 +295,132 @@ struct FContainerSourceSpec
 
 struct FCookedFileStatData
 {
-	enum EFileExt { UMap, UAsset, UExp, UBulk, UPtnl, UMappedBulk, UShaderByteCode };
-	enum EFileType { PackageHeader, PackageData, BulkData, ShaderLibrary };
+	enum EFileType
+	{
+		PackageHeader,
+		PackageData,
+		BulkData,
+		OptionalBulkData,
+		MemoryMappedBulkData,
+		ShaderLibrary,
+		Regions,
+		Invalid
+	};
 
 	int64 FileSize = 0;
-	EFileType FileType = PackageHeader;
-	EFileExt FileExt = UMap;
-
-	TArray<FFileRegion> FileRegions;
+	EFileType FileType = Invalid;
 };
 
-using FCookedFileStatMap = TMap<FString, FCookedFileStatData>;
+static int32 GetFullExtensionStartIndex(FStringView Path)
+{
+	int32 ExtensionStartIndex = -1;
+	for (int32 Index = Path.Len() - 1; Index >= 0; --Index)
+	{
+		if (FPathViews::IsSeparator(Path[Index]))
+		{
+			break;
+		}
+		else if (Path[Index] == '.')
+		{
+			ExtensionStartIndex = Index;
+		}
+	}
+	return ExtensionStartIndex;
+}
+
+static FStringView GetBaseFilenameWithoutAnyExtension(FStringView Path)
+{
+	int32 ExtensionStartIndex = GetFullExtensionStartIndex(Path);
+	if (ExtensionStartIndex < 0)
+	{
+		return FStringView();
+	}
+	else
+	{
+		return Path.Left(ExtensionStartIndex);
+	}
+}
+
+static FStringView GetFullExtension(FStringView Path)
+{
+	int32 ExtensionStartIndex = GetFullExtensionStartIndex(Path);
+	if (ExtensionStartIndex < 0)
+	{
+		return FStringView();
+	}
+	else
+	{
+		return Path.RightChop(ExtensionStartIndex);
+	}
+}
+
+class FCookedFileStatMap
+{
+public:
+	FCookedFileStatMap()
+	{
+		Extensions.Emplace(TEXT(".umap"), FCookedFileStatData::PackageHeader);
+		Extensions.Emplace(TEXT(".uasset"), FCookedFileStatData::PackageHeader);
+		Extensions.Emplace(TEXT(".uexp"), FCookedFileStatData::PackageData);
+		Extensions.Emplace(TEXT(".ubulk"), FCookedFileStatData::BulkData);
+		Extensions.Emplace(TEXT(".uptnl"), FCookedFileStatData::OptionalBulkData);
+		Extensions.Emplace(TEXT(".m.ubulk"), FCookedFileStatData::MemoryMappedBulkData);
+		Extensions.Emplace(*FString::Printf(TEXT(".uexp%s"), FFileRegion::RegionsFileExtension), FCookedFileStatData::Regions);
+		Extensions.Emplace(*FString::Printf(TEXT(".uptnl%s"), FFileRegion::RegionsFileExtension), FCookedFileStatData::Regions);
+		Extensions.Emplace(*FString::Printf(TEXT(".ubulk%s"), FFileRegion::RegionsFileExtension), FCookedFileStatData::Regions);
+		Extensions.Emplace(*FString::Printf(TEXT(".m.ubulk%s"), FFileRegion::RegionsFileExtension), FCookedFileStatData::Regions);
+		Extensions.Emplace(TEXT(".ushaderbytecode"), FCookedFileStatData::ShaderLibrary);
+	}
+
+	int32 Num() const
+	{
+		return Map.Num();
+	}
+
+	void Add(const TCHAR* Path, int64 FileSize)
+	{
+		FString NormalizedPath(Path);
+		FPaths::NormalizeFilename(NormalizedPath);
+
+		FStringView NormalizePathView(NormalizedPath);
+		int32 ExtensionStartIndex = GetFullExtensionStartIndex(NormalizePathView);
+		if (ExtensionStartIndex < 0)
+		{
+			return;
+		}
+		FStringView Extension = NormalizePathView.RightChop(ExtensionStartIndex);
+		const FCookedFileStatData::EFileType* FindFileType = FindFileTypeFromExtension(Extension);
+		if (!FindFileType)
+		{
+			return;
+		}
+
+		FCookedFileStatData& CookedFileStatData = Map.Add(MoveTemp(NormalizedPath));
+		CookedFileStatData.FileType = *FindFileType;
+		CookedFileStatData.FileSize = FileSize;
+	}
+
+	const FCookedFileStatData* Find(FStringView NormalizedPath) const
+	{
+		return Map.FindByHash(GetTypeHash(NormalizedPath), NormalizedPath);
+	}
+
+private:
+	const FCookedFileStatData::EFileType* FindFileTypeFromExtension(const FStringView& Extension)
+	{
+		for (const auto& Pair : Extensions)
+		{
+			if (Pair.Key == Extension)
+			{
+				return &Pair.Value;
+			}
+		}
+		return nullptr;
+	}
+
+	TArray<TTuple<FString, FCookedFileStatData::EFileType>> Extensions;
+	TMap<FString, FCookedFileStatData> Map;
+};
 
 struct FContainerTargetSpec;
 
@@ -344,7 +459,8 @@ enum class EContainerChunkType
 	PackageData,
 	BulkData,
 	OptionalBulkData,
-	MemoryMappedBulkData
+	MemoryMappedBulkData,
+	Invalid
 };
 
 struct FContainerTargetFile
@@ -1827,13 +1943,18 @@ void InitializeContainerTargetsAndPackages(
 		const FCookedFileStatData* CookedFileStatData = OriginalCookedFileStatData;
 		if (CookedFileStatData->FileType == FCookedFileStatData::PackageHeader)
 		{
-			OutTargetFile.NormalizedSourcePath = FPaths::ChangeExtension(SourceFile.NormalizedPath, TEXT(".uexp"));
-			CookedFileStatData = Arguments.CookedFileStatMap.Find(OutTargetFile.NormalizedSourcePath);
+			FStringView NormalizedSourcePathView(SourceFile.NormalizedPath);
+			int32 ExtensionStartIndex = GetFullExtensionStartIndex(NormalizedSourcePathView);
+			TStringBuilder<512> UexpPath;
+			UexpPath.Append(NormalizedSourcePathView.Left(ExtensionStartIndex));
+			UexpPath.Append(TEXT(".uexp"));
+			CookedFileStatData = Arguments.CookedFileStatMap.Find(*UexpPath);
 			if (!CookedFileStatData)
 			{
-				UE_LOG(LogIoStore, Warning, TEXT("File not found: '%s'"), *OutTargetFile.NormalizedSourcePath);
+				UE_LOG(LogIoStore, Warning, TEXT("Couldn't find .uexp file for: '%s'"), *SourceFile.NormalizedPath);
 				return false;
 			}
+			OutTargetFile.NormalizedSourcePath = UexpPath;
 		}
 		else
 		{
@@ -1845,61 +1966,60 @@ void InitializeContainerTargetsAndPackages(
 		{
 			OutTargetFile.ChunkType = EContainerChunkType::ShaderCodeLibrary;
 		}
-		else if (CookedFileStatData->FileExt == FCookedFileStatData::EFileExt::UMappedBulk)
-		{
-			OutTargetFile.ChunkType = EContainerChunkType::MemoryMappedBulkData;
-			FString TmpFileName = FString(SourceFile.NormalizedPath.Len() - 8, GetData(SourceFile.NormalizedPath)) + TEXT(".ubulk");
-			OutTargetFile.Package = FindOrAddPackage(Arguments, *TmpFileName, Packages, PackageNameMap, PackageIdMap);
-		}
 		else
 		{
 			OutTargetFile.Package = FindOrAddPackage(Arguments, *SourceFile.NormalizedPath, Packages, PackageNameMap, PackageIdMap);
-			if (CookedFileStatData->FileExt == FCookedFileStatData::UPtnl)
-			{
-				OutTargetFile.ChunkType = EContainerChunkType::OptionalBulkData;
-			}
-			else if (CookedFileStatData->FileType == FCookedFileStatData::BulkData)
-			{
-				OutTargetFile.ChunkType = EContainerChunkType::BulkData;
-			}
-			else
+
+			switch (CookedFileStatData->FileType)
 			{
+			case FCookedFileStatData::PackageData:
 				OutTargetFile.ChunkType = EContainerChunkType::PackageData;
+				OutTargetFile.ChunkId = CreateIoChunkId(OutTargetFile.Package->GlobalPackageId.Value(), 0, EIoChunkType::ExportBundleData);
+				OutTargetFile.Package->FileName = SourceFile.NormalizedPath; // .uasset path
+				OutTargetFile.Package->UAssetSize = OriginalCookedFileStatData->FileSize;
+				OutTargetFile.Package->UExpSize = CookedFileStatData->FileSize;
+				break;
+			case FCookedFileStatData::BulkData:
+				OutTargetFile.ChunkType = EContainerChunkType::BulkData;
+				OutTargetFile.ChunkId = CreateIoChunkId(OutTargetFile.Package->GlobalPackageId.Value(), 0, EIoChunkType::BulkData);
+				OutTargetFile.Package->TotalBulkDataSize += CookedFileStatData->FileSize;
+				break;
+			case FCookedFileStatData::OptionalBulkData:
+				OutTargetFile.ChunkType = EContainerChunkType::OptionalBulkData;
+				OutTargetFile.ChunkId = CreateIoChunkId(OutTargetFile.Package->GlobalPackageId.Value(), 0, EIoChunkType::OptionalBulkData);
+				OutTargetFile.Package->TotalBulkDataSize += CookedFileStatData->FileSize;
+				break;
+			case FCookedFileStatData::MemoryMappedBulkData:
+				OutTargetFile.ChunkType = EContainerChunkType::MemoryMappedBulkData;
+				OutTargetFile.ChunkId = CreateIoChunkId(OutTargetFile.Package->GlobalPackageId.Value(), 0, EIoChunkType::MemoryMappedBulkData);
+				OutTargetFile.Package->TotalBulkDataSize += CookedFileStatData->FileSize;
+				break;
+			default:
+				UE_LOG(LogIoStore, Fatal, TEXT("Unexpected file type %d for file '%s'"), CookedFileStatData->FileType, *OutTargetFile.NormalizedSourcePath);
+				return false;
 			}
 		}
 
-		switch (OutTargetFile.ChunkType)
-		{
-		case EContainerChunkType::OptionalBulkData:
-		case EContainerChunkType::MemoryMappedBulkData:
-		case EContainerChunkType::BulkData:
-			OutTargetFile.Package->TotalBulkDataSize += CookedFileStatData->FileSize;
-			break;
-		}
-
-		switch (OutTargetFile.ChunkType)
-		{
-		case EContainerChunkType::OptionalBulkData:
-			OutTargetFile.ChunkId = CreateIoChunkId(OutTargetFile.Package->GlobalPackageId.Value(), 0, EIoChunkType::OptionalBulkData);
-			break;
-		case EContainerChunkType::MemoryMappedBulkData:
-			OutTargetFile.ChunkId = CreateIoChunkId(OutTargetFile.Package->GlobalPackageId.Value(), 0, EIoChunkType::MemoryMappedBulkData);
-			break;
-		case EContainerChunkType::BulkData:
-			OutTargetFile.ChunkId = CreateIoChunkId(OutTargetFile.Package->GlobalPackageId.Value(), 0, EIoChunkType::BulkData);
-			break;
-		case EContainerChunkType::PackageData:
-			OutTargetFile.Package->FileName = SourceFile.NormalizedPath; // .uasset path
-			OutTargetFile.Package->UAssetSize = OriginalCookedFileStatData->FileSize;
-			OutTargetFile.Package->UExpSize = CookedFileStatData->FileSize;
-			OutTargetFile.ChunkId = CreateIoChunkId(OutTargetFile.Package->GlobalPackageId.Value(), 0, EIoChunkType::ExportBundleData);
-			break;
-		}
-
 		// Only keep the regions for the file if neither compression nor encryption are enabled, otherwise the regions will be meaningless.
-		if (!SourceFile.bNeedsCompression && !SourceFile.bNeedsEncryption)
+		if (Arguments.bFileRegions && !SourceFile.bNeedsCompression && !SourceFile.bNeedsEncryption)
 		{
-			OutTargetFile.FileRegions = CookedFileStatData->FileRegions;
+			// Read the matching regions file, if it exists.
+			TStringBuilder<512> RegionsFilePath;
+			RegionsFilePath.Append(OutTargetFile.NormalizedSourcePath);
+			RegionsFilePath.Append(FFileRegion::RegionsFileExtension);
+			const FCookedFileStatData* RegionsFileStatData = Arguments.CookedFileStatMap.Find(RegionsFilePath);
+			if (RegionsFileStatData)
+			{
+				TUniquePtr<FArchive> RegionsFile(IFileManager::Get().CreateFileReader(*RegionsFilePath));
+				if (!RegionsFile.IsValid())
+				{
+					UE_LOG(LogIoStore, Warning, TEXT("Failed reading file '%s'"), *RegionsFilePath);
+				}
+				else
+				{
+					FFileRegion::SerializeFileRegions(*RegionsFile.Get(), OutTargetFile.FileRegions);
+				}
+			}
 		}
 
 		return true;
@@ -1915,16 +2035,17 @@ void InitializeContainerTargetsAndPackages(
 		
 		OutTargetFile.NormalizedSourcePath = SourceFile.NormalizedPath;
 		
-		if (SourceFile.NormalizedPath.EndsWith(TEXT(".ushaderbytecode")))
+		FStringView Extension = GetFullExtension(SourceFile.NormalizedPath);
+		if (Extension == TEXT(".ushaderbytecode"))
 		{
-			int64 FileSize = IFileManager::Get().FileSize(*SourceFile.NormalizedPath);
-			if (FileSize == INDEX_NONE)
+			const FCookedFileStatData* CookedFileStatData = Arguments.CookedFileStatMap.Find(SourceFile.NormalizedPath);
+			if (!CookedFileStatData)
 			{
 				UE_LOG(LogIoStore, Warning, TEXT("File not found: '%s'"), *SourceFile.NormalizedPath);
 				return false;
 			}
 			OutTargetFile.ChunkType = EContainerChunkType::ShaderCodeLibrary;
-			OutTargetFile.SourceSize = FileSize;
+			OutTargetFile.SourceSize = uint64(CookedFileStatData->FileSize);
 			return true;
 		}
 
@@ -1944,49 +2065,60 @@ void InitializeContainerTargetsAndPackages(
 		OutTargetFile.SourceSize = ChunkSize.ValueOrDie();
 
 		FName PackageName = PackageStore.GetPackageNameFromChunkId(OutTargetFile.ChunkId);
-		if (!PackageName.IsNone())
+		if (PackageName.IsNone())
+		{
+			UE_LOG(LogIoStore, Warning, TEXT("Package name not found for: '%s'"), *SourceFile.NormalizedPath);
+			return false;
+		}
+
+		OutTargetFile.Package = FindOrAddPackage(Arguments, PackageName, Packages, PackageNameMap, PackageIdMap);
+		OutTargetFile.Package->PackageStoreEntry = PackageStore.GetPackageStoreEntry(OutTargetFile.Package->GlobalPackageId);
+		if (!OutTargetFile.Package->PackageStoreEntry)
 		{
-			OutTargetFile.Package = FindOrAddPackage(Arguments, PackageName, Packages, PackageNameMap, PackageIdMap);
+			UE_LOG(LogIoStore, Warning, TEXT("Failed to find package store entry for package: '%s'"), *PackageName.ToString());
+			return false;
 		}
 
-		if (SourceFile.NormalizedPath.EndsWith(TEXT(".m.ubulk")))
+		if (Extension == TEXT(".m.ubulk"))
 		{
 			OutTargetFile.ChunkType = EContainerChunkType::MemoryMappedBulkData;
+			OutTargetFile.Package->TotalBulkDataSize += OutTargetFile.SourceSize;
 		}
-		else if (SourceFile.NormalizedPath.EndsWith(TEXT(".ubulk")))
+		else if (Extension == TEXT(".ubulk"))
 		{
 			OutTargetFile.ChunkType = EContainerChunkType::BulkData;
+			OutTargetFile.Package->TotalBulkDataSize += OutTargetFile.SourceSize;
 		}
-		else if (SourceFile.NormalizedPath.EndsWith(TEXT(".uptnl")))
+		else if (Extension == TEXT(".uptnl"))
 		{
 			OutTargetFile.ChunkType = EContainerChunkType::OptionalBulkData;
+			OutTargetFile.Package->TotalBulkDataSize += OutTargetFile.SourceSize;
 		}
-		else
+		else if (Extension == TEXT(".uasset") || Extension == TEXT(".umap"))
 		{
-			check(OutTargetFile.Package);
-			OutTargetFile.Package->PackageStoreEntry = PackageStore.GetPackageStoreEntry(OutTargetFile.Package->GlobalPackageId);
-			if (!OutTargetFile.Package->PackageStoreEntry)
-			{
-				UE_LOG(LogIoStore, Warning, TEXT("Failed to find package store entry for package: '%s'"), *PackageName.ToString());
-				return false;
-			}
 			OutTargetFile.ChunkType = EContainerChunkType::PackageData;
 			OutTargetFile.Package->FileName = SourceFile.NormalizedPath;
 			OutTargetFile.Package->UAssetSize = OutTargetFile.SourceSize;
 		}
-		
-		// Only keep the regions for the file if neither compression nor encryption are enabled, otherwise the regions will be meaningless.
-		if (Arguments.bFileRegions && !SourceFile.bNeedsCompression && !SourceFile.bNeedsEncryption)
+		else
 		{
-			FString RegionsFilename = OutTargetFile.ChunkType == EContainerChunkType::PackageData
-				? FPaths::ChangeExtension(SourceFile.NormalizedPath, FString(TEXT(".uexp")) + FFileRegion::RegionsFileExtension)
-				: SourceFile.NormalizedPath + FFileRegion::RegionsFileExtension;
-			TUniquePtr<FArchive> RegionsFile(IFileManager::Get().CreateFileReader(*RegionsFilename));
-			if (RegionsFile.IsValid())
-			{
-				FFileRegion::SerializeFileRegions(*RegionsFile.Get(), OutTargetFile.FileRegions);
-			}
+			UE_LOG(LogIoStore, Warning, TEXT("Unexpected file: '%s'"), *SourceFile.NormalizedPath);
+			return false;
 		}
+
+		// Only keep the regions for the file if neither compression nor encryption are enabled, otherwise the regions will be meaningless.
+		// TODO: There are no region files when cooking for Zen
+		//if (Arguments.bFileRegions && !SourceFile.bNeedsCompression && !SourceFile.bNeedsEncryption)
+		//{
+		//	FString RegionsFilename = OutTargetFile.ChunkType == EContainerChunkType::PackageData
+		//		? FPaths::ChangeExtension(SourceFile.NormalizedPath, FString(TEXT(".uexp")) + FFileRegion::RegionsFileExtension)
+		//		: SourceFile.NormalizedPath + FFileRegion::RegionsFileExtension;
+		//	TUniquePtr<FArchive> RegionsFile(IFileManager::Get().CreateFileReader(*RegionsFilename));
+		//	if (RegionsFile.IsValid())
+		//	{
+		//		FFileRegion::SerializeFileRegions(*RegionsFile.Get(), OutTargetFile.FileRegions);
+		//	}
+		//}
 		
 		return true;
 	};
@@ -5320,103 +5452,22 @@ static bool ParsePakOrderFile(const TCHAR* FilePath, FFileOrderMap& Map, const F
 class FCookedFileVisitor : public IPlatformFile::FDirectoryStatVisitor
 {
 	FCookedFileStatMap& CookedFileStatMap;
-	FContainerSourceSpec* ContainerSpec = nullptr;
-	bool bFileRegions;
 
 public:
-	FCookedFileVisitor(FCookedFileStatMap& InCookedFileSizes, FContainerSourceSpec* InContainerSpec, bool bInFileRegions)
+	FCookedFileVisitor(FCookedFileStatMap& InCookedFileSizes)
 		: CookedFileStatMap(InCookedFileSizes)
-		, ContainerSpec(InContainerSpec)
-		, bFileRegions(bInFileRegions)
-	{}
-
-	FCookedFileVisitor(FCookedFileStatMap& InFileSizes, bool bInFileRegions)
-		: CookedFileStatMap(InFileSizes)
-		, bFileRegions(bInFileRegions)
-	{}
+	{
+		
+	}
 
 	virtual bool Visit(const TCHAR* FilenameOrDirectory, const FFileStatData& StatData)
 	{
-		// Should match FCookedFileStatData::EFileExt
-		static const TCHAR* Extensions[] = { TEXT("umap"), TEXT("uasset"), TEXT("uexp"), TEXT("ubulk"), TEXT("uptnl"), TEXT("m.ubulk"), TEXT("ushaderbytecode") };
-		static const int32 NumPackageExtensions = 2;
-		static const int32 UExpExtensionIndex = 2;
-		static const int32 UShaderByteCodeExtensionIndex = 6;
-
 		if (StatData.bIsDirectory)
 		{
 			return true;
 		}
 
-		const TCHAR* Extension = FCString::Strrchr(FilenameOrDirectory, '.');
-		if (!Extension || *(++Extension) == TEXT('\0'))
-		{
-			return true;
-		}
-
-		int32 ExtIndex = 0;
-		if (0 == FCString::Stricmp(Extension, Extensions[3]))
-		{
-			ExtIndex = 3;
-			if (0 == FCString::Stricmp(Extension - 3, TEXT(".m.ubulk")))
-			{
-				ExtIndex = 5;
-			}
-		}
-		else
-		{
-			for (ExtIndex = 0; ExtIndex < UE_ARRAY_COUNT(Extensions); ++ExtIndex)
-			{
-				if (0 == FCString::Stricmp(Extension, Extensions[ExtIndex]))
-					break;
-			}
-		}
-
-		if (ExtIndex >= UE_ARRAY_COUNT(Extensions))
-		{
-			return true;
-		}
-
-		FString Path = FilenameOrDirectory;
-		FPaths::NormalizeFilename(Path);
-
-		if (ContainerSpec && ExtIndex != UExpExtensionIndex)
-		{
-			FContainerSourceFile& FileEntry = ContainerSpec->SourceFiles.AddDefaulted_GetRef();
-			FileEntry.NormalizedPath = Path;
-		}
-
-		// Read the matching regions file, if it exists.
-		TUniquePtr<FArchive> RegionsFile;
-		if (bFileRegions)
-		{
-			RegionsFile.Reset(IFileManager::Get().CreateFileReader(*(Path + FFileRegion::RegionsFileExtension)));
-		}
-
-		FCookedFileStatData& CookedFileStatData = CookedFileStatMap.Add(MoveTemp(Path));
-		CookedFileStatData.FileSize = StatData.FileSize;
-		CookedFileStatData.FileExt = FCookedFileStatData::EFileExt(ExtIndex);
-		if (ExtIndex < NumPackageExtensions)
-		{
-			CookedFileStatData.FileType = FCookedFileStatData::PackageHeader;
-		}
-		else if (ExtIndex == UExpExtensionIndex)
-		{
-			CookedFileStatData.FileType = FCookedFileStatData::PackageData;
-		}
-		else if (ExtIndex == UShaderByteCodeExtensionIndex)
-		{
-			CookedFileStatData.FileType = FCookedFileStatData::ShaderLibrary;
-		}
-		else
-		{
-			CookedFileStatData.FileType = FCookedFileStatData::BulkData;
-		}
-
-		if (RegionsFile.IsValid())
-		{
-			FFileRegion::SerializeFileRegions(*RegionsFile.Get(), CookedFileStatData.FileRegions);
-		}
+		CookedFileStatMap.Add(FilenameOrDirectory, StatData.FileSize);
 
 		return true;
 	}
@@ -6005,7 +6056,7 @@ int32 CreateIoStoreContainerFiles(const TCHAR* CmdLine)
 	{
 		IOSTORE_CPU_SCOPE(FindCookedAssets);
 		UE_LOG(LogIoStore, Display, TEXT("Searching for cooked assets in folder '%s'"), *Arguments.CookedDir);
-		FCookedFileVisitor CookedFileVistor(Arguments.CookedFileStatMap, nullptr, Arguments.bFileRegions);
+		FCookedFileVisitor CookedFileVistor(Arguments.CookedFileStatMap);
 		IFileManager::Get().IterateDirectoryStatRecursively(*Arguments.CookedDir, CookedFileVistor);
 		UE_LOG(LogIoStore, Display, TEXT("Found '%d' files"), Arguments.CookedFileStatMap.Num());
 	}
