IoStore: Fixed race when deleting write queue entries

diff --git a/Engine/Source/Developer/IoStoreUtilities/Private/IoStoreUtilities.cpp b/Engine/Source/Developer/IoStoreUtilities/Private/IoStoreUtilities.cpp
--- a/Engine/Source/Developer/IoStoreUtilities/Private/IoStoreUtilities.cpp
+++ b/Engine/Source/Developer/IoStoreUtilities/Private/IoStoreUtilities.cpp
@@ -2462,7 +2462,7 @@ private:
 
 		void OnSourceBufferLoaded()
 		{
-			Manager.ScheduleRetire(this);
+			QueueEntry->ReleaseRef(Manager);
 			CompletionEvent->DispatchSubsequents();
 		}
 
@@ -2505,6 +2505,7 @@ private:
 			QueueEntry->FileHandle.Reset(
 				FPlatformFileManager::Get().GetPlatformFile().OpenAsyncRead(*TargetFile.NormalizedSourcePath));
 			
+			QueueEntry->AddRef(); // Must keep it around until we've assigned the ReadRequest pointer
 			FAsyncFileCallBack Callback = [this](bool, IAsyncReadRequest* ReadRequest)
 			{
 				if (TargetFile.ChunkType == EContainerChunkType::PackageData)
@@ -2517,6 +2518,7 @@ private:
 
 			QueueEntry->ReadRequest.Reset(
 				QueueEntry->FileHandle->ReadRequest(0, SourceBuffer.DataSize(), AIOP_Normal, &Callback, SourceBuffer.Data()));
+			QueueEntry->ReleaseRef(Manager);
 		}
 	};
 
@@ -2554,6 +2556,22 @@ private:
 		TUniquePtr<IAsyncReadFileHandle> FileHandle;
 		TUniquePtr<IAsyncReadRequest> ReadRequest;
 		FWriteContainerTargetFileRequest* WriteRequest = nullptr;
+
+		void AddRef()
+		{
+			++RefCount;
+		}
+
+		void ReleaseRef(FIoStoreWriteRequestManager& Manager)
+		{
+			if (--RefCount == 0)
+			{
+				Manager.ScheduleRetire(this);
+			}
+		}
+
+	private:
+		TAtomic<int32> RefCount{ 1 };
 	};
 
 	class FQueue
@@ -2637,9 +2655,9 @@ private:
 		InitiatorQueue.Enqueue(QueueEntry);
 	}
 
-	void ScheduleRetire(FWriteContainerTargetFileRequest* WriteRequest)
+	void ScheduleRetire(FQueueEntry* QueueEntry)
 	{
-		RetirerQueue.Enqueue(WriteRequest->QueueEntry);
+		RetirerQueue.Enqueue(QueueEntry);
 	}
 
 	void Start(FQueueEntry* QueueEntry)
