AsyncLoading2: Implemented loading of optional package segments
#rnx

diff --git a/Engine/Source/Developer/IoStoreUtilities/Internal/PackageStoreOptimizer.h b/Engine/Source/Developer/IoStoreUtilities/Internal/PackageStoreOptimizer.h
--- a/Engine/Source/Developer/IoStoreUtilities/Internal/PackageStoreOptimizer.h
+++ b/Engine/Source/Developer/IoStoreUtilities/Internal/PackageStoreOptimizer.h
@@ -204,7 +204,6 @@ private:
 		FName FromPackageName;
 		int32 FromPackageNameLen = 0;
 		FPackageId FromPackageId;
-		FPackageId FromOptionalPackageId;
 		bool bIsScriptImport = false;
 		bool bIsImportOfPackage = false;
 		bool bIsImportOptional = false;
@@ -295,7 +294,6 @@ private:
 	bool bPermanentMark = false;
 
 	bool bIsRedirected = false;
-	bool bIsOptional = false;
 
 	friend class FPackageStoreOptimizer;
 };
@@ -341,8 +339,14 @@ public:
 	FPackageStorePackage* CreatePackageFromPackageStoreHeader(const FName& Name, const FIoBuffer& Buffer, const FPackageStoreEntryResource& PackageStoreEntry) const;
 	void FinalizePackage(FPackageStorePackage* Package);
 	FIoBuffer CreatePackageBuffer(const FPackageStorePackage* Package, const FIoBuffer& CookedExportsBuffer, TArray<FFileRegion>* InOutFileRegions) const;
-	FPackageStoreEntryResource CreatePackageStoreEntry(const FPackageStorePackage* Package) const;
-	FIoContainerHeader CreateContainerHeader(const FIoContainerId& ContainerId, TArrayView<const FPackageStoreEntryResource> PackageStoreEntries) const;
+	FPackageStoreEntryResource CreatePackageStoreEntry(const FPackageStorePackage* Package, const FPackageStorePackage* OptionalSegmentPackage) const;
+	enum EContainerHeaderInclusionFilter
+	{
+		IncludeNonOptionalSegments = 0x1,
+		IncludeOptionalSegments = 0x2,
+		IncludeAllSegments = IncludeNonOptionalSegments | IncludeOptionalSegments
+	};
+	FIoContainerHeader CreateContainerHeader(const FIoContainerId& ContainerId, TArrayView<const FPackageStoreEntryResource> PackageStoreEntries, EContainerHeaderInclusionFilter InclusionFilter) const;
 	IOSTOREUTILITIES_API FIoBuffer CreateScriptObjectsBuffer() const;
 	void LoadScriptObjectsBuffer(const FIoBuffer& ScriptObjectsBuffer);
 	void ProcessRedirects(const TMap<FPackageId, FPackageStorePackage*>& PackagesMap, bool bIsBuildingDLC) const;
diff --git a/Engine/Source/Developer/IoStoreUtilities/Private/IoStoreUtilities.cpp b/Engine/Source/Developer/IoStoreUtilities/Private/IoStoreUtilities.cpp
--- a/Engine/Source/Developer/IoStoreUtilities/Private/IoStoreUtilities.cpp
+++ b/Engine/Source/Developer/IoStoreUtilities/Private/IoStoreUtilities.cpp
@@ -304,6 +304,9 @@ struct FCookedFileStatData
 		MemoryMappedBulkData,
 		ShaderLibrary,
 		Regions,
+		OptionalSegmentPackageHeader,
+		OptionalSegmentPackageData,
+		OptionalSegmentBulkData,
 		Invalid
 	};
 
@@ -370,6 +373,10 @@ public:
 		Extensions.Emplace(*FString::Printf(TEXT(".ubulk%s"), FFileRegion::RegionsFileExtension), FCookedFileStatData::Regions);
 		Extensions.Emplace(*FString::Printf(TEXT(".m.ubulk%s"), FFileRegion::RegionsFileExtension), FCookedFileStatData::Regions);
 		Extensions.Emplace(TEXT(".ushaderbytecode"), FCookedFileStatData::ShaderLibrary);
+		Extensions.Emplace(TEXT(".o.umap"), FCookedFileStatData::OptionalSegmentPackageHeader);
+		Extensions.Emplace(TEXT(".o.uasset"), FCookedFileStatData::OptionalSegmentPackageHeader);
+		Extensions.Emplace(TEXT(".o.uexp"), FCookedFileStatData::OptionalSegmentPackageData);
+		Extensions.Emplace(TEXT(".o.ubulk"), FCookedFileStatData::OptionalSegmentBulkData);
 	}
 
 	int32 Num() const
@@ -445,9 +452,13 @@ struct FLegacyCookedPackage
 	FString FileName;
 	uint64 UAssetSize = 0;
 	uint64 UExpSize = 0;
+	FString OptionalSegmentFileName;
+	uint64 OptionalSegmentUAssetSize = 0;
+	uint64 OptionalSegmentUExpSize = 0;
 	uint64 TotalBulkDataSize = 0;
 	uint64 DiskLayoutOrder = MAX_uint64;
 	FPackageStorePackage* OptimizedPackage = nullptr;
+	FPackageStorePackage* OptimizedOptionalSegmentPackage = nullptr;
 	TArray<FShaderInfo*> Shaders;
 	const FPackageStoreEntryResource* PackageStoreEntry = nullptr;
 };
@@ -460,6 +471,8 @@ enum class EContainerChunkType
 	BulkData,
 	OptionalBulkData,
 	MemoryMappedBulkData,
+	OptionalSegmentPackageData,
+	OptionalSegmentBulkData,
 	Invalid
 };
 
@@ -515,8 +528,8 @@ public:
 			FIoContainerHeader ContainerHeader;
 			Ar << ContainerHeader;
 
-			PackageEntries.Reserve(ContainerHeader.PackageCount);
-			TArrayView<const FFilePackageStoreEntry> FilePackageEntries = MakeArrayView<const FFilePackageStoreEntry>(reinterpret_cast<const FFilePackageStoreEntry*>(ContainerHeader.StoreEntries.GetData()), ContainerHeader.PackageCount);
+			PackageEntries.Reserve(ContainerHeader.PackageIds.Num());
+			TArrayView<const FFilePackageStoreEntry> FilePackageEntries = MakeArrayView<const FFilePackageStoreEntry>(reinterpret_cast<const FFilePackageStoreEntry*>(ContainerHeader.StoreEntries.GetData()), ContainerHeader.PackageIds.Num());
 
 			int32 Idx = 0;
 			for (const FFilePackageStoreEntry& FilePackageEntry : FilePackageEntries)
@@ -539,9 +552,9 @@ public:
 		{
 			FName PackageName = FName(PackageInfo.PackageName);
 			PackageNameToPackageInfoMap.Add(PackageName, &PackageInfo);
-			if (PackageInfo.PackageChunkId.IsValid())
+			for (const FIoChunkId& ChunkId : PackageInfo.ExportBundleChunkIds)
 			{
-				ChunkIdToPackageNameMap.Add(PackageInfo.PackageChunkId, PackageName);
+				ChunkIdToPackageNameMap.Add(ChunkId, PackageName);
 			}
 			for (const FIoChunkId& ChunkId : PackageInfo.BulkDataChunkIds)
 			{
@@ -759,10 +772,13 @@ struct FContainerTargetSpec
 {
 	FIoContainerId ContainerId;
 	FIoContainerHeader Header;
+	FIoContainerHeader OptionalSegmentHeader;
 	FName Name;
 	FGuid EncryptionKeyGuid;
 	FString OutputPath;
+	FString OptionalSegmentOutputPath;
 	TSharedPtr<IIoStoreWriter> IoStoreWriter;
+	TSharedPtr<IIoStoreWriter> OptionalSegmentIoStoreWriter;
 	TArray<FContainerTargetFile> TargetFiles;
 	TArray<TUniquePtr<FIoStoreReader>> PatchSourceReaders;
 	EIoContainerFlags ContainerFlags = EIoContainerFlags::None;
@@ -1359,48 +1375,76 @@ static void ParsePackageAssetsFromFiles(TArray<FLegacyCookedPackage*>& Packages,
 	TArray<FPackageFileSummary> PackageFileSummaries;
 	PackageFileSummaries.SetNum(TotalPackageCount);
 
+	TArray<FPackageFileSummary> OptionalSegmentPackageFileSummaries;
+	OptionalSegmentPackageFileSummaries.SetNum(TotalPackageCount);
+
 	uint8* UAssetMemory = nullptr;
+	uint8* OptionalSegmentUAssetMemory = nullptr;
+
 	TArray<uint8*> PackageAssetBuffers;
 	PackageAssetBuffers.SetNum(TotalPackageCount);
 
+	TArray<uint8*> OptionalSegmentPackageAssetBuffers;
+	OptionalSegmentPackageAssetBuffers.SetNum(TotalPackageCount);
+
 	UE_LOG(LogIoStore, Display, TEXT("Reading package assets..."));
 	{
 		IOSTORE_CPU_SCOPE(ReadUAssetFiles);
 
 		uint64 TotalUAssetSize = 0;
+		uint64 TotalOptionalSegmentUAssetSize = 0;
 		for (const FLegacyCookedPackage* Package : Packages)
 		{
 			TotalUAssetSize += Package->UAssetSize;
+			TotalOptionalSegmentUAssetSize += Package->OptionalSegmentUAssetSize;
 		}
 		UAssetMemory = reinterpret_cast<uint8*>(FMemory::Malloc(TotalUAssetSize));
 		uint8* UAssetMemoryPtr = UAssetMemory;
+		OptionalSegmentUAssetMemory = reinterpret_cast<uint8*>(FMemory::Malloc(TotalOptionalSegmentUAssetSize));
+		uint8* OptionalSegmentUAssetMemoryPtr = OptionalSegmentUAssetMemory;
+
 		for (int32 Index = 0; Index < TotalPackageCount; ++Index)
 		{
 			PackageAssetBuffers[Index] = UAssetMemoryPtr;
 			UAssetMemoryPtr += Packages[Index]->UAssetSize;
+			OptionalSegmentPackageAssetBuffers[Index] = OptionalSegmentUAssetMemoryPtr;
+			OptionalSegmentUAssetMemoryPtr += Packages[Index]->OptionalSegmentUAssetSize;
 		}
 
 		TAtomic<uint64> CurrentFileIndex{ 0 };
-		ParallelFor(TotalPackageCount, [&ReadCount, &PackageAssetBuffers, &Packages, &CurrentFileIndex](int32 Index)
+		ParallelFor(TotalPackageCount, [&ReadCount, &PackageAssetBuffers, &OptionalSegmentPackageAssetBuffers, &Packages, &CurrentFileIndex](int32 Index)
 		{
 			TRACE_CPUPROFILER_EVENT_SCOPE(ReadUAssetFile);
 			FLegacyCookedPackage* Package = Packages[Index];
-			if (!Package->UAssetSize)
-			{
-				return;
-			}
-			uint8* Buffer = PackageAssetBuffers[Index];
-			IFileHandle* FileHandle = FPlatformFileManager::Get().GetPlatformFile().OpenRead(*Package->FileName);
-			if (FileHandle)
+			if (Package->UAssetSize)
 			{
-				bool bSuccess = FileHandle->Read(Buffer, Package->UAssetSize);
-				UE_CLOG(!bSuccess, LogIoStore, Warning, TEXT("Failed reading file '%s'"), *Package->FileName);
-				delete FileHandle;
+				uint8* Buffer = PackageAssetBuffers[Index];
+				TUniquePtr<IFileHandle> FileHandle(FPlatformFileManager::Get().GetPlatformFile().OpenRead(*Package->FileName));
+				if (FileHandle)
+				{
+					bool bSuccess = FileHandle->Read(Buffer, Package->UAssetSize);
+					UE_CLOG(!bSuccess, LogIoStore, Warning, TEXT("Failed reading file '%s'"), *Package->FileName);
+				}
+				else
+				{
+					UE_LOG(LogIoStore, Warning, TEXT("Couldn't open file '%s'"), *Package->FileName);
+				}
 			}
-			else
+			if (Package->OptionalSegmentUAssetSize)
 			{
-				UE_LOG(LogIoStore, Warning, TEXT("Couldn't open file '%s'"), *Package->FileName);
+				uint8* Buffer = OptionalSegmentPackageAssetBuffers[Index];
+				TUniquePtr<IFileHandle> FileHandle(FPlatformFileManager::Get().GetPlatformFile().OpenRead(*Package->OptionalSegmentFileName));
+				if (FileHandle)
+				{
+					bool bSuccess = FileHandle->Read(Buffer, Package->OptionalSegmentUAssetSize);
+					UE_CLOG(!bSuccess, LogIoStore, Warning, TEXT("Failed reading file '%s'"), *Package->OptionalSegmentFileName);
+				}
+				else
+				{
+					UE_LOG(LogIoStore, Warning, TEXT("Couldn't open file '%s'"), *Package->OptionalSegmentFileName);
+				}
 			}
+
 			uint64 LocalFileIndex = CurrentFileIndex.IncrementExchange() + 1;
 			UE_CLOG(LocalFileIndex % 1000 == 0, LogIoStore, Display, TEXT("Reading %d/%d: '%s'"), LocalFileIndex, Packages.Num(), *Package->FileName);
 		}, EParallelForFlags::Unbalanced);
@@ -1412,15 +1456,16 @@ static void ParsePackageAssetsFromFiles(TArray<FLegacyCookedPackage*>& Packages,
 		ParallelFor(TotalPackageCount, [
 				&ReadCount,
 				&PackageAssetBuffers,
+				&OptionalSegmentPackageAssetBuffers,
 				&PackageFileSummaries,
 				&Packages,
 				&PackageStoreOptimizer](int32 Index)
 		{
-			uint8* PackageBuffer = PackageAssetBuffers[Index];
 			FLegacyCookedPackage* Package = Packages[Index];
 
 			if (Package->UAssetSize)
 			{
+				uint8* PackageBuffer = PackageAssetBuffers[Index];
 				FIoBuffer CookedHeaderBuffer = FIoBuffer(FIoBuffer::Wrap, PackageBuffer, Package->UAssetSize);
 				Package->OptimizedPackage = PackageStoreOptimizer.CreatePackageFromCookedHeader(Package->PackageName, CookedHeaderBuffer);
 			}
@@ -1429,10 +1474,18 @@ static void ParsePackageAssetsFromFiles(TArray<FLegacyCookedPackage*>& Packages,
 				Package->OptimizedPackage = PackageStoreOptimizer.CreateMissingPackage(Package->PackageName);
 			}
 			check(Package->OptimizedPackage->GetId() == Package->GlobalPackageId);
+			if (Package->OptionalSegmentUAssetSize)
+			{
+				uint8* OptionalSegmentPackageBuffer = OptionalSegmentPackageAssetBuffers[Index];
+				FIoBuffer OptionalSegmentCookedHeaderBuffer = FIoBuffer(FIoBuffer::Wrap, OptionalSegmentPackageBuffer, Package->OptionalSegmentUAssetSize);
+				Package->OptimizedOptionalSegmentPackage = PackageStoreOptimizer.CreatePackageFromCookedHeader(Package->PackageName, OptionalSegmentCookedHeaderBuffer);
+				check(Package->OptimizedOptionalSegmentPackage->GetId() == Package->GlobalPackageId);
+			}
 		}, EParallelForFlags::Unbalanced);
 	}
 
 	FMemory::Free(UAssetMemory);
+	FMemory::Free(OptionalSegmentUAssetMemory);
 }
 
 static void ParsePackageAssetsFromPackageStore(FCookedPackageStore& PackageStore, TArray<FLegacyCookedPackage*>& Packages, const FPackageStoreOptimizer& PackageStoreOptimizer)
@@ -1956,6 +2009,21 @@ void InitializeContainerTargetsAndPackages(
 			}
 			OutTargetFile.NormalizedSourcePath = UexpPath;
 		}
+		else if (CookedFileStatData->FileType == FCookedFileStatData::OptionalSegmentPackageHeader)
+		{
+			FStringView NormalizedSourcePathView(SourceFile.NormalizedPath);
+			int32 ExtensionStartIndex = GetFullExtensionStartIndex(NormalizedSourcePathView);
+			TStringBuilder<512> UexpPath;
+			UexpPath.Append(NormalizedSourcePathView.Left(ExtensionStartIndex));
+			UexpPath.Append(TEXT(".o.uexp"));
+			CookedFileStatData = Arguments.CookedFileStatMap.Find(*UexpPath);
+			if (!CookedFileStatData)
+			{
+				UE_LOG(LogIoStore, Warning, TEXT("Couldn't find .o.uexp file for: '%s'"), *SourceFile.NormalizedPath);
+				return false;
+			}
+			OutTargetFile.NormalizedSourcePath = UexpPath;
+		}
 		else
 		{
 			OutTargetFile.NormalizedSourcePath = SourceFile.NormalizedPath;
@@ -1994,6 +2062,17 @@ void InitializeContainerTargetsAndPackages(
 				OutTargetFile.ChunkId = CreateIoChunkId(OutTargetFile.Package->GlobalPackageId.Value(), 0, EIoChunkType::MemoryMappedBulkData);
 				OutTargetFile.Package->TotalBulkDataSize += CookedFileStatData->FileSize;
 				break;
+			case FCookedFileStatData::OptionalSegmentPackageData:
+				OutTargetFile.ChunkType = EContainerChunkType::OptionalSegmentPackageData;
+				OutTargetFile.ChunkId = CreateIoChunkId(OutTargetFile.Package->GlobalPackageId.Value(), 1, EIoChunkType::ExportBundleData);
+				OutTargetFile.Package->OptionalSegmentFileName = SourceFile.NormalizedPath; // .o.uasset path
+				OutTargetFile.Package->OptionalSegmentUAssetSize = OriginalCookedFileStatData->FileSize;
+				OutTargetFile.Package->OptionalSegmentUExpSize = CookedFileStatData->FileSize;
+				break;
+			case FCookedFileStatData::OptionalSegmentBulkData:
+				OutTargetFile.ChunkType = EContainerChunkType::OptionalSegmentBulkData;
+				OutTargetFile.ChunkId = CreateIoChunkId(OutTargetFile.Package->GlobalPackageId.Value(), 1, EIoChunkType::BulkData);
+				break;
 			default:
 				UE_LOG(LogIoStore, Fatal, TEXT("Unexpected file type %d for file '%s'"), CookedFileStatData->FileType, *OutTargetFile.NormalizedSourcePath);
 				return false;
@@ -2100,6 +2179,16 @@ void InitializeContainerTargetsAndPackages(
 			OutTargetFile.Package->FileName = SourceFile.NormalizedPath;
 			OutTargetFile.Package->UAssetSize = OutTargetFile.SourceSize;
 		}
+		else if (Extension == TEXT(".o.uasset") || Extension == TEXT(".o.umap"))
+		{
+			OutTargetFile.ChunkType = EContainerChunkType::OptionalSegmentPackageData;
+			OutTargetFile.Package->OptionalSegmentFileName = SourceFile.NormalizedPath;
+			OutTargetFile.Package->OptionalSegmentUAssetSize = OutTargetFile.SourceSize;
+		}
+		else if (Extension == TEXT(".o.ubulk"))
+		{
+			OutTargetFile.ChunkType = EContainerChunkType::OptionalSegmentBulkData;
+		}
 		else
 		{
 			UE_LOG(LogIoStore, Warning, TEXT("Unexpected file: '%s'"), *SourceFile.NormalizedPath);
@@ -2142,6 +2231,7 @@ void InitializeContainerTargetsAndPackages(
 
 		{
 			IOSTORE_CPU_SCOPE(ProcessSourceFiles);
+			bool bHasOptionalSegmentPackages = false;
 			for (const FContainerSourceFile& SourceFile : ContainerSource.SourceFiles)
 			{
 				FContainerTargetFile TargetFile;
@@ -2173,9 +2263,18 @@ void InitializeContainerTargetsAndPackages(
 					check(TargetFile.Package);
 					ContainerTarget->Packages.Add(TargetFile.Package);
 				}
+				else if (TargetFile.ChunkType == EContainerChunkType::OptionalSegmentPackageData)
+				{
+					bHasOptionalSegmentPackages = true;
+				}
 
 				ContainerTarget->TargetFiles.Emplace(MoveTemp(TargetFile));
 			}
+
+			if (bHasOptionalSegmentPackages)
+			{
+				ContainerTarget->OptionalSegmentOutputPath = ContainerTarget->OutputPath + FPackagePath::GetOptionalSegmentExtensionModifier();
+			}
 		}
 	}
 
@@ -2513,6 +2612,12 @@ private:
 					SourceBuffer = Manager.PackageStoreOptimizer.CreatePackageBuffer(TargetFile.Package->OptimizedPackage, SourceBuffer, bHasUpdatedExportBundleRegions ? nullptr : &FileRegions);
 					bHasUpdatedExportBundleRegions = true;
 				}
+				else if (TargetFile.ChunkType == EContainerChunkType::OptionalSegmentPackageData)
+				{
+					check(TargetFile.Package->OptimizedOptionalSegmentPackage);
+					SourceBuffer = Manager.PackageStoreOptimizer.CreatePackageBuffer(TargetFile.Package->OptimizedOptionalSegmentPackage, SourceBuffer, bHasUpdatedExportBundleRegions ? nullptr : &FileRegions);
+					bHasUpdatedExportBundleRegions = true;
+				}
 				OnSourceBufferLoaded();
 			};
 
@@ -3087,6 +3192,11 @@ int32 CreateTarget(const FIoStoreArguments& Arguments, const FIoStoreWriterSetti
 				ContainerTarget->IoStoreWriter = IoStoreWriterContext->CreateContainer(*ContainerTarget->OutputPath, ContainerSettings);
 				ContainerTarget->IoStoreWriter->EnableDiskLayoutOrdering(ContainerTarget->PatchSourceReaders);
 				IoStoreWriters.Add(ContainerTarget->IoStoreWriter);
+				if (!ContainerTarget->OptionalSegmentOutputPath.IsEmpty())
+				{
+					ContainerTarget->OptionalSegmentIoStoreWriter = IoStoreWriterContext->CreateContainer(*ContainerTarget->OptionalSegmentOutputPath, ContainerSettings);
+					IoStoreWriters.Add(ContainerTarget->OptionalSegmentIoStoreWriter);
+				}
 			}
 		}
 	}
@@ -3123,14 +3233,21 @@ int32 CreateTarget(const FIoStoreArguments& Arguments, const FIoStoreWriterSetti
 		{
 			for (FContainerTargetFile& TargetFile : ContainerTarget->TargetFiles)
 			{
-				if (TargetFile.ChunkType != EContainerChunkType::PackageData)
+				if (TargetFile.ChunkType != EContainerChunkType::PackageData && TargetFile.ChunkType != EContainerChunkType::OptionalSegmentPackageData)
 				{
 					FIoWriteOptions WriteOptions;
 					WriteOptions.DebugName = *TargetFile.DestinationPath;
 					WriteOptions.bForceUncompressed = TargetFile.bForceUncompressed;
 					WriteOptions.bIsMemoryMapped = TargetFile.ChunkType == EContainerChunkType::MemoryMappedBulkData;
 					WriteOptions.FileName = TargetFile.DestinationPath;
-					ContainerTarget->IoStoreWriter->Append(TargetFile.ChunkId, WriteRequestManager.Read(TargetFile), WriteOptions);
+					if (TargetFile.ChunkType == EContainerChunkType::OptionalSegmentBulkData)
+					{
+						ContainerTarget->OptionalSegmentIoStoreWriter->Append(TargetFile.ChunkId, WriteRequestManager.Read(TargetFile), WriteOptions);
+					}
+					else
+					{
+						ContainerTarget->IoStoreWriter->Append(TargetFile.ChunkId, WriteRequestManager.Read(TargetFile), WriteOptions);
+					}
 				}
 			}
 		}
@@ -3154,6 +3271,10 @@ int32 CreateTarget(const FIoStoreArguments& Arguments, const FIoStoreWriterSetti
 	{
 		check(Package->OptimizedPackage);
 		PackageStoreOptimizer.FinalizePackage(Package->OptimizedPackage);
+		if (Package->OptimizedOptionalSegmentPackage)
+		{
+			PackageStoreOptimizer.FinalizePackage(Package->OptimizedOptionalSegmentPackage);
+		}
 	}
 
 	UE_LOG(LogIoStore, Display, TEXT("Creating disk layout..."));
@@ -3171,31 +3292,51 @@ int32 CreateTarget(const FIoStoreArguments& Arguments, const FIoStoreWriterSetti
 			TArray<FPackageStoreEntryResource> PackageStoreEntries;
 			for (FContainerTargetFile& TargetFile : ContainerTarget->TargetFiles)
 			{
-				if (TargetFile.ChunkType == EContainerChunkType::PackageData)
+				if (TargetFile.ChunkType == EContainerChunkType::PackageData || TargetFile.ChunkType == EContainerChunkType::OptionalSegmentPackageData)
 				{
 					check(TargetFile.Package);
 					FIoWriteOptions WriteOptions;
 					WriteOptions.DebugName = *TargetFile.DestinationPath;
 					WriteOptions.bForceUncompressed = TargetFile.bForceUncompressed;
 					WriteOptions.FileName = TargetFile.DestinationPath;
-					ContainerTarget->IoStoreWriter->Append(TargetFile.ChunkId, WriteRequestManager.Read(TargetFile), WriteOptions);
-					PackageStoreEntries.Add(PackageStoreOptimizer.CreatePackageStoreEntry(TargetFile.Package->OptimizedPackage));
+					if (TargetFile.ChunkType == EContainerChunkType::OptionalSegmentPackageData)
+					{
+						check(ContainerTarget->OptionalSegmentIoStoreWriter);
+						check(TargetFile.Package->OptimizedOptionalSegmentPackage);
+						ContainerTarget->OptionalSegmentIoStoreWriter->Append(TargetFile.ChunkId, WriteRequestManager.Read(TargetFile), WriteOptions);
+					}
+					else
+					{
+						ContainerTarget->IoStoreWriter->Append(TargetFile.ChunkId, WriteRequestManager.Read(TargetFile), WriteOptions);
+						PackageStoreEntries.Add(PackageStoreOptimizer.CreatePackageStoreEntry(TargetFile.Package->OptimizedPackage, TargetFile.Package->OptimizedOptionalSegmentPackage));
+					}
 				}
 			}
 
-			ContainerTarget->Header = PackageStoreOptimizer.CreateContainerHeader(ContainerTarget->ContainerId, PackageStoreEntries);
-			FLargeMemoryWriter HeaderAr(0, true);
-			HeaderAr << ContainerTarget->Header;
-			int64 DataSize = HeaderAr.TotalSize();
-			FIoBuffer ContainerHeaderBuffer(FIoBuffer::AssumeOwnership, HeaderAr.ReleaseOwnership(), DataSize);
-			
-			FIoWriteOptions WriteOptions;
-			WriteOptions.DebugName = TEXT("ContainerHeader");
-			WriteOptions.bForceUncompressed = true;
-			ContainerTarget->IoStoreWriter->Append(
-				CreateIoChunkId(ContainerTarget->ContainerId.Value(), 0, EIoChunkType::ContainerHeader),
-				ContainerHeaderBuffer,
-				WriteOptions);
+			auto WriteContainerHeaderChunk = [](FIoContainerHeader& Header, IIoStoreWriter* IoStoreWriter)
+			{
+				FLargeMemoryWriter HeaderAr(0, true);
+				HeaderAr << Header;
+				int64 DataSize = HeaderAr.TotalSize();
+				FIoBuffer ContainerHeaderBuffer(FIoBuffer::AssumeOwnership, HeaderAr.ReleaseOwnership(), DataSize);
+
+				FIoWriteOptions WriteOptions;
+				WriteOptions.DebugName = TEXT("ContainerHeader");
+				WriteOptions.bForceUncompressed = true;
+				IoStoreWriter->Append(
+					CreateIoChunkId(Header.ContainerId.Value(), 0, EIoChunkType::ContainerHeader),
+					ContainerHeaderBuffer,
+					WriteOptions);
+			};
+
+			ContainerTarget->Header = PackageStoreOptimizer.CreateContainerHeader(ContainerTarget->ContainerId, PackageStoreEntries, FPackageStoreOptimizer::IncludeNonOptionalSegments);
+			WriteContainerHeaderChunk(ContainerTarget->Header, ContainerTarget->IoStoreWriter.Get());
+
+			if (ContainerTarget->OptionalSegmentIoStoreWriter)
+			{
+				ContainerTarget->OptionalSegmentHeader = PackageStoreOptimizer.CreateContainerHeader(ContainerTarget->ContainerId, PackageStoreEntries, FPackageStoreOptimizer::IncludeOptionalSegments);
+				WriteContainerHeaderChunk(ContainerTarget->OptionalSegmentHeader, ContainerTarget->OptionalSegmentIoStoreWriter.Get());
+			}
 		}
 
 		// Check if we need to dump the final order of the packages. Useful, to debug packing.
@@ -3997,7 +4138,7 @@ int32 Describe(
 			Job.RawLocalizedPackages = ContainerHeader.LocalizedPackages;
 			Job.RawPackageRedirects = ContainerHeader.PackageRedirects;
 
-			TArrayView<FFilePackageStoreEntry> StoreEntries(reinterpret_cast<FFilePackageStoreEntry*>(ContainerHeader.StoreEntries.GetData()), ContainerHeader.PackageCount);
+			TArrayView<FFilePackageStoreEntry> StoreEntries(reinterpret_cast<FFilePackageStoreEntry*>(ContainerHeader.StoreEntries.GetData()), ContainerHeader.PackageIds.Num());
 
 			int32 PackageIndex = 0;
 			Job.Packages.Reserve(StoreEntries.Num());
@@ -5149,7 +5290,7 @@ int32 Staged2Zen(const FString& BuildPath, const FKeyChain& KeyChain, const FStr
 		ZenStoreWriter->BeginPackage(BeginPackageInfo);
 
 		IPackageWriter::FPackageInfo PackageStorePackageInfo;
-		PackageStorePackageInfo.OutputPackageName = PackageStorePackageInfo.InputPackageName = PackageInfo.PackageName;
+		PackageStorePackageInfo.PackageName = PackageInfo.PackageName;
 		PackageStorePackageInfo.LooseFilePath = PackageInfo.FileName;
 		PackageStorePackageInfo.ChunkId = PackageInfo.Chunk.Value;
 
@@ -5159,7 +5300,7 @@ int32 Staged2Zen(const FString& BuildPath, const FKeyChain& KeyChain, const FStr
 		for (const FBulkDataInfo& BulkDataInfo : PackageInfo.BulkData)
 		{
 			IPackageWriter::FBulkDataInfo PackageStoreBulkDataInfo;
-			PackageStoreBulkDataInfo.OutputPackageName = PackageStoreBulkDataInfo.InputPackageName = PackageInfo.PackageName;
+			PackageStoreBulkDataInfo.PackageName = PackageInfo.PackageName;
 			PackageStoreBulkDataInfo.LooseFilePath = BulkDataInfo.FileName;
 			PackageStoreBulkDataInfo.ChunkId = BulkDataInfo.Chunk.Value;
 			PackageStoreBulkDataInfo.BulkDataType = BulkDataInfo.BulkDataType;
diff --git a/Engine/Source/Developer/IoStoreUtilities/Private/PackageStoreManifest.cpp b/Engine/Source/Developer/IoStoreUtilities/Private/PackageStoreManifest.cpp
--- a/Engine/Source/Developer/IoStoreUtilities/Private/PackageStoreManifest.cpp
+++ b/Engine/Source/Developer/IoStoreUtilities/Private/PackageStoreManifest.cpp
@@ -15,39 +15,38 @@ FPackageStoreManifest::FPackageStoreManifest(const FString& InCookedOutputPath)
 	FPaths::NormalizeFilename(CookedOutputPath);
 }
 
-void FPackageStoreManifest::BeginPackage(FName InputPackageName)
+void FPackageStoreManifest::BeginPackage(FName PackageName)
 {
 	FScopeLock Lock(&CriticalSection);
-	FPackageInfo& PackageInfo = MainPackageInfoByNameMap.FindOrAdd(InputPackageName);
-	PackageInfo.PackageName = InputPackageName;
-	if (PackageInfo.PackageChunkId.IsValid())
+	FPackageInfo& PackageInfo = PackageInfoByNameMap.FindOrAdd(PackageName);
+	PackageInfo.PackageName = PackageName;
+	for (const FIoChunkId& ExportBundleChunkId : PackageInfo.ExportBundleChunkIds)
 	{
-		FileNameByChunkIdMap.Remove(PackageInfo.PackageChunkId);
+		FileNameByChunkIdMap.Remove(ExportBundleChunkId);
 	}
 	for (const FIoChunkId& BulkDataChunkId : PackageInfo.BulkDataChunkIds)
 	{
 		FileNameByChunkIdMap.Remove(BulkDataChunkId);
 	}
 	PackageInfo.BulkDataChunkIds.Reset();
-	OptionalPackageInfoByNameMap.Remove(InputPackageName);
 }
 
-void FPackageStoreManifest::AddPackageData(FName InputPackageName, FName OutputPackageName, const FString& FileName, const FIoChunkId& ChunkId)
+void FPackageStoreManifest::AddPackageData(FName PackageName, const FString& FileName, const FIoChunkId& ChunkId)
 {
 	FScopeLock Lock(&CriticalSection);
-	FPackageInfo* PackageInfo = GetPackageInfo_NoLock(InputPackageName, OutputPackageName);
+	FPackageInfo* PackageInfo = GetPackageInfo_NoLock(PackageName);
 	check(PackageInfo);
-	PackageInfo->PackageChunkId = ChunkId;
+	PackageInfo->ExportBundleChunkIds.Add(ChunkId);
 	if (!FileName.IsEmpty())
 	{
 		FileNameByChunkIdMap.Add(ChunkId, FileName);
 	}
 }
 
-void FPackageStoreManifest::AddBulkData(FName InputPackageName, FName OutputPackageName, const FString& FileName, const FIoChunkId& ChunkId)
+void FPackageStoreManifest::AddBulkData(FName PackageName, const FString& FileName, const FIoChunkId& ChunkId)
 {
 	FScopeLock Lock(&CriticalSection);
-	FPackageInfo* PackageInfo = GetPackageInfo_NoLock(InputPackageName, OutputPackageName);
+	FPackageInfo* PackageInfo = GetPackageInfo_NoLock(PackageName);
 	check(PackageInfo);
 	PackageInfo->BulkDataChunkIds.Add(ChunkId);
 	if (!FileName.IsEmpty())
@@ -105,7 +104,15 @@ FIoStatus FPackageStoreManifest::Save(const TCHAR* Filename) const
 			Writer->WriteObjectStart();
 		}
 		Writer->WriteValue(TEXT("Name"), PackageInfo.PackageName.ToString());
-		Writer->WriteValue(TEXT("PackageChunkId"), ChunkIdToString(PackageInfo.PackageChunkId));
+		if (!PackageInfo.ExportBundleChunkIds.IsEmpty())
+		{
+			Writer->WriteArrayStart(TEXT("ExportBundleChunkIds"));
+			for (const FIoChunkId& ChunkId : PackageInfo.ExportBundleChunkIds)
+			{
+				Writer->WriteValue(ChunkIdToString(ChunkId));
+			}
+			Writer->WriteArrayEnd();
+		}
 		if (!PackageInfo.BulkDataChunkIds.IsEmpty())
 		{
 			Writer->WriteArrayStart(TEXT("BulkDataChunkIds"));
@@ -119,27 +126,13 @@ FIoStatus FPackageStoreManifest::Save(const TCHAR* Filename) const
 	};
 
 	Writer->WriteArrayStart(TEXT("Packages"));
-	for (const auto& PackageNameInfoPair : MainPackageInfoByNameMap)
+	for (const auto& PackageNameInfoPair : PackageInfoByNameMap)
 	{
 		const FPackageInfo& PackageInfo = PackageNameInfoPair.Value;
 		WritePackageInfoObject(PackageInfo);
 	}
 	Writer->WriteArrayEnd();
 
-	Writer->WriteArrayStart(TEXT("OptionalPackages"));
-	for (const auto& PackageNameInfoPair : OptionalPackageInfoByNameMap)
-	{
-		Writer->WriteObjectStart();
-		Writer->WriteValue(TEXT("PackageName"), PackageNameInfoPair.Key.ToString());
-		{
-			const FPackageInfo& PackageInfo = PackageNameInfoPair.Value;
-			WritePackageInfoObject(PackageInfo, TEXT("PackageInfo"));
-		}
-		Writer->WriteObjectEnd();
-	}
-	Writer->WriteArrayEnd();
-
-
 	Writer->WriteObjectEnd();
 	Writer->Close();
 
@@ -154,8 +147,7 @@ FIoStatus FPackageStoreManifest::Save(const TCHAR* Filename) const
 FIoStatus FPackageStoreManifest::Load(const TCHAR* Filename)
 {
 	FScopeLock Lock(&CriticalSection);
-	MainPackageInfoByNameMap.Empty();
-	OptionalPackageInfoByNameMap.Empty();
+	PackageInfoByNameMap.Empty();
 	FileNameByChunkIdMap.Empty();
 
 	auto ChunkIdFromString = [](const FString& ChunkIdString)
@@ -212,7 +204,17 @@ FIoStatus FPackageStoreManifest::Load(const TCHAR* Filename)
 	auto ReadPackageInfo = [&ChunkIdFromString](TSharedPtr<FJsonObject> PackageObject, FPackageInfo& PackageInfo)
 	{
 		check(!PackageInfo.PackageName.IsNone());
-		PackageInfo.PackageChunkId = ChunkIdFromString(PackageObject->Values.FindRef(TEXT("PackageChunkId"))->AsString());
+		
+		TSharedPtr<FJsonValue> ExportBundleChunkIdsValue = PackageObject->Values.FindRef(TEXT("ExportBundleChunkIds"));
+		if (ExportBundleChunkIdsValue.IsValid())
+		{
+			TArray<TSharedPtr<FJsonValue>> ExportBundleChunkIdsArray = ExportBundleChunkIdsValue->AsArray();
+			PackageInfo.ExportBundleChunkIds.Reserve(ExportBundleChunkIdsArray.Num());
+			for (const TSharedPtr<FJsonValue>& ExportBundleChunkIdValue : ExportBundleChunkIdsArray)
+			{
+				PackageInfo.ExportBundleChunkIds.Add(ChunkIdFromString(ExportBundleChunkIdValue->AsString()));
+			}
+		}
 
 		TSharedPtr<FJsonValue> BulkDataChunkIdsValue = PackageObject->Values.FindRef(TEXT("BulkDataChunkIds"));
 		if (BulkDataChunkIdsValue.IsValid())
@@ -227,33 +229,17 @@ FIoStatus FPackageStoreManifest::Load(const TCHAR* Filename)
 	};
 
 	TArray<TSharedPtr<FJsonValue>> PackagesArray = JsonObject->Values.FindRef(TEXT("Packages"))->AsArray();
-	MainPackageInfoByNameMap.Reserve(PackagesArray.Num());
+	PackageInfoByNameMap.Reserve(PackagesArray.Num());
 	for (const TSharedPtr<FJsonValue>& PackageValue : PackagesArray)
 	{
 		TSharedPtr<FJsonObject> PackageObject = PackageValue->AsObject();
 		FName PackageName = FName(PackageObject->Values.FindRef(TEXT("Name"))->AsString());
 
-		FPackageInfo& PackageInfo = MainPackageInfoByNameMap.FindOrAdd(PackageName);
+		FPackageInfo& PackageInfo = PackageInfoByNameMap.FindOrAdd(PackageName);
 		PackageInfo.PackageName = PackageName;
 		ReadPackageInfo(PackageObject, PackageInfo);
 	}
 
-	TArray<TSharedPtr<FJsonValue>> OptionalPackagesArray = JsonObject->Values.FindRef(TEXT("OptionalPackages"))->AsArray();
-	OptionalPackageInfoByNameMap.Reserve(OptionalPackagesArray.Num());
-	for (const TSharedPtr<FJsonValue>& PackageNameInfoPair : OptionalPackagesArray)
-	{
-		TSharedPtr<FJsonObject> PairObject = PackageNameInfoPair->AsObject();
-		FName InputPackageName = FName(PairObject->Values.FindRef(TEXT("PackageName"))->AsString());
-
-		TSharedPtr<FJsonObject> PackageObject = PairObject->Values.FindRef(TEXT("PackageInfo"))->AsObject();
-		FName OutputPackageName = FName(PackageObject->Values.FindRef(TEXT("Name"))->AsString());
-
-		FPackageInfo PackageInfo;
-		PackageInfo.PackageName = OutputPackageName;
-		ReadPackageInfo(PackageObject, PackageInfo);
-		OptionalPackageInfoByNameMap.Add(InputPackageName, MoveTemp(PackageInfo));
-	}
-
 	return FIoStatus::Ok;
 }
 
@@ -273,10 +259,7 @@ TArray<FPackageStoreManifest::FPackageInfo> FPackageStoreManifest::GetPackages()
 {
 	FScopeLock Lock(&CriticalSection);
 	TArray<FPackageInfo> Packages;
-	MainPackageInfoByNameMap.GenerateValueArray(Packages);
-	TArray<FPackageInfo> OptionalPackages;
-	OptionalPackageInfoByNameMap.GenerateValueArray(OptionalPackages);
-	Packages.Append(MoveTemp(OptionalPackages));
+	PackageInfoByNameMap.GenerateValueArray(Packages);
 	return Packages;
 }
 
@@ -296,32 +279,8 @@ const FPackageStoreManifest::FZenServerInfo* FPackageStoreManifest::ReadZenServe
 	return ZenServerInfo.Get();
 }
 
-FPackageStoreManifest::FPackageInfo* FPackageStoreManifest::GetPackageInfo_NoLock(FName InputPackageName, FName OutputPackageName)
+FPackageStoreManifest::FPackageInfo* FPackageStoreManifest::GetPackageInfo_NoLock(FName PackageName)
 {
-	FPackageInfo* PackageInfo = nullptr;
-	if (InputPackageName == OutputPackageName)
-	{
-		PackageInfo = MainPackageInfoByNameMap.Find(InputPackageName);
-	}
-	else
-	{
-		TArray<FPackageInfo*, TInlineAllocator<2>> PackageInfos;
-		OptionalPackageInfoByNameMap.MultiFindPointer(InputPackageName, PackageInfos);
-		FPackageInfo** PackageInfoPtr = PackageInfos.FindByPredicate([OutputPackageName](FPackageInfo* Info)
-			{
-				return Info->PackageName == OutputPackageName;
-			});
-		// If we didn't find any, create one
-		if (PackageInfoPtr == nullptr)
-		{
-			PackageInfo = &OptionalPackageInfoByNameMap.Add(InputPackageName);
-			PackageInfo->PackageName = OutputPackageName;
-		}
-		else
-		{
-			PackageInfo = *PackageInfoPtr;
-		}
-	}
-	return PackageInfo;
+	return PackageInfoByNameMap.Find(PackageName);
 }
 
diff --git a/Engine/Source/Developer/IoStoreUtilities/Private/PackageStoreOptimizer.cpp b/Engine/Source/Developer/IoStoreUtilities/Private/PackageStoreOptimizer.cpp
--- a/Engine/Source/Developer/IoStoreUtilities/Private/PackageStoreOptimizer.cpp
+++ b/Engine/Source/Developer/IoStoreUtilities/Private/PackageStoreOptimizer.cpp
@@ -120,21 +120,8 @@ FPackageStorePackage* FPackageStoreOptimizer::CreatePackageFromCookedHeader(cons
 {
 	FPackageStorePackage* Package = new FPackageStorePackage();
 	Package->Id = FPackageId::FromName(Name);
-
-	// The package id should be generated from the original name
-	// However strip any '.' from the passed in name if any to generate the PackageName,
-	// <PackageName>.o is currently used to identify an optional package (optional exports chunk)
+	Package->Name = Name;
 	FString NameStr = Name.ToString();
-	FName PackageName = Name;
-	int32 Index = NameStr.Find(FPackagePath::GetOptionalSegmentExtensionModifier());
-	if (Index != INDEX_NONE)
-	{
-		NameStr.LeftInline(Index, false);
-		PackageName = *NameStr;
-		Package->bIsOptional = true;
-	}
-
-	Package->Name = PackageName;
 	Package->SourceName = *RemapLocalizationPathIfNeeded(NameStr, &Package->Region);
 
 	FCookedHeaderData CookedHeaderData = LoadCookedHeader(CookedHeaderBuffer);
@@ -403,7 +390,6 @@ void FPackageStoreOptimizer::ResolveImport(FPackageStorePackage::FUnresolvedImpo
 			PackageName.AppendString(Import->FullName);
 			Import->FullName.ToLowerInline();
 			Import->FromPackageId = FPackageId::FromName(PackageName);
-			Import->FromOptionalPackageId = FPackageId::FromName(PackageName, true);
 			Import->FromPackageName = PackageName;
 			Import->FromPackageNameLen = Import->FullName.Len();
 			Import->bIsScriptImport = Import->FullName.StartsWith(TEXT("/Script/"));
@@ -421,7 +407,6 @@ void FPackageStoreOptimizer::ResolveImport(FPackageStorePackage::FUnresolvedImpo
 			ObjectImport->ObjectName.AppendString(Import->FullName);
 			Import->FullName.ToLowerInline();
 			Import->FromPackageId = OuterImport->FromPackageId;
-			Import->FromOptionalPackageId = OuterImport->FromOptionalPackageId;
 			Import->bIsImportOptional = ObjectImport->bImportOptional;
 			Import->FromPackageName = OuterImport->FromPackageName;
 			Import->FromPackageNameLen = OuterImport->FromPackageNameLen;
@@ -453,11 +438,6 @@ void FPackageStoreOptimizer::ProcessImports(const FCookedHeaderData& CookedHeade
 			{
 				ImportedPackageIds.Add(UnresolvedImport.FromPackageId);
 			}
-			// Also add the associated optional package id to the list of imported package ids if the import is optional
-			else if (UnresolvedImport.bIsImportOptional)
-			{
-				ImportedPackageIds.Add(UnresolvedImport.FromOptionalPackageId);
-			}
 		}
 	}
 	Package->ImportedPackageIds = ImportedPackageIds.Array();
@@ -480,7 +460,7 @@ void FPackageStoreOptimizer::ProcessImports(const FCookedHeaderData& CookedHeade
 			bool bFoundPackageIndex = false;
 			for (uint32 PackageIndex = 0, PackageCount = static_cast<uint32>(Package->ImportedPackageIds.Num()); PackageIndex < PackageCount; ++PackageIndex)
 			{
-				FPackageId FromPackageId = !UnresolvedImport.bIsImportOptional ? UnresolvedImport.FromPackageId : UnresolvedImport.FromOptionalPackageId;
+				FPackageId FromPackageId = UnresolvedImport.FromPackageId;
 				if (FromPackageId == Package->ImportedPackageIds[PackageIndex])
 				{
 					FStringView PackageRelativeName = FStringView(UnresolvedImport.FullName).RightChop(UnresolvedImport.FromPackageNameLen);
@@ -1825,11 +1805,10 @@ void FPackageStoreOptimizer::LoadScriptObjectsBuffer(const FIoBuffer& ScriptObje
 	}
 }
 
-FPackageStoreEntryResource FPackageStoreOptimizer::CreatePackageStoreEntry(const FPackageStorePackage* Package) const
+FPackageStoreEntryResource FPackageStoreOptimizer::CreatePackageStoreEntry(const FPackageStorePackage* Package, const FPackageStorePackage* OptionalSegmentPackage) const
 {
 	FPackageStoreEntryResource Result;
 	Result.Flags = Package->bIsRedirected ? EPackageStoreEntryFlags::Redirected : EPackageStoreEntryFlags::None;
-	Result.Flags |= Package->bIsOptional ? EPackageStoreEntryFlags::Optional : EPackageStoreEntryFlags::None;
 	Result.PackageName = Package->Name;
 	Result.SourcePackageName = Package->SourceName;
 	Result.Region = FName(*Package->Region);
@@ -1837,27 +1816,76 @@ FPackageStoreEntryResource FPackageStoreOptimizer::CreatePackageStoreEntry(const
 	Result.ExportInfo.ExportBundleCount = Package->GraphData.ExportBundles.Num();
 	Result.ImportedPackageIds = Package->ImportedPackageIds;
 	Result.ShaderMapHashes = Package->ShaderMapHashes.Array();
+	if (OptionalSegmentPackage)
+	{
+		Result.OptionalSegmentExportInfo.ExportCount = OptionalSegmentPackage->Exports.Num();
+		Result.OptionalSegmentExportInfo.ExportBundleCount = OptionalSegmentPackage->GraphData.ExportBundles.Num();
+		Result.OptionalSegmentImportedPackageIds = OptionalSegmentPackage->ImportedPackageIds;
+	}
 	return Result;
 }
 
-FIoContainerHeader FPackageStoreOptimizer::CreateContainerHeader(const FIoContainerId& ContainerId, TArrayView<const FPackageStoreEntryResource> PackageStoreEntries) const
+FIoContainerHeader FPackageStoreOptimizer::CreateContainerHeader(const FIoContainerId& ContainerId, TArrayView<const FPackageStoreEntryResource> PackageStoreEntries, EContainerHeaderInclusionFilter InclusionFilter) const
 {
+	const bool bIncludeNonOptionalSegments = (InclusionFilter & IncludeNonOptionalSegments) > 0;
+	const bool bIncludeOptionalSegments = (InclusionFilter & IncludeOptionalSegments) > 0;
+
 	FIoContainerHeader Header;
 	Header.ContainerId = ContainerId;
-	Header.PackageCount = PackageStoreEntries.Num();
-	int32 StoreTocSize = Header.PackageCount * sizeof(FFilePackageStoreEntry);
-	FLargeMemoryWriter StoreTocArchive(0, true);
-	FLargeMemoryWriter StoreDataArchive(0, true);
+	
+	const int32 NonOptionalSegmentStoreEntriesCount = bIncludeNonOptionalSegments ? PackageStoreEntries.Num() : 0;
+
+	int32 OptionalSegmentStoreEntriesCount = 0;
+	if (bIncludeOptionalSegments)
+	{
+		for (const FPackageStoreEntryResource& Entry : PackageStoreEntries)
+		{
+			if (Entry.OptionalSegmentExportInfo.ExportCount)
+			{
+				++OptionalSegmentStoreEntriesCount;
+			}
+		}
+	}
+
+	struct FStoreEntriesWriter
+	{
+		const int32 StoreTocSize;
+		FLargeMemoryWriter StoreTocArchive = FLargeMemoryWriter(0, true);
+		FLargeMemoryWriter StoreDataArchive = FLargeMemoryWriter(0, true);
+
+		void Flush(TArray<uint8>& OutputBuffer)
+		{
+			check(StoreTocArchive.TotalSize() == StoreTocSize);
+			if (StoreTocSize)
+			{
+				const int32 StoreByteCount = StoreTocArchive.TotalSize() + StoreDataArchive.TotalSize();
+				OutputBuffer.AddUninitialized(StoreByteCount);
+				FBufferWriter PackageStoreArchive(OutputBuffer.GetData(), StoreByteCount);
+				PackageStoreArchive.Serialize(StoreTocArchive.GetData(), StoreTocArchive.TotalSize());
+				PackageStoreArchive.Serialize(StoreDataArchive.GetData(), StoreDataArchive.TotalSize());
+			}
+		}
+	};
+
+	FStoreEntriesWriter StoreEntriesWriter
+	{
+		static_cast<int32>(NonOptionalSegmentStoreEntriesCount * sizeof(FFilePackageStoreEntry))
+	};
+
+	FStoreEntriesWriter OptionalSegmentStoreEntriesWriter
+	{
+		static_cast<int32>(OptionalSegmentStoreEntriesCount * sizeof(FFilePackageStoreEntry))
+	};
 
-	auto SerializePackageEntryCArrayHeader = [&StoreTocSize, &StoreTocArchive, &StoreDataArchive](int32 Count)
+	auto SerializePackageEntryCArrayHeader = [](FStoreEntriesWriter& Writer, int32 Count)
 	{
-		const int32 RemainingTocSize = StoreTocSize - StoreTocArchive.Tell();
-		const int32 OffsetFromThis = RemainingTocSize + StoreDataArchive.Tell();
+		const int32 RemainingTocSize = Writer.StoreTocSize - Writer.StoreTocArchive.Tell();
+		const int32 OffsetFromThis = RemainingTocSize + Writer.StoreDataArchive.Tell();
 		uint32 ArrayNum = Count > 0 ? Count : 0;
 		uint32 OffsetToDataFromThis = ArrayNum > 0 ? OffsetFromThis : 0;
 
-		StoreTocArchive << ArrayNum;
-		StoreTocArchive << OffsetToDataFromThis;
+		Writer.StoreTocArchive << ArrayNum;
+		Writer.StoreTocArchive << OffsetToDataFromThis;
 	};
 
 	TArray<const FPackageStoreEntryResource*> SortedPackageStoreEntries;
@@ -1871,12 +1899,37 @@ FIoContainerHeader FPackageStoreOptimizer::CreateContainerHeader(const FIoContai
 		return A->GetPackageId() < B->GetPackageId();
 	});
 
-	Header.PackageIds.Reserve(SortedPackageStoreEntries.Num());
+	Header.PackageIds.Reserve(NonOptionalSegmentStoreEntriesCount);
+	Header.OptionalSegmentPackageIds.Reserve(OptionalSegmentStoreEntriesCount);
 	FPackageStoreNameMapBuilder RedirectsNameMapBuilder;
 	RedirectsNameMapBuilder.SetNameMapType(FMappedName::EType::Container);
 	TSet<FPackageId> AllLocalizedPackages;
 	for (const FPackageStoreEntryResource* Entry : SortedPackageStoreEntries)
 	{
+		if (bIncludeOptionalSegments && Entry->OptionalSegmentExportInfo.ExportCount)
+		{
+			Header.OptionalSegmentPackageIds.Add(Entry->GetPackageId());
+
+			FPackageStoreExportInfo OptionalSegmentExportInfo = Entry->OptionalSegmentExportInfo;
+			OptionalSegmentStoreEntriesWriter.StoreTocArchive << OptionalSegmentExportInfo;
+
+			const TArray<FPackageId>& OptionalSegmentImportedPackageIds = Entry->OptionalSegmentImportedPackageIds;
+			SerializePackageEntryCArrayHeader(OptionalSegmentStoreEntriesWriter, OptionalSegmentImportedPackageIds.Num());
+			for (FPackageId OptionalSegmentImportedPackageId : OptionalSegmentImportedPackageIds)
+			{
+				check(OptionalSegmentImportedPackageId.IsValid());
+				OptionalSegmentStoreEntriesWriter.StoreDataArchive << OptionalSegmentImportedPackageId;
+			}
+
+			// ShaderMapHashes is N/A for optional segments
+			SerializePackageEntryCArrayHeader(OptionalSegmentStoreEntriesWriter, 0);
+		}
+
+		if (!bIncludeNonOptionalSegments)
+		{
+			continue;
+		}
+
 		Header.PackageIds.Add(Entry->GetPackageId());
 		if (Entry->IsRedirected())
 		{
@@ -1898,34 +1951,29 @@ FIoContainerHeader FPackageStoreOptimizer::CreateContainerHeader(const FIoContai
 		
 		// StoreEntries
 		FPackageStoreExportInfo ExportInfo = Entry->ExportInfo;
-		StoreTocArchive << ExportInfo; 
+		StoreEntriesWriter.StoreTocArchive << ExportInfo;
 		
 		// ImportedPackages
 		const TArray<FPackageId>& ImportedPackageIds = Entry->ImportedPackageIds;
-		SerializePackageEntryCArrayHeader(ImportedPackageIds.Num());
+		SerializePackageEntryCArrayHeader(StoreEntriesWriter, ImportedPackageIds.Num());
 		for (FPackageId ImportedPackageId : ImportedPackageIds)
 		{
 			check(ImportedPackageId.IsValid());
-			StoreDataArchive << ImportedPackageId;
+			StoreEntriesWriter.StoreDataArchive << ImportedPackageId;
 		}
 		
 		// ShaderMapHashes
 		const TArray<FSHAHash>& ShaderMapHashes = Entry->ShaderMapHashes;
-		SerializePackageEntryCArrayHeader(ShaderMapHashes.Num());
+		SerializePackageEntryCArrayHeader(StoreEntriesWriter, ShaderMapHashes.Num());
 		for (const FSHAHash& ShaderMapHash : ShaderMapHashes)
 		{
-			StoreDataArchive << const_cast<FSHAHash&>(ShaderMapHash);
+			StoreEntriesWriter.StoreDataArchive << const_cast<FSHAHash&>(ShaderMapHash);
 		}
 	}
 	Header.RedirectsNameMap = RedirectsNameMapBuilder.GetNameMap();
 
-	check(StoreTocArchive.TotalSize() == StoreTocSize);
-
-	const int32 StoreByteCount = StoreTocArchive.TotalSize() + StoreDataArchive.TotalSize();
-	Header.StoreEntries.AddUninitialized(StoreByteCount);
-	FBufferWriter PackageStoreArchive(Header.StoreEntries.GetData(), StoreByteCount);
-	PackageStoreArchive.Serialize(StoreTocArchive.GetData(), StoreTocArchive.TotalSize());
-	PackageStoreArchive.Serialize(StoreDataArchive.GetData(), StoreDataArchive.TotalSize());
+	StoreEntriesWriter.Flush(Header.StoreEntries);
+	OptionalSegmentStoreEntriesWriter.Flush(Header.OptionalSegmentStoreEntries);
 
 	return Header;
 }
diff --git a/Engine/Source/Developer/IoStoreUtilities/Private/ZenStoreWriter.cpp b/Engine/Source/Developer/IoStoreUtilities/Private/ZenStoreWriter.cpp
--- a/Engine/Source/Developer/IoStoreUtilities/Private/ZenStoreWriter.cpp
+++ b/Engine/Source/Developer/IoStoreUtilities/Private/ZenStoreWriter.cpp
@@ -230,7 +230,7 @@ void FZenStoreWriter::WritePackageData(const FPackageInfo& Info, FLargeMemoryWri
 
 	FIoBuffer CookedHeaderBuffer = FIoBuffer(PackageData.Data(), Info.HeaderSize, PackageData);
 	FIoBuffer CookedExportsBuffer = FIoBuffer(PackageData.Data() + Info.HeaderSize, PackageData.DataSize() - Info.HeaderSize, PackageData);
-	TUniquePtr<FPackageStorePackage> Package{PackageStoreOptimizer->CreatePackageFromCookedHeader(Info.OutputPackageName, CookedHeaderBuffer)};
+	TUniquePtr<FPackageStorePackage> Package{PackageStoreOptimizer->CreatePackageFromCookedHeader(Info.PackageName, CookedHeaderBuffer)};
 	PackageStoreOptimizer->FinalizePackage(Package.Get());
 	TArray<FFileRegion> FileRegionsCopy(FileRegions);
 	for (FFileRegion& Region : FileRegionsCopy)
@@ -239,8 +239,7 @@ void FZenStoreWriter::WritePackageData(const FPackageInfo& Info, FLargeMemoryWri
 		Region.Offset -= Info.HeaderSize;
 	}
 	FIoBuffer PackageBuffer = PackageStoreOptimizer->CreatePackageBuffer(Package.Get(), CookedExportsBuffer, &FileRegionsCopy);
-	FIoChunkId ChunkId = CreateIoChunkId(Package->GetId().Value(), 0, EIoChunkType::ExportBundleData);
-	PackageStoreManifest.AddPackageData(Info.InputPackageName, Info.OutputPackageName, Info.LooseFilePath, Info.ChunkId);
+	PackageStoreManifest.AddPackageData(Info.PackageName, Info.LooseFilePath, Info.ChunkId);
 	for (FFileRegion& Region : FileRegionsCopy)
 	{
 		// Adjust regions once more so they are relative to the exports bundle buffer
@@ -252,7 +251,7 @@ void FZenStoreWriter::WritePackageData(const FPackageInfo& Info, FLargeMemoryWri
 
 	FCbObjectId ChunkOid = ToObjectId(Info.ChunkId);
 
-	FPendingPackageState& ExistingState = GetPendingPackage(Info.InputPackageName);
+	FPendingPackageState& ExistingState = GetPendingPackage(Info.PackageName);
 
 	FPackageDataEntry& Entry = ExistingState.PackageData;
 
@@ -263,7 +262,7 @@ void FZenStoreWriter::WritePackageData(const FPackageInfo& Info, FLargeMemoryWri
 
 	Entry.Info				= Info;
 	Entry.ChunkId			= ChunkOid;
-	Entry.PackageStoreEntry = PackageStoreOptimizer->CreatePackageStoreEntry(Package.Get());
+	Entry.PackageStoreEntry = PackageStoreOptimizer->CreatePackageStoreEntry(Package.Get(), nullptr); // TODO: Can we separate out the optional segment package store entry and do this when we commit instead?
 	Entry.IsValid			= true;
 
 	if (EntryCreatedEvent.IsBound())
@@ -283,12 +282,12 @@ void FZenStoreWriter::WriteIoStorePackageData(const FPackageInfo& Info, const FI
 
 	TRACE_CPUPROFILER_EVENT_SCOPE(WriteIoStorePackageData);
 
-	PackageStoreManifest.AddPackageData(Info.InputPackageName, Info.OutputPackageName, Info.LooseFilePath, Info.ChunkId);
+	PackageStoreManifest.AddPackageData(Info.PackageName, Info.LooseFilePath, Info.ChunkId);
 	//WriteFileRegions(*FPaths::ChangeExtension(Info.LooseFilePath, FString(".uexp") + FFileRegion::RegionsFileExtension), FileRegionsCopy);
 
 	FCbObjectId ChunkOid = ToObjectId(Info.ChunkId);
 
-	FPendingPackageState& ExistingState = GetPendingPackage(Info.InputPackageName);
+	FPendingPackageState& ExistingState = GetPendingPackage(Info.PackageName);
 
 	FPackageDataEntry& Entry = ExistingState.PackageData;
 
@@ -311,7 +310,7 @@ void FZenStoreWriter::WriteBulkData(const FBulkDataInfo& Info, const FIoBuffer&
 
 	FCbObjectId ChunkOid = ToObjectId(Info.ChunkId);
 
-	FPendingPackageState& ExistingState = GetPendingPackage(Info.InputPackageName);
+	FPendingPackageState& ExistingState = GetPendingPackage(Info.PackageName);
 
 	FBulkDataEntry& BulkEntry = ExistingState.BulkData.AddDefaulted_GetRef(); 
 
@@ -326,7 +325,7 @@ void FZenStoreWriter::WriteBulkData(const FBulkDataInfo& Info, const FIoBuffer&
 	BulkEntry.ChunkId	= ChunkOid;
 	BulkEntry.IsValid	= true;
 
-	PackageStoreManifest.AddBulkData(Info.InputPackageName, Info.OutputPackageName, Info.LooseFilePath, Info.ChunkId);
+	PackageStoreManifest.AddBulkData(Info.PackageName, Info.LooseFilePath, Info.ChunkId);
 
 	//	WriteFileRegions(*(Info.LooseFilePath + FFileRegion::RegionsFileExtension), FileRegions);
 }
@@ -335,7 +334,7 @@ void FZenStoreWriter::WriteAdditionalFile(const FAdditionalFileInfo& Info, const
 {
 	const FZenFileSystemManifest::FManifestEntry& ManifestEntry = ZenFileSystemManifest->CreateManifestEntry(Info.Filename);
 	
-	FPendingPackageState& ExistingState = GetPendingPackage(Info.InputPackageName);
+	FPendingPackageState& ExistingState = GetPendingPackage(Info.PackageName);
 
 	FFileDataEntry& FileEntry = ExistingState.FileData.AddDefaulted_GetRef();
 	
@@ -775,7 +774,6 @@ void FZenStoreWriter::CommitPackageInternal(FCommitPackageInfo&& CommitInfo)
 
 				CommitEventArgs.AdditionalFiles.Add(FAdditionalFileInfo
 				{ 
-					CommitInfo.PackageName,
 					CommitInfo.PackageName,
 					File.ZenManifestClientPath,
 					File.Info.ChunkId
@@ -1084,7 +1082,7 @@ void FZenStoreWriter::CreateProjectMetaData(FCbPackage& Pkg, FCbWriter& PackageO
 			FIoBuffer HeaderBuffer;
 
 			{
-				FIoContainerHeader Header = PackageStoreOptimizer->CreateContainerHeader(ContainerId, PackageStoreEntries);
+				FIoContainerHeader Header = PackageStoreOptimizer->CreateContainerHeader(ContainerId, PackageStoreEntries, FPackageStoreOptimizer::IncludeAllSegments);
 				FLargeMemoryWriter HeaderAr(0, true);
 				HeaderAr << Header;
 				int64 DataSize = HeaderAr.TotalSize();
diff --git a/Engine/Source/Developer/IoStoreUtilities/Public/PackageStoreManifest.h b/Engine/Source/Developer/IoStoreUtilities/Public/PackageStoreManifest.h
--- a/Engine/Source/Developer/IoStoreUtilities/Public/PackageStoreManifest.h
+++ b/Engine/Source/Developer/IoStoreUtilities/Public/PackageStoreManifest.h
@@ -17,7 +17,7 @@ public:
 	struct FPackageInfo
 	{
 		FName PackageName;
-		FIoChunkId PackageChunkId;
+		TArray<FIoChunkId, TInlineAllocator<1>> ExportBundleChunkIds;
 		TArray<FIoChunkId> BulkDataChunkIds;
 	};
 	
@@ -31,9 +31,9 @@ public:
 	IOSTOREUTILITIES_API FPackageStoreManifest(const FString& CookedOutputPath);
 	IOSTOREUTILITIES_API ~FPackageStoreManifest() = default;
 
-	IOSTOREUTILITIES_API void BeginPackage(FName InputPackageName);
-	IOSTOREUTILITIES_API void AddPackageData(FName InputPackageName, FName OutputPackageName, const FString& FileName, const FIoChunkId& ChunkId);
-	IOSTOREUTILITIES_API void AddBulkData(FName InputPackageName, FName OutputPackageName, const FString& FileName, const FIoChunkId& ChunkId);
+	IOSTOREUTILITIES_API void BeginPackage(FName PackageName);
+	IOSTOREUTILITIES_API void AddPackageData(FName PackageName, const FString& FileName, const FIoChunkId& ChunkId);
+	IOSTOREUTILITIES_API void AddBulkData(FName PackageName, const FString& FileName, const FIoChunkId& ChunkId);
 
 	IOSTOREUTILITIES_API FIoStatus Save(const TCHAR* Filename) const;
 	IOSTOREUTILITIES_API FIoStatus Load(const TCHAR* Filename);
@@ -45,14 +45,11 @@ public:
 	IOSTOREUTILITIES_API const FZenServerInfo* ReadZenServerInfo() const;
 
 private:
-	FPackageInfo* GetPackageInfo_NoLock(FName InputPackageName, FName OutputPackageName);
+	FPackageInfo* GetPackageInfo_NoLock(FName PackageName);
 
 	mutable FCriticalSection CriticalSection;
 	FString CookedOutputPath;
-	/** Main Package Output Info map */
-	TMap<FName, FPackageInfo> MainPackageInfoByNameMap;
-	/** Additionally Generated Package Info Output map */
-	TMultiMap<FName, FPackageInfo> OptionalPackageInfoByNameMap;
+	TMap<FName, FPackageInfo> PackageInfoByNameMap;
 
 	TMap<FIoChunkId, FString> FileNameByChunkIdMap;
 	TUniquePtr<FZenServerInfo> ZenServerInfo;
diff --git a/Engine/Source/Editor/UnrealEd/Private/Cooker/LooseCookedPackageWriter.cpp b/Engine/Source/Editor/UnrealEd/Private/Cooker/LooseCookedPackageWriter.cpp
--- a/Engine/Source/Editor/UnrealEd/Private/Cooker/LooseCookedPackageWriter.cpp
+++ b/Engine/Source/Editor/UnrealEd/Private/Cooker/LooseCookedPackageWriter.cpp
@@ -84,7 +84,7 @@ void FLooseCookedPackageWriter::AsyncSave(FRecord& Record, const FCommitPackageI
 
 void FLooseCookedPackageWriter::CompleteExportsArchiveForDiff(const FPackageInfo& Info, FLargeMemoryWriter& ExportsArchive)
 {
-	FPackageWriterRecords::FPackage& BaseRecord = Records.FindRecordChecked(Info.InputPackageName);
+	FPackageWriterRecords::FPackage& BaseRecord = Records.FindRecordChecked(Info.PackageName);
 	FRecord& Record = static_cast<FRecord&>(BaseRecord);
 	Record.bCompletedExportsArchiveForDiff = true;
 
@@ -417,11 +417,11 @@ void FLooseCookedPackageWriter::UpdateManifest(FRecord& Record)
 {
 	for (const FPackageWriterRecords::FWritePackage& Package : Record.Packages)
 	{
-		PackageStoreManifest.AddPackageData(Package.Info.InputPackageName, Package.Info.OutputPackageName, Package.Info.LooseFilePath, Package.Info.ChunkId);
+		PackageStoreManifest.AddPackageData(Package.Info.PackageName, Package.Info.LooseFilePath, Package.Info.ChunkId);
 	}
 	for (const FPackageWriterRecords::FBulkData& BulkData : Record.BulkDatas)
 	{
-		PackageStoreManifest.AddBulkData(BulkData.Info.InputPackageName, BulkData.Info.OutputPackageName, BulkData.Info.LooseFilePath, BulkData.Info.ChunkId);
+		PackageStoreManifest.AddBulkData(BulkData.Info.PackageName, BulkData.Info.LooseFilePath, BulkData.Info.ChunkId);
 	}
 }
 
diff --git a/Engine/Source/Runtime/Core/Private/IO/IoContainerHeader.cpp b/Engine/Source/Runtime/Core/Private/IO/IoContainerHeader.cpp
--- a/Engine/Source/Runtime/Core/Private/IO/IoContainerHeader.cpp
+++ b/Engine/Source/Runtime/Core/Private/IO/IoContainerHeader.cpp
@@ -37,9 +37,10 @@ FArchive& operator<<(FArchive& Ar, FIoContainerHeader& ContainerHeader)
 		return Ar;
 	}
 	Ar << ContainerHeader.ContainerId;
-	Ar << ContainerHeader.PackageCount;
 	Ar << ContainerHeader.PackageIds;
 	Ar << ContainerHeader.StoreEntries;
+	Ar << ContainerHeader.OptionalSegmentPackageIds;
+	Ar << ContainerHeader.OptionalSegmentStoreEntries;
 	if (Ar.IsLoading())
 	{
 		ContainerHeader.RedirectsNameMap = LoadNameBatch(Ar);
diff --git a/Engine/Source/Runtime/Core/Private/IO/PackageId.cpp b/Engine/Source/Runtime/Core/Private/IO/PackageId.cpp
--- a/Engine/Source/Runtime/Core/Private/IO/PackageId.cpp
+++ b/Engine/Source/Runtime/Core/Private/IO/PackageId.cpp
@@ -3,18 +3,11 @@
 #include "IO/PackageId.h"
 #include "Serialization/StructuredArchive.h"
 
-FPackageId FPackageId::FromName(const FName& Name, bool bAsOptional)
+FPackageId FPackageId::FromName(const FName& Name)
 {
 	TCHAR NameStr[FName::StringBufferSize + 2];
 	uint32 NameLen = Name.ToString(NameStr);
 
-	if (bAsOptional)
-	{
-		NameStr[NameLen++] = '.';
-		NameStr[NameLen++] = 'o';
-		NameStr[NameLen] = '\0';
-	}
-
 	for (uint32 I = 0; I < NameLen; ++I)
 	{
 		NameStr[I] = TChar<TCHAR>::ToLower(NameStr[I]);
diff --git a/Engine/Source/Runtime/Core/Public/IO/IoContainerHeader.h b/Engine/Source/Runtime/Core/Public/IO/IoContainerHeader.h
--- a/Engine/Source/Runtime/Core/Public/IO/IoContainerHeader.h
+++ b/Engine/Source/Runtime/Core/Public/IO/IoContainerHeader.h
@@ -65,6 +65,7 @@ enum class EIoContainerHeaderVersion : uint32
 {
 	Initial = 0,
 	LocalizedPackages = 1,
+	OptionalSegmentPackages = 2,
 
 	LatestPlusOne,
 	Latest = LatestPlusOne - 1
@@ -78,9 +79,10 @@ struct FIoContainerHeader
 	};
 
 	FIoContainerId ContainerId;
-	uint32 PackageCount = 0;
 	TArray<FPackageId> PackageIds;
-	TArray<uint8> StoreEntries; //FPackageStoreEntry[PackageCount]
+	TArray<uint8> StoreEntries; //FPackageStoreEntry[PackageIds.Num()]
+	TArray<FPackageId> OptionalSegmentPackageIds;
+	TArray<uint8> OptionalSegmentStoreEntries; //FPackageStoreEntry[OptionalSegmentPackageIds.Num()]
 	TArray<FDisplayNameEntryId> RedirectsNameMap;
 	TArray<FIoContainerHeaderLocalizedPackage> LocalizedPackages;
 	TArray<FIoContainerHeaderPackageRedirect> PackageRedirects;
diff --git a/Engine/Source/Runtime/Core/Public/IO/PackageId.h b/Engine/Source/Runtime/Core/Public/IO/PackageId.h
--- a/Engine/Source/Runtime/Core/Public/IO/PackageId.h
+++ b/Engine/Source/Runtime/Core/Public/IO/PackageId.h
@@ -17,7 +17,7 @@ class FPackageId
 public:
 	FPackageId() = default;
 
-	CORE_API static FPackageId FromName(const FName& Name, bool bAsOptional = false);
+	CORE_API static FPackageId FromName(const FName& Name);
 
 	static FPackageId FromValue(const uint64 Value)
 	{
diff --git a/Engine/Source/Runtime/Core/Public/IO/PackageStore.h b/Engine/Source/Runtime/Core/Public/IO/PackageStore.h
--- a/Engine/Source/Runtime/Core/Public/IO/PackageStore.h
+++ b/Engine/Source/Runtime/Core/Public/IO/PackageStore.h
@@ -48,6 +48,8 @@ struct FPackageStoreEntry
 #if WITH_EDITOR
 	FName UncookedPackageName;
 	uint8 UncookedPackageHeaderExtension; // TODO: Can't include PackagePath.h
+	FPackageStoreExportInfo OptionalSegmentExportInfo;
+	TArrayView<const FPackageId> OptionalSegmentImportedPackageIds;
 #endif
 };
 
@@ -58,7 +60,6 @@ enum class EPackageStoreEntryFlags : uint32
 {
 	None		= 0,
 	Redirected	= 0x01,
-	Optional	= 0x02,
 };
 ENUM_CLASS_FLAGS(EPackageStoreEntryFlags);
 
@@ -85,17 +86,21 @@ struct FPackageStoreEntryResource
 	TArray<FPackageId> ImportedPackageIds;
 	/** Referenced shader map hashes. */
 	TArray<FSHAHash> ShaderMapHashes;
+	/** The editor data package export information. */
+	FPackageStoreExportInfo OptionalSegmentExportInfo;
+	/** Editor data imported package IDs. */
+	TArray<FPackageId> OptionalSegmentImportedPackageIds;
 
 	/** Returns the package ID. */
 	FPackageId GetPackageId() const
 	{
-		return FPackageId::FromName(PackageName, IsOptional());
+		return FPackageId::FromName(PackageName);
 	}
 
 	/** Returns the source package ID. */
 	FPackageId GetSourcePackageId() const
 	{
-		return SourcePackageName.IsNone() ? FPackageId() : FPackageId::FromName(SourcePackageName, IsOptional());
+		return SourcePackageName.IsNone() ? FPackageId() : FPackageId::FromName(SourcePackageName);
 	}
 
 	FName GetSourcePackageName() const
@@ -109,12 +114,6 @@ struct FPackageStoreEntryResource
 		return EnumHasAnyFlags(Flags, EPackageStoreEntryFlags::Redirected); 
 	}
 
-	/** Returns whether this package is optional. */
-	bool IsOptional() const
-	{
-		return EnumHasAnyFlags(Flags, EPackageStoreEntryFlags::Optional);
-	}
-
 	CORE_API friend FArchive& operator<<(FArchive& Ar, FPackageStoreEntryResource& PackageStoreEntry);
 	
 	CORE_API friend FCbWriter& operator<<(FCbWriter& Writer, const FPackageStoreEntryResource& PackageStoreEntry);
diff --git a/Engine/Source/Runtime/Core/Public/Serialization/PackageWriter.h b/Engine/Source/Runtime/Core/Public/Serialization/PackageWriter.h
--- a/Engine/Source/Runtime/Core/Public/Serialization/PackageWriter.h
+++ b/Engine/Source/Runtime/Core/Public/Serialization/PackageWriter.h
@@ -89,9 +89,7 @@ public:
 	struct FPackageInfo
 	{
 		/** Associated Package Name Entry from BeginPackage */
-		FName		InputPackageName;
-		/** Output Package Name (Input Package can produce multiple output) */
-		FName		OutputPackageName;
+		FName		PackageName;
 		FString		LooseFilePath;
 		uint64		HeaderSize = 0;
 		FIoChunkId	ChunkId = FIoChunkId::InvalidChunkId;
@@ -117,9 +115,7 @@ public:
 		};
 
 		/** Associated Package Name Entry */
-		FName		InputPackageName;
-		/** Output Package Name (Input Package can produce multiple output) */
-		FName		OutputPackageName;
+		FName		PackageName;
 		EType		BulkDataType = BulkSegment;
 		FString		LooseFilePath;
 		FIoChunkId	ChunkId = FIoChunkId::InvalidChunkId;
@@ -133,9 +129,7 @@ public:
 	struct FAdditionalFileInfo
 	{
 		/** Associated Package Name Entry */
-		FName		InputPackageName;
-		/** Output Package Name (Input Package can produce multiple output) */
-		FName		OutputPackageName;
+		FName		PackageName;
 		FString		Filename;
 		FIoChunkId	ChunkId = FIoChunkId::InvalidChunkId;
 		uint32		MultiOutputIndex = 0;
@@ -147,9 +141,7 @@ public:
 	struct FLinkerAdditionalDataInfo
 	{
 		/** Associated Package Name Entry */
-		FName	InputPackageName;
-		/** Output Package Name (Input Package can produce multiple output) */
-		FName	OutputPackageName;
+		FName	PackageName;
 		uint32	MultiOutputIndex = 0;
 	};
 	/** Write separate data written by UObjects via FLinkerSave::AdditionalDataToAppend. */
diff --git a/Engine/Source/Runtime/CoreUObject/Private/Serialization/AsyncLoading2.cpp b/Engine/Source/Runtime/CoreUObject/Private/Serialization/AsyncLoading2.cpp
--- a/Engine/Source/Runtime/CoreUObject/Private/Serialization/AsyncLoading2.cpp
+++ b/Engine/Source/Runtime/CoreUObject/Private/Serialization/AsyncLoading2.cpp
@@ -863,6 +863,7 @@ public:
 struct FAsyncPackageHeaderData
 {
 	uint32 CookedHeaderSize = 0;
+	uint32 ExportCount = 0; // Need to keep this count around until after ExportMap is being cleared
 	TOptional<FZenPackageVersioningInfo> VersioningInfo;
 	FNameMap NameMap;
 	FName PackageName;
@@ -1707,6 +1708,10 @@ private:
 	FAsyncPackageData Data;
 	FAsyncPackageHeaderData HeaderData;
 	FAsyncPackageSerializationState SerializationState;
+#if WITH_EDITOR
+	TOptional<FAsyncPackageHeaderData> OptionalSegmentHeaderData;
+	TOptional<FAsyncPackageSerializationState> OptionalSegmentSerializationState;
+#endif
 	/** Cached async loading thread object this package was created by */
 	FAsyncLoadingThread2& AsyncLoadingThread;
 	FAsyncLoadEventGraphAllocator& GraphAllocator;
@@ -2416,16 +2421,37 @@ private:
 		const uint64 ExportBundleEntriesMemSize = Align(sizeof(FExportBundleEntry) * ExportBundleEntriesCount, 8);
 		const int32 ImportedPackagesCount = PackageStoreEntry.ImportedPackageIds.Num();
 		const uint64 ImportedPackageIdsMemSize = Align(sizeof(FPackageId) * ImportedPackagesCount, 8);
+#if WITH_EDITOR
+		const int32 OptionalSegmentExportCount = PackageStoreEntry.OptionalSegmentExportInfo.ExportCount;
+		const int32 OptionalSegmentExportBundleCount = PackageStoreEntry.OptionalSegmentExportInfo.ExportBundleCount;
+		const uint64 OptionalSegmentExportBundleHeadersMemSize = Align(sizeof(FExportBundleHeader) * OptionalSegmentExportBundleCount, 8);
+		const int32 OptionalSegmentExportBundleEntriesCount = OptionalSegmentExportCount * FExportBundleEntry::ExportCommandType_Count;
+		const uint64 OptionalSegmentExportBundleEntriesMemSize = Align(sizeof(FExportBundleEntry) * OptionalSegmentExportCount * FExportBundleEntry::ExportCommandType_Count, 8);
+		const int32 OptionalSegmentImportedPackagesCount = PackageStoreEntry.OptionalSegmentImportedPackageIds.Num();
+		const uint64 OptionalSegmentImportedPackageIdsMemSize = Align(sizeof(FPackageId) * OptionalSegmentImportedPackagesCount, 8);
+
+		const int32 TotalImportedPackagesCount = ImportedPackagesCount + OptionalSegmentImportedPackagesCount;
+		const int32 TotalExportCount = ExportCount + OptionalSegmentExportCount;
+		AsyncPackage->Data.ExportBundleCount = ExportBundleCount + OptionalSegmentExportBundleCount;
+#else
+		const int32 TotalImportedPackagesCount = ImportedPackagesCount;
+		const int32 TotalExportCount = ExportCount;
 		AsyncPackage->Data.ExportBundleCount = ExportBundleCount;
+#endif
 		const int32 ExportBundleNodeCount = AsyncPackage->Data.ExportBundleCount * EEventLoadNode2::ExportBundle_NumPhases;
 		const int32 ShaderMapHashesCount = PackageStoreEntry.ShaderMapHashes.Num();
 
-		const uint64 ImportedPackagesMemSize = Align(sizeof(FAsyncPackage2*) * ImportedPackagesCount, 8);
-		const uint64 ExportsMemSize = Align(sizeof(FExportObject) * ExportCount, 8);
+		const uint64 ImportedPackagesMemSize = Align(sizeof(FAsyncPackage2*) * TotalImportedPackagesCount, 8);
+		const uint64 ExportsMemSize = Align(sizeof(FExportObject) * TotalExportCount, 8);
 		const uint64 ExportBundleNodesMemSize = Align(sizeof(FEventLoadNode2) * ExportBundleNodeCount, 8);
 		const uint64 ShaderMapHashesMemSize = Align(sizeof(FSHAHash) * ShaderMapHashesCount, 8);
-		const uint64 ExportToBundleMappingMemSize = Align(sizeof(FAsyncPackageExportToBundleMapping) * ExportCount, 8);
+		const uint64 ExportToBundleMappingMemSize = Align(sizeof(FAsyncPackageExportToBundleMapping) * TotalExportCount, 8);
 		const uint64 MemoryBufferSize =
+#if WITH_EDITOR
+			OptionalSegmentExportBundleHeadersMemSize +
+			OptionalSegmentExportBundleEntriesMemSize +
+			OptionalSegmentImportedPackageIdsMemSize +
+#endif
 			ExportBundleHeadersMemSize +
 			ExportBundleEntriesMemSize +
 			ImportedPackageIdsMemSize +
@@ -2439,7 +2465,7 @@ private:
 
 		uint8* DataPtr = Data.MemoryBuffer;
 
-		Data.Exports = MakeArrayView(reinterpret_cast<FExportObject*>(DataPtr), ExportCount);
+		Data.Exports = MakeArrayView(reinterpret_cast<FExportObject*>(DataPtr), TotalExportCount);
 		DataPtr += ExportsMemSize;
 		Data.ExportBundleNodes = MakeArrayView(reinterpret_cast<FEventLoadNode2*>(DataPtr), ExportBundleNodeCount);
 		DataPtr += ExportBundleNodesMemSize;
@@ -2448,9 +2474,10 @@ private:
 		DataPtr += ShaderMapHashesMemSize;
 		Data.ImportedAsyncPackages = MakeArrayView(reinterpret_cast<FAsyncPackage2**>(DataPtr), 0);
 		DataPtr += ImportedPackagesMemSize;
-		Data.ExportToBundleMappings = MakeArrayView(reinterpret_cast<FAsyncPackageExportToBundleMapping*>(DataPtr), ExportCount);
+		Data.ExportToBundleMappings = MakeArrayView(reinterpret_cast<FAsyncPackageExportToBundleMapping*>(DataPtr), TotalExportCount);
 		DataPtr += ExportToBundleMappingMemSize;
 
+		HeaderData.ExportCount = ExportCount;
 		HeaderData.ExportBundleHeaders = MakeArrayView(reinterpret_cast<FExportBundleHeader*>(DataPtr), ExportBundleCount);
 		DataPtr += ExportBundleHeadersMemSize;
 		HeaderData.ExportBundleEntries = MakeArrayView(reinterpret_cast<FExportBundleEntry*>(DataPtr), ExportBundleEntriesCount);
@@ -2458,6 +2485,22 @@ private:
 		HeaderData.ImportedPackageIds = MakeArrayView(reinterpret_cast<FPackageId*>(DataPtr), ImportedPackagesCount);
 		FMemory::Memcpy((void*)HeaderData.ImportedPackageIds.GetData(), PackageStoreEntry.ImportedPackageIds.GetData(), sizeof(FPackageId) * ImportedPackagesCount);
 		DataPtr += ImportedPackageIdsMemSize;
+#if WITH_EDITOR
+		if (OptionalSegmentExportCount)
+		{
+			AsyncPackage->OptionalSegmentSerializationState.Emplace();
+			FAsyncPackageHeaderData& OptionalSegmentHeaderData = AsyncPackage->OptionalSegmentHeaderData.Emplace();
+			OptionalSegmentHeaderData.ExportCount = OptionalSegmentExportCount;
+			OptionalSegmentHeaderData.ExportBundleHeaders = MakeArrayView(reinterpret_cast<FExportBundleHeader*>(DataPtr), OptionalSegmentExportBundleCount);
+			DataPtr += OptionalSegmentExportBundleHeadersMemSize;
+			OptionalSegmentHeaderData.ExportBundleEntries = MakeArrayView(reinterpret_cast<FExportBundleEntry*>(DataPtr), OptionalSegmentExportBundleEntriesCount);
+			DataPtr += OptionalSegmentExportBundleEntriesMemSize;
+			OptionalSegmentHeaderData.ImportedPackageIds = MakeArrayView(reinterpret_cast<FPackageId*>(DataPtr), OptionalSegmentImportedPackagesCount);
+			FMemory::Memcpy((void*)OptionalSegmentHeaderData.ImportedPackageIds.GetData(), PackageStoreEntry.OptionalSegmentImportedPackageIds.GetData(), sizeof(FPackageId) * OptionalSegmentImportedPackagesCount);
+			DataPtr += OptionalSegmentImportedPackageIdsMemSize;
+			AsyncPackage->ImportStore.AddImportedPackageReferences(OptionalSegmentHeaderData.ImportedPackageIds);
+		}
+#endif
 		AsyncPackage->ImportStore.AddImportedPackageReferences(HeaderData.ImportedPackageIds);
 		AsyncPackage->ImportStore.AddPackageReference(AsyncPackage->Desc);
 		check(DataPtr - Data.MemoryBuffer == MemoryBufferSize);
@@ -2590,6 +2633,12 @@ void FAsyncLoadingThread2::UpdatePackagePriority(FAsyncPackage2* Package, int32
 	TRACE_CPUPROFILER_EVENT_SCOPE(UpdatePackagePriority);
 	Package->Desc.Priority = NewPriority;
 	Package->SerializationState.IoRequest.UpdatePriority(NewPriority);
+#if WITH_EDITOR
+	if (Package->OptionalSegmentSerializationState.IsSet())
+	{
+		Package->OptionalSegmentSerializationState->IoRequest.UpdatePriority(NewPriority);
+	}
+#endif
 }
 
 FAsyncPackage2* FAsyncLoadingThread2::FindOrInsertPackage(FAsyncPackageDesc2& Desc, bool& bInserted, TUniquePtr<FLoadPackageAsyncDelegate>&& PackageLoadedDelegate)
@@ -3188,6 +3237,12 @@ void FAsyncPackage2::SetupScriptDependencies()
 	// Then collect them here, and wait for them to be created before allowing this package to proceed.
 	TArray<UClass*, TInlineAllocator<8>> UnresolvedCDOs;
 	ImportStore.GetUnresolvedCDOs(HeaderData, UnresolvedCDOs);
+#if WITH_EDITOR
+	if (OptionalSegmentHeaderData.IsSet())
+	{
+		ImportStore.GetUnresolvedCDOs(*OptionalSegmentHeaderData, UnresolvedCDOs);
+	}
+#endif
 	if (!UnresolvedCDOs.IsEmpty())
 	{
 		AsyncLoadingThread.AddPendingCDOs(this, UnresolvedCDOs);
@@ -3421,6 +3476,12 @@ void FAsyncPackage2::ImportPackagesRecursive(FIoBatch& IoBatch, IPackageStore& P
 	check(AsyncPackageLoadingState == EAsyncPackageLoadingState2::NewPackage);
 
 	int32 ImportedPackageCount = HeaderData.ImportedPackageIds.Num();
+#if WITH_EDITOR
+	if (OptionalSegmentHeaderData.IsSet())
+	{
+		ImportedPackageCount += OptionalSegmentHeaderData->ImportedPackageIds.Num();
+	}
+#endif
 	if (!ImportedPackageCount)
 	{
 		AsyncPackageLoadingState = EAsyncPackageLoadingState2::ImportPackagesDone;
@@ -3436,6 +3497,12 @@ void FAsyncPackage2::ImportPackagesRecursive(FIoBatch& IoBatch, IPackageStore& P
 	Data.ImportedAsyncPackages = MakeArrayView(Data.ImportedAsyncPackages.GetData(), ImportedPackageCount);
 	
 	ImportPackagesRecursiveInner(IoBatch, PackageStore, HeaderData.ImportedPackageIds, ImportedPackageIndex);
+#if WITH_EDITOR
+	if (OptionalSegmentHeaderData.IsSet())
+	{
+		ImportPackagesRecursiveInner(IoBatch, PackageStore, OptionalSegmentHeaderData->ImportedPackageIds, ImportedPackageIndex);
+	}
+#endif
 
 	UE_ASYNC_PACKAGE_LOG_VERBOSE(VeryVerbose, Desc, TEXT("ImportPackages: ImportsDone"),
 		TEXT("All imported packages are now being loaded."));
@@ -3454,10 +3521,35 @@ void FAsyncPackage2::StartLoading(FIoBatch& IoBatch)
 
 	AsyncPackageLoadingState = EAsyncPackageLoadingState2::WaitingForIo;
 
+	FIoReadOptions ReadOptions;
+#if WITH_EDITOR
+	if (OptionalSegmentHeaderData.IsSet())
+	{
+		int32 LocalPendingIoRequestsCounter = AsyncLoadingThread.PendingIoRequestsCounter.IncrementExchange() + 1;
+		TRACE_COUNTER_SET(AsyncLoadingPendingIoRequests, LocalPendingIoRequestsCounter);
+
+		GetPackageNode(EEventLoadNode2::Package_ProcessSummary).AddBarrier();
+		OptionalSegmentSerializationState->IoRequest = IoBatch.ReadWithCallback(CreateIoChunkId(Desc.PackageIdToLoad.Value(), 1, EIoChunkType::ExportBundleData),
+			ReadOptions,
+			Desc.Priority,
+			[this](TIoStatusOr<FIoBuffer> Result)
+			{
+				if (!Result.IsOk())
+				{
+					UE_ASYNC_PACKAGE_LOG(Warning, Desc, TEXT("StartBundleIoRequests: FailedRead"),
+						TEXT("Failed reading optional chunk for package: %s"), *Result.Status().ToString());
+					bLoadHasFailed = true;
+				}
+				int32 LocalPendingIoRequestsCounter = AsyncLoadingThread.PendingIoRequestsCounter.DecrementExchange() - 1;
+				TRACE_COUNTER_SET(AsyncLoadingPendingIoRequests, LocalPendingIoRequestsCounter);
+				GetPackageNode(EEventLoadNode2::Package_ProcessSummary).ReleaseBarrier();
+			});
+	}
+#endif
+
 	int32 LocalPendingIoRequestsCounter = AsyncLoadingThread.PendingIoRequestsCounter.IncrementExchange() + 1;
 	TRACE_COUNTER_SET(AsyncLoadingPendingIoRequests, LocalPendingIoRequestsCounter);
 
-	FIoReadOptions ReadOptions;
 	SerializationState.IoRequest = IoBatch.ReadWithCallback(CreatePackageDataChunkId(Desc.PackageIdToLoad),
 		ReadOptions,
 		Desc.Priority,
@@ -3530,6 +3622,7 @@ static void ReadAsyncPackageHeader(FAsyncPackageSerializationState& Serializatio
 	HeaderData.ExportMap = TArrayView<const FExportMapEntry>(
 		reinterpret_cast<const FExportMapEntry*>(PackageHeaderDataPtr + PackageSummary->ExportMapOffset),
 		(PackageSummary->ExportBundleEntriesOffset - PackageSummary->ExportMapOffset) / sizeof(FExportMapEntry));
+	check(HeaderData.ExportMap.Num() == HeaderData.ExportCount);
 
 	const uint64 ExportBundleHeadersOffset = PackageSummary->GraphDataOffset;
 	const uint64 ExportBundleHeadersSize = sizeof(FExportBundleHeader) * HeaderData.ExportBundleHeaders.Num();
@@ -3567,16 +3660,33 @@ EAsyncPackageState::Type FAsyncPackage2::Event_ProcessPackageSummary(FAsyncLoadi
 		check(Package->ExportBundleEntryIndex == 0);
 
 		ReadAsyncPackageHeader(Package->SerializationState, Package->HeaderData);
-
+#if WITH_EDITOR
+		FAsyncPackageHeaderData* OptionalSegmentHeaderData = Package->OptionalSegmentHeaderData.GetPtrOrNull();
+		if (OptionalSegmentHeaderData)
+		{
+			ReadAsyncPackageHeader(*Package->OptionalSegmentSerializationState, *OptionalSegmentHeaderData);
+		}
+#endif
 		for (int32 ExportBundleIndex = 0; ExportBundleIndex < Package->Data.ExportBundleCount; ++ExportBundleIndex)
 		{
 			const FAsyncPackageHeaderData* HeaderData = &Package->HeaderData;
-			const FExportBundleHeader& ExportBundle = HeaderData->ExportBundleHeaders[ExportBundleIndex];
+			int32 LocalExportBundleIndex = ExportBundleIndex;
+			int32 ExportIndexOffset = 0;
+#if WITH_EDITOR
+			if (ExportBundleIndex >= Package->HeaderData.ExportBundleHeaders.Num())
+			{
+				check(OptionalSegmentHeaderData);
+				HeaderData = OptionalSegmentHeaderData;
+				LocalExportBundleIndex -= Package->HeaderData.ExportBundleHeaders.Num();
+				ExportIndexOffset = Package->HeaderData.ExportMap.Num();
+			}
+#endif
+			const FExportBundleHeader& ExportBundle = HeaderData->ExportBundleHeaders[LocalExportBundleIndex];
 			for (int32 EntryIndex = ExportBundle.FirstEntryIndex, EntryEnd = ExportBundle.FirstEntryIndex + ExportBundle.EntryCount; EntryIndex < EntryEnd; ++EntryIndex)
 			{
 				const FExportBundleEntry& BundleEntry = HeaderData->ExportBundleEntries[EntryIndex];
 				const FExportMapEntry& ExportMapEntry = HeaderData->ExportMap[BundleEntry.LocalExportIndex];
-				FAsyncPackageExportToBundleMapping& ExportToBundleMapping = Package->Data.ExportToBundleMappings[BundleEntry.LocalExportIndex];
+				FAsyncPackageExportToBundleMapping& ExportToBundleMapping = Package->Data.ExportToBundleMappings[BundleEntry.LocalExportIndex + ExportIndexOffset];
 				ExportToBundleMapping.ExportHash = ExportMapEntry.PublicExportHash;
 				ExportToBundleMapping.BundleIndex[BundleEntry.CommandType] = ExportBundleIndex;
 			}
@@ -3615,6 +3725,12 @@ EAsyncPackageState::Type FAsyncPackage2::Event_SetupDependencies(FAsyncLoadingTh
 		TArrayView<FEventLoadNode2> ExportBundleNodesView = Package->Data.ExportBundleNodes;
 		TArrayView<FAsyncPackage2*> ImportedAsyncPackagesView = Package->Data.ImportedAsyncPackages;
 		Package->SetupSerializedArcs(Package->HeaderData, ExportBundleNodesView.Left(Package->HeaderData.ExportBundleHeaders.Num() * EEventLoadNode2::ExportBundle_NumPhases), ImportedAsyncPackagesView.Left(Package->HeaderData.ImportedPackageIds.Num()));
+#if WITH_EDITOR
+		if (Package->OptionalSegmentHeaderData.IsSet())
+		{
+			Package->SetupSerializedArcs(*Package->OptionalSegmentHeaderData, ExportBundleNodesView.Right(Package->OptionalSegmentHeaderData->ExportBundleHeaders.Num() * EEventLoadNode2::ExportBundle_NumPhases), ImportedAsyncPackagesView.Right(Package->OptionalSegmentHeaderData->ImportedPackageIds.Num()));
+		}
+#endif
 	}
 	Package->AsyncPackageLoadingState = EAsyncPackageLoadingState2::ProcessExportBundles;
 	for (int32 ExportBundleIndex = 0; ExportBundleIndex < Package->Data.ExportBundleCount; ++ExportBundleIndex)
@@ -3660,10 +3776,33 @@ EAsyncPackageState::Type FAsyncPackage2::Event_ProcessExportBundle(FAsyncLoading
 	
 	if (!Package->bLoadHasFailed)
 	{
+#if WITH_EDITOR
+		const FAsyncPackageHeaderData* HeaderData;
+		const FExportBundleHeader* ExportBundle;
+		FAsyncPackageSerializationState* SerializationState;
+		TArrayView<FExportObject> Exports = Package->Data.Exports;
+		if (InExportBundleIndex >= Package->HeaderData.ExportBundleHeaders.Num())
+		{
+			HeaderData = Package->OptionalSegmentHeaderData.GetPtrOrNull();
+			check(HeaderData);
+			ExportBundle = &HeaderData->ExportBundleHeaders[InExportBundleIndex - Package->HeaderData.ExportBundleHeaders.Num()];
+			Exports.RightInline(HeaderData->ExportCount);
+			SerializationState = Package->OptionalSegmentSerializationState.GetPtrOrNull();
+			check(SerializationState);
+		}
+		else
+		{
+			HeaderData = &Package->HeaderData;
+			ExportBundle = &Package->HeaderData.ExportBundleHeaders[InExportBundleIndex];
+			SerializationState = &Package->SerializationState;
+			Exports.LeftInline(Package->HeaderData.ExportCount);
+		}
+#else
 		const FAsyncPackageHeaderData* HeaderData = &Package->HeaderData;
 		const FExportBundleHeader* ExportBundle = &Package->HeaderData.ExportBundleHeaders[InExportBundleIndex];
 		FAsyncPackageSerializationState* SerializationState = &Package->SerializationState;
 		const TArrayView<FExportObject>& Exports = Package->Data.Exports;
+#endif
 		const FIoBuffer& IoBuffer = SerializationState->IoRequest.GetResultOrDie();
 		const uint64 AllExportDataSize = IoBuffer.DataSize() - (SerializationState->AllExportDataPtr - IoBuffer.Data());
 		if (Package->ExportBundleEntryIndex == 0)
@@ -3764,6 +3903,13 @@ EAsyncPackageState::Type FAsyncPackage2::Event_ProcessExportBundle(FAsyncLoading
 		Package->ProcessedExportBundlesCount = 0;
 		Package->HeaderData.OnReleaseHeaderBuffer();
 		Package->SerializationState.ReleaseIoRequest();
+#if WITH_EDITOR
+		if (Package->OptionalSegmentHeaderData.IsSet())
+		{
+			Package->OptionalSegmentHeaderData->OnReleaseHeaderBuffer();
+			Package->OptionalSegmentSerializationState->ReleaseIoRequest();
+		}
+#endif
 
 		if (Package->ExternalReadDependencies.Num() == 0)
 		{
@@ -4390,9 +4536,28 @@ EAsyncPackageState::Type FAsyncPackage2::Event_PostLoadExportBundle(FAsyncLoadin
 		const bool bAsyncPostLoadEnabled = FAsyncLoadingThreadSettings::Get().bAsyncPostLoadEnabled;
 		const bool bIsMultithreaded = Package->AsyncLoadingThread.IsMultithreaded();
 
+#if WITH_EDITOR
+		const FAsyncPackageHeaderData* HeaderData;
+		const FExportBundleHeader* ExportBundle;
+		TArrayView<FExportObject> Exports = Package->Data.Exports;
+		if (InExportBundleIndex >= Package->HeaderData.ExportBundleHeaders.Num())
+		{
+			HeaderData = Package->OptionalSegmentHeaderData.GetPtrOrNull();
+			check(HeaderData);
+			ExportBundle = &HeaderData->ExportBundleHeaders[InExportBundleIndex - Package->HeaderData.ExportBundleHeaders.Num()];
+			Exports.RightInline(HeaderData->ExportCount);
+		}
+		else
+		{
+			HeaderData = &Package->HeaderData;
+			ExportBundle = &Package->HeaderData.ExportBundleHeaders[InExportBundleIndex];
+			Exports.LeftInline(Package->HeaderData.ExportCount);
+		}
+#else
 		const FAsyncPackageHeaderData* HeaderData = &Package->HeaderData;
 		const FExportBundleHeader* ExportBundle = &Package->HeaderData.ExportBundleHeaders[InExportBundleIndex];
 		const TArrayView<FExportObject>& Exports = Package->Data.Exports;
+#endif
 
 		while (Package->ExportBundleEntryIndex < int32(ExportBundle->EntryCount))
 		{
@@ -4492,9 +4657,28 @@ EAsyncPackageState::Type FAsyncPackage2::Event_DeferredPostLoadExportBundle(FAsy
 		TGuardValue<bool> GuardIsRoutingPostLoad(PackageScope.ThreadContext.IsRoutingPostLoad, true);
 		FAsyncLoadingTickScope2 InAsyncLoadingTick(Package->AsyncLoadingThread);
 
+#if WITH_EDITOR
+		const FAsyncPackageHeaderData* HeaderData;
+		const FExportBundleHeader* ExportBundle;
+		TArrayView<FExportObject> Exports = Package->Data.Exports;
+		if (InExportBundleIndex >= Package->HeaderData.ExportBundleHeaders.Num())
+		{
+			HeaderData = Package->OptionalSegmentHeaderData.GetPtrOrNull();
+			check(HeaderData);
+			ExportBundle = &HeaderData->ExportBundleHeaders[InExportBundleIndex - Package->HeaderData.ExportBundleHeaders.Num()];
+			Exports.RightInline(HeaderData->ExportCount);
+		}
+		else
+		{
+			HeaderData = &Package->HeaderData;
+			ExportBundle = &Package->HeaderData.ExportBundleHeaders[InExportBundleIndex];
+			Exports.LeftInline(Package->HeaderData.ExportCount);
+		}
+#else
 		const FAsyncPackageHeaderData* HeaderData = &Package->HeaderData;
 		const FExportBundleHeader* ExportBundle = &Package->HeaderData.ExportBundleHeaders[InExportBundleIndex];
 		const TArrayView<FExportObject>& Exports = Package->Data.Exports;
+#endif
 
 		while (Package->ExportBundleEntryIndex < int32(ExportBundle->EntryCount))
 		{
diff --git a/Engine/Source/Runtime/CoreUObject/Private/Serialization/PackageWriterToSharedBuffer.cpp b/Engine/Source/Runtime/CoreUObject/Private/Serialization/PackageWriterToSharedBuffer.cpp
--- a/Engine/Source/Runtime/CoreUObject/Private/Serialization/PackageWriterToSharedBuffer.cpp
+++ b/Engine/Source/Runtime/CoreUObject/Private/Serialization/PackageWriterToSharedBuffer.cpp
@@ -38,7 +38,7 @@ void FPackageWriterRecords::BeginPackage(FPackage* Record, const IPackageWriter:
 void FPackageWriterRecords::WritePackageData(const IPackageWriter::FPackageInfo& Info,
 	FLargeMemoryWriter& ExportsArchive, const TArray<FFileRegion>& FileRegions)
 {
-	FPackage& Record = FindRecordChecked(Info.InputPackageName);
+	FPackage& Record = FindRecordChecked(Info.PackageName);
 	int64 DataSize = ExportsArchive.TotalSize();
 	checkf(DataSize > 0, TEXT("IPackageWriter->WritePackageData must not be called with an empty ExportsArchive"));
 	checkf(static_cast<uint64>(DataSize) >= Info.HeaderSize,
@@ -51,19 +51,19 @@ void FPackageWriterRecords::WritePackageData(const IPackageWriter::FPackageInfo&
 void FPackageWriterRecords::WriteBulkData(const IPackageWriter::FBulkDataInfo& Info, const FIoBuffer& BulkData,
 	const TArray<FFileRegion>& FileRegions)
 {
-	FPackage& Record = FindRecordChecked(Info.InputPackageName);
+	FPackage& Record = FindRecordChecked(Info.PackageName);
 	Record.BulkDatas.Add(FBulkData{ Info, IoBufferToSharedBuffer(BulkData), FileRegions });
 }
 
 void FPackageWriterRecords::WriteAdditionalFile(const IPackageWriter::FAdditionalFileInfo& Info,
 	const FIoBuffer& FileData)
 {
-	FPackage& Record = FindRecordChecked(Info.InputPackageName);
+	FPackage& Record = FindRecordChecked(Info.PackageName);
 	FAdditionalFile& AdditionalFile = Record.AdditionalFiles.Add_GetRef(FAdditionalFile{ Info, IoBufferToSharedBuffer(FileData) });
 	FIoChunkId ChunkId = CreateExternalFileChunkId(AdditionalFile.Info.Filename);
 	if (AdditionalFile.Info.ChunkId.IsValid() && AdditionalFile.Info.ChunkId != ChunkId)
 	{
-		UE_LOG(LogPackageWriter, Warning, TEXT("PackageWriter->WriteAdditionalFile called with an unexpected chunkid: Should match CreateExternalFileChunkId (Package: %s Filename: %s)"), *Info.InputPackageName.ToString(), *AdditionalFile.Info.Filename);
+		UE_LOG(LogPackageWriter, Warning, TEXT("PackageWriter->WriteAdditionalFile called with an unexpected chunkid: Should match CreateExternalFileChunkId (Package: %s Filename: %s)"), *Info.PackageName.ToString(), *AdditionalFile.Info.Filename);
 	}
 	else
 	{
@@ -74,7 +74,7 @@ void FPackageWriterRecords::WriteAdditionalFile(const IPackageWriter::FAdditiona
 void FPackageWriterRecords::WriteLinkerAdditionalData(const IPackageWriter::FLinkerAdditionalDataInfo& Info,
 	const FIoBuffer& Data, const TArray<FFileRegion>& FileRegions)
 {
-	FPackage& Record = FindRecordChecked(Info.InputPackageName);
+	FPackage& Record = FindRecordChecked(Info.PackageName);
 	Record.LinkerAdditionalDatas.Add(
 		FLinkerAdditionalData{ Info, IoBufferToSharedBuffer(Data), FileRegions });
 }
diff --git a/Engine/Source/Runtime/CoreUObject/Private/UObject/SavePackage.cpp b/Engine/Source/Runtime/CoreUObject/Private/UObject/SavePackage.cpp
--- a/Engine/Source/Runtime/CoreUObject/Private/UObject/SavePackage.cpp
+++ b/Engine/Source/Runtime/CoreUObject/Private/UObject/SavePackage.cpp
@@ -3857,7 +3857,7 @@ FSavePackageResultStruct UPackage::Save(UPackage* InOuter, UObject* InAsset, con
 							TotalPackageSizeUncompressed += Size;
 
 							IPackageWriter::FAdditionalFileInfo FileInfo;
-							FileInfo.OutputPackageName = FileInfo.InputPackageName = InOuter->GetFName();
+							FileInfo.PackageName = InOuter->GetFName();
 							FileInfo.Filename = *Writer.GetArchiveName();
 
 							FIoBuffer FileData(FIoBuffer::AssumeOwnership, Writer.ReleaseOwnership(), Size);
@@ -4137,11 +4137,11 @@ FSavePackageResultStruct UPackage::Save(UPackage* InOuter, UObject* InAsset, con
 						if (PackageWriter)
 						{
 							IPackageWriter::FPackageInfo PackageInfo;
-							PackageInfo.OutputPackageName = PackageInfo.InputPackageName = InOuter->GetFName();
+							PackageInfo.PackageName = InOuter->GetFName();
 							PackageInfo.LooseFilePath = Filename;
 							PackageInfo.HeaderSize = Linker->Summary.TotalHeaderSize;
 
-							FPackageId PackageId = FPackageId::FromName(PackageInfo.OutputPackageName);
+							FPackageId PackageId = FPackageId::FromName(PackageInfo.PackageName);
 							PackageInfo.ChunkId = CreateIoChunkId(PackageId.Value(), 0, EIoChunkType::ExportBundleData);
 
 							PackageWriter->WritePackageData(PackageInfo, *Writer, Linker->FileRegions);
diff --git a/Engine/Source/Runtime/CoreUObject/Private/UObject/SavePackage/SavePackageUtilities.cpp b/Engine/Source/Runtime/CoreUObject/Private/UObject/SavePackage/SavePackageUtilities.cpp
--- a/Engine/Source/Runtime/CoreUObject/Private/UObject/SavePackage/SavePackageUtilities.cpp
+++ b/Engine/Source/Runtime/CoreUObject/Private/UObject/SavePackage/SavePackageUtilities.cpp
@@ -1973,7 +1973,7 @@ ESavePackageResult CreatePayloadSidecarFile(FLinkerSave& Linker, const FPackageP
 	if (PackageWriter)
 	{
 		IPackageWriter::FAdditionalFileInfo SidecarSegmentInfo;
-		SidecarSegmentInfo.OutputPackageName = SidecarSegmentInfo.InputPackageName = PackagePath.GetPackageFName();
+		SidecarSegmentInfo.PackageName = PackagePath.GetPackageFName();
 		SidecarSegmentInfo.Filename = TargetFilePath;
 		FIoBuffer FileData(FIoBuffer::AssumeOwnership, Ar.ReleaseOwnership(), DataSize);
 		PackageWriter->WriteAdditionalFile(SidecarSegmentInfo, FileData);
@@ -2275,26 +2275,23 @@ ESavePackageResult SaveBulkData(FLinkerSave* Linker, int64& InOutStartOffset, co
 			};
 
 			IPackageWriter::FBulkDataInfo BulkInfo;
-			BulkInfo.InputPackageName = InOuter->GetFName();
-			// Adjust the OutputPackageName and LooseFilePath if needed
+			BulkInfo.PackageName = InOuter->GetFName();
+			// Adjust LooseFilePath if needed
 			if (bIsOptionalRealm)
 			{
 				// Optional output have the form PackagePath.o.ext
-				FString OutputPackageName = BulkInfo.InputPackageName.ToString() + FPackagePath::GetOptionalSegmentExtensionModifier();
-				BulkInfo.OutputPackageName = *OutputPackageName;
 				BulkInfo.LooseFilePath = FPathViews::ChangeExtension(Filename, TEXT("o.") + FPaths::GetExtension(Filename));
 				BulkInfo.MultiOutputIndex = 1;
 			}
 			else
 			{
-				BulkInfo.OutputPackageName = BulkInfo.InputPackageName;
 				BulkInfo.LooseFilePath = Filename;
 			}
-			FPackageId PackageId = FPackageId::FromName(BulkInfo.OutputPackageName);
+			FPackageId PackageId = FPackageId::FromName(BulkInfo.PackageName);
 				
 			if (BulkArchive->TotalSize())
 			{
-				BulkInfo.ChunkId = CreateIoChunkId(PackageId.Value(), 0, EIoChunkType::BulkData);
+				BulkInfo.ChunkId = CreateIoChunkId(PackageId.Value(), BulkInfo.MultiOutputIndex, EIoChunkType::BulkData);
 				BulkInfo.BulkDataType = bSeparateSegmentsEnabled ?
 					IPackageWriter::FBulkDataInfo::BulkSegment : IPackageWriter::FBulkDataInfo::AppendToExports;
 				BulkInfo.LooseFilePath = FPathViews::ChangeExtension(BulkInfo.LooseFilePath, LexToString(EPackageExtension::BulkDataDefault));
diff --git a/Engine/Source/Runtime/CoreUObject/Private/UObject/SavePackage2.cpp b/Engine/Source/Runtime/CoreUObject/Private/UObject/SavePackage2.cpp
--- a/Engine/Source/Runtime/CoreUObject/Private/UObject/SavePackage2.cpp
+++ b/Engine/Source/Runtime/CoreUObject/Private/UObject/SavePackage2.cpp
@@ -1847,7 +1847,7 @@ ESavePackageResult WriteAdditionalExportFiles(FSaveContext& SaveContext)
 			SaveContext.TotalPackageSizeUncompressed += Size;
 
 			IPackageWriter::FAdditionalFileInfo FileInfo;
-			FileInfo.OutputPackageName = FileInfo.InputPackageName = SaveContext.GetPackage()->GetFName();
+			FileInfo.PackageName = SaveContext.GetPackage()->GetFName();
 			FileInfo.Filename = *Writer.GetArchiveName();
 
 			FIoBuffer FileData(FIoBuffer::AssumeOwnership, Writer.ReleaseOwnership(), Size);
@@ -1991,25 +1991,22 @@ ESavePackageResult FinalizeFile(FStructuredArchive::FRecord& StructuredArchiveRo
 		if (PackageWriter)
 		{
 			IPackageWriter::FPackageInfo PackageInfo;
-			PackageInfo.InputPackageName = SaveContext.GetPackage()->GetFName();
-			// Adjust the OutputPackageName and LooseFilePath if needed
+			PackageInfo.PackageName = SaveContext.GetPackage()->GetFName();
+			// Adjust LooseFilePath if needed
 			if (bIsOptionalRealm)
 			{
 				// Optional output have the form PackagePath.o.ext
-				FString OutputPackageName = PackageInfo.InputPackageName.ToString() + FPackagePath::GetOptionalSegmentExtensionModifier();
-				PackageInfo.OutputPackageName = *OutputPackageName;
 				PackageInfo.LooseFilePath = FPathViews::ChangeExtension(SaveContext.GetFilename(), TEXT("o.") + FPaths::GetExtension(SaveContext.GetFilename()));
 				PackageInfo.MultiOutputIndex = 1;
 			}
 			else
 			{
-				PackageInfo.OutputPackageName = PackageInfo.InputPackageName;
 				PackageInfo.LooseFilePath = SaveContext.GetFilename();
 			}
 			PackageInfo.HeaderSize = Linker->Summary.TotalHeaderSize;
 
-			FPackageId PackageId = FPackageId::FromName(PackageInfo.OutputPackageName);
-			PackageInfo.ChunkId = CreateIoChunkId(PackageId.Value(), 0, EIoChunkType::ExportBundleData);
+			FPackageId PackageId = FPackageId::FromName(PackageInfo.PackageName);
+			PackageInfo.ChunkId = CreateIoChunkId(PackageId.Value(), PackageInfo.MultiOutputIndex, EIoChunkType::ExportBundleData);
 
 			PackageWriter->WritePackageData(PackageInfo, *Writer, Linker->FileRegions);
 		}
diff --git a/Engine/Source/Runtime/PakFile/Private/FilePackageStore.cpp b/Engine/Source/Runtime/PakFile/Private/FilePackageStore.cpp
--- a/Engine/Source/Runtime/PakFile/Private/FilePackageStore.cpp
+++ b/Engine/Source/Runtime/PakFile/Private/FilePackageStore.cpp
@@ -102,6 +102,15 @@ EPackageStoreEntryStatus FFilePackageStore::GetPackageStoreEntry(FPackageId Pack
 		OutPackageStoreEntry.ExportInfo.ExportBundleCount = FindEntry->ExportBundleCount;
 		OutPackageStoreEntry.ImportedPackageIds = MakeArrayView(FindEntry->ImportedPackages.Data(), FindEntry->ImportedPackages.Num());
 		OutPackageStoreEntry.ShaderMapHashes = MakeArrayView(FindEntry->ShaderMapHashes.Data(), FindEntry->ShaderMapHashes.Num());
+#if WITH_EDITOR
+		const FFilePackageStoreEntry* FindOptionalSegmentEntry = OptionalSegmentStoreEntriesMap.FindRef(PackageId);
+		if (FindOptionalSegmentEntry)
+		{
+			OutPackageStoreEntry.OptionalSegmentExportInfo.ExportCount = FindOptionalSegmentEntry->ExportCount;
+			OutPackageStoreEntry.OptionalSegmentExportInfo.ExportBundleCount = FindOptionalSegmentEntry->ExportBundleCount;
+			OutPackageStoreEntry.OptionalSegmentImportedPackageIds = MakeArrayView(FindOptionalSegmentEntry->ImportedPackages.Data(), FindOptionalSegmentEntry->ImportedPackages.Num());
+		}
+#endif
 		return EPackageStoreEntryStatus::Ok;
 	}
 	return EPackageStoreEntryStatus::Missing;
@@ -142,9 +151,13 @@ void FFilePackageStore::Mount(const FIoContainerHeader* ContainerHeader, uint32
 {
 	LLM_SCOPE(ELLMTag::AsyncLoading);
 	FWriteScopeLock _(EntriesLock);
-	MountedContainers.Add({ ContainerHeader, Order });
+	MountedContainers.Add({ ContainerHeader, Order, NextSequence++ });
 	Algo::Sort(MountedContainers, [](const FMountedContainer& A, const FMountedContainer& B)
 		{
+			if (A.Order == B.Order)
+			{
+				return A.Sequence < B.Sequence;
+			}
 			return A.Order < B.Order;
 		});
 	bNeedsUpdate = true;
@@ -178,18 +191,21 @@ void FFilePackageStore::Update()
 	StoreEntriesMap.Empty();
 	LocalizedPackages.Empty();
 	RedirectsPackageMap.Empty();
+#if WITH_EDITOR
+	OptionalSegmentStoreEntriesMap.Empty();
+#endif
 
 	uint32 TotalPackageCount = 0;
 	for (const FMountedContainer& MountedContainer : MountedContainers)
 	{
-		TotalPackageCount += MountedContainer.ContainerHeader->PackageCount;
+		TotalPackageCount += MountedContainer.ContainerHeader->PackageIds.Num();
 	}
 
 	StoreEntriesMap.Reserve(TotalPackageCount);
 	for (const FMountedContainer& MountedContainer : MountedContainers)
 	{
 		const FIoContainerHeader* ContainerHeader = MountedContainer.ContainerHeader;
-		TArrayView<const FFilePackageStoreEntry> ContainerStoreEntries(reinterpret_cast<const FFilePackageStoreEntry*>(ContainerHeader->StoreEntries.GetData()), ContainerHeader->PackageCount);
+		TArrayView<const FFilePackageStoreEntry> ContainerStoreEntries(reinterpret_cast<const FFilePackageStoreEntry*>(ContainerHeader->StoreEntries.GetData()), ContainerHeader->PackageIds.Num());
 		int32 Index = 0;
 		for (const FFilePackageStoreEntry& StoreEntry : ContainerStoreEntries)
 		{
@@ -199,6 +215,18 @@ void FFilePackageStore::Update()
 			++Index;
 		}
 
+#if WITH_EDITOR
+		TArrayView<const FFilePackageStoreEntry> ContainerOptionalSegmentStoreEntries(reinterpret_cast<const FFilePackageStoreEntry*>(ContainerHeader->OptionalSegmentStoreEntries.GetData()), ContainerHeader->OptionalSegmentPackageIds.Num());
+		Index = 0;
+		for (const FFilePackageStoreEntry& OptionalSegmentStoreEntry : ContainerOptionalSegmentStoreEntries)
+		{
+			const FPackageId& PackageId = ContainerHeader->OptionalSegmentPackageIds[Index];
+			check(PackageId.IsValid());
+			OptionalSegmentStoreEntriesMap.FindOrAdd(PackageId, &OptionalSegmentStoreEntry);
+			++Index;
+		}
+#endif
+
 		for (const FIoContainerHeaderLocalizedPackage& LocalizedPackage : ContainerHeader->LocalizedPackages)
 		{
 			FName& LocalizedPackageSourceName = LocalizedPackages.FindOrAdd(LocalizedPackage.SourcePackageId);
diff --git a/Engine/Source/Runtime/PakFile/Private/FilePackageStore.h b/Engine/Source/Runtime/PakFile/Private/FilePackageStore.h
--- a/Engine/Source/Runtime/PakFile/Private/FilePackageStore.h
+++ b/Engine/Source/Runtime/PakFile/Private/FilePackageStore.h
@@ -35,6 +35,7 @@ private:
 	{
 		const FIoContainerHeader* ContainerHeader;
 		uint32 Order;
+		uint32 Sequence;
 	};
 
 #if WITH_EDITOR
@@ -53,6 +54,7 @@ private:
 	FRWLock EntriesLock;
 	FCriticalSection UpdateLock;
 	TArray<FMountedContainer> MountedContainers;
+	TAtomic<uint32> NextSequence{ 0 };
 	TMap<FPackageId, const FFilePackageStoreEntry*> StoreEntriesMap;
 	TMap<FPackageId, TTuple<FName, FPackageId>> RedirectsPackageMap;
 	TMap<FPackageId, FName> LocalizedPackages;
@@ -64,6 +66,9 @@ private:
 	TSet<FString> PendingUncookedPackageRoots;
 	TMap<FPackageId, FUncookedPackage> UncookedPackagesMap;
 #endif
+#if WITH_EDITOR
+	TMap<FPackageId, const FFilePackageStoreEntry*> OptionalSegmentStoreEntriesMap;
+#endif
 
 	static thread_local int32 LockedOnThreadCount;
 };
diff --git a/Engine/Source/Runtime/PakFile/Private/IPlatformFilePak.cpp b/Engine/Source/Runtime/PakFile/Private/IPlatformFilePak.cpp
--- a/Engine/Source/Runtime/PakFile/Private/IPlatformFilePak.cpp
+++ b/Engine/Source/Runtime/PakFile/Private/IPlatformFilePak.cpp
@@ -7756,6 +7756,23 @@ bool FPakPlatformFile::Mount(const TCHAR* InPakFilename, uint32 PakOrder, const
 					UE_LOG(LogPakFile, Display, TEXT("Mounted IoStore container \"%s\""), *UtocPath);
 					Pak->IoContainerHeader = MakeUnique<FIoContainerHeader>(MountResult.ConsumeValueOrDie());
 					FilePackageStore->Mount(Pak->IoContainerHeader.Get(), PakOrder);
+#if WITH_EDITOR
+					FString OptionalSegmentUtocPath = FPaths::ChangeExtension(InPakFilename, FString::Printf(TEXT("%s.utoc"), FPackagePath::GetOptionalSegmentExtensionModifier()));
+					if (FPlatformFileManager::Get().GetPlatformFile().FileExists(*OptionalSegmentUtocPath))
+					{
+						MountResult = IoDispatcherFileBackend->Mount(*OptionalSegmentUtocPath, PakOrder, EncryptionKeyGuid, EncryptionKey);
+						if (MountResult.IsOk())
+						{
+							Pak->OptionalSegmentIoContainerHeader = MakeUnique<FIoContainerHeader>(MountResult.ConsumeValueOrDie());
+							FilePackageStore->Mount(Pak->OptionalSegmentIoContainerHeader.Get(), PakOrder);
+							UE_LOG(LogPakFile, Display, TEXT("Mounted optional segment extension IoStore container \"%s\""), *OptionalSegmentUtocPath);
+						}
+						else
+						{
+							UE_LOG(LogPakFile, Warning, TEXT("Failed to mount optional segment extension IoStore container \"%s\" [%s]"), *OptionalSegmentUtocPath, *MountResult.Status().ToString());
+						}
+					}
+#endif
 				}
 				else
 				{
@@ -7841,6 +7858,14 @@ bool FPakPlatformFile::Unmount(const TCHAR* InPakFilename)
 			}
 			FString ContainerPath = FPaths::ChangeExtension(InPakFilename, FString());
 			bRemovedContainerFile = IoDispatcherFileBackend->Unmount(*ContainerPath);
+#if WITH_EDITOR
+			if (UnmountedPak && UnmountedPak->OptionalSegmentIoContainerHeader.IsValid())
+			{
+				FilePackageStore->Unmount(UnmountedPak->OptionalSegmentIoContainerHeader.Get());
+				FString OptionalSegmentContainerPath = ContainerPath + FPackagePath::GetOptionalSegmentExtensionModifier();
+				IoDispatcherFileBackend->Unmount(*OptionalSegmentContainerPath);
+			}
+#endif
 		}
 
 		if (UnmountedPak)
diff --git a/Engine/Source/Runtime/PakFile/Public/IPlatformFilePak.h b/Engine/Source/Runtime/PakFile/Public/IPlatformFilePak.h
--- a/Engine/Source/Runtime/PakFile/Public/IPlatformFilePak.h
+++ b/Engine/Source/Runtime/PakFile/Public/IPlatformFilePak.h
@@ -799,6 +799,9 @@ private:
 	bool bIsMounted;
 
 	TUniquePtr<FIoContainerHeader> IoContainerHeader;
+#if WITH_EDITOR
+	TUniquePtr<FIoContainerHeader> OptionalSegmentIoContainerHeader;
+#endif
 
 	static inline int32 CDECL CompareFilenameHashes(const void* Left, const void* Right)
 	{
diff --git a/Engine/Source/Runtime/StorageServerClient/Private/StorageServerPackageStore.cpp b/Engine/Source/Runtime/StorageServerClient/Private/StorageServerPackageStore.cpp
--- a/Engine/Source/Runtime/StorageServerClient/Private/StorageServerPackageStore.cpp
+++ b/Engine/Source/Runtime/StorageServerClient/Private/StorageServerPackageStore.cpp
@@ -21,8 +21,8 @@ FStorageServerPackageStore::FStorageServerPackageStore(FStorageServerConnection&
 			FIoContainerHeader ContainerHeader;
 			Ar << ContainerHeader;
 			StoreEntriesData = MoveTemp(ContainerHeader.StoreEntries);
-			StoreEntriesMap.Reserve(ContainerHeader.PackageCount);
-			TArrayView<const FFilePackageStoreEntry> StoreEntries(reinterpret_cast<const FFilePackageStoreEntry*>(StoreEntriesData.GetData()), ContainerHeader.PackageCount);
+			StoreEntriesMap.Reserve(ContainerHeader.PackageIds.Num());
+			TArrayView<const FFilePackageStoreEntry> StoreEntries(reinterpret_cast<const FFilePackageStoreEntry*>(StoreEntriesData.GetData()), ContainerHeader.PackageIds.Num());
 			int32 Index = 0;
 			for (const FFilePackageStoreEntry& StoreEntry : StoreEntries)
 			{
