ZenStoreWriter: Remove spurious warnings when packages are committed twice during the cook because of reconsidering only-editoronly-referenced packages.
#rnx

diff --git a/Engine/Source/Developer/IoStoreUtilities/Private/ZenStoreWriter.cpp b/Engine/Source/Developer/IoStoreUtilities/Private/ZenStoreWriter.cpp
--- a/Engine/Source/Developer/IoStoreUtilities/Private/ZenStoreWriter.cpp
+++ b/Engine/Source/Developer/IoStoreUtilities/Private/ZenStoreWriter.cpp
@@ -592,14 +592,20 @@ void FZenStoreWriter::CommitPackage(FCommitPackageInfo&& Info)
 		FPendingPackageState& ExistingState = GetPendingPackage(Info.PackageName);
 		ExistingState.PackageHashes = new FPackageHashes();
 
-		// This looks weird but we're finding the _refcount_, not the hashes. So if it gets
-		// constructed, it's not actually assigned a pointer.
-		TRefCountPtr<FPackageHashes>& ExistingPackageHashes = AllPackageHashes.FindOrAdd(Info.PackageName);
-		if (ExistingPackageHashes.IsValid())
+		if (Info.bSucceeded)
 		{
-			UE_LOG(LogZenStoreWriter, Error, TEXT("FZenStoreWriter commiting the same package twice during a cook! (%s)"), *Info.PackageName.ToString());
+			// Only record hashes for successful saves. A single package can be saved unsuccessfully multiple times
+			// during a cook if its rejected for being only referenced by editor-only references and we keep finding
+			// new references to it.
+			TRefCountPtr<FPackageHashes>& ExistingPackageHashes = AllPackageHashes.FindOrAdd(Info.PackageName);
+			// This looks weird but we've found the _TRefCountPtr_, not the FPackageHashes. When newly assigned
+			// it will be an empty pointer, which is what we want.
+			if (ExistingPackageHashes.IsValid())
+			{
+				UE_LOG(LogZenStoreWriter, Error, TEXT("FZenStoreWriter commiting the same package twice during a cook! (%s)"), *Info.PackageName.ToString());
+			}
+			ExistingPackageHashes = ExistingState.PackageHashes;
 		}
-		ExistingPackageHashes = ExistingState.PackageHashes;
 	}
 
 	if (FPlatformProcess::SupportsMultithreading())
