ZenCOTF: Add Oodle dictionary files to the file manifest

#rnx

diff --git a/Engine/Source/Developer/IoStoreUtilities/Private/ZenFileSystemManifest.cpp b/Engine/Source/Developer/IoStoreUtilities/Private/ZenFileSystemManifest.cpp
--- a/Engine/Source/Developer/IoStoreUtilities/Private/ZenFileSystemManifest.cpp
+++ b/Engine/Source/Developer/IoStoreUtilities/Private/ZenFileSystemManifest.cpp
@@ -281,6 +281,10 @@ int32 FZenFileSystemManifest::Generate()
 	AddFilesFromDirectory(TEXT("/{engine}/Content/Movies"), FPaths::Combine(EngineDir, TEXT("Content"), TEXT("Movies")), true, &ContentFilter);
 	AddFilesFromDirectory(TEXT("/{project}/Content/Movies"), FPaths::Combine(ProjectDir, TEXT("Content"), TEXT("Movies")), true, &ContentFilter);
 	
+	FFileFilter OoodleDictionaryFilter = FFileFilter()
+		.IncludeExtension(TEXT("udic"));
+	AddFilesFromDirectory(TEXT("/{project}/Content/Oodle"), FPaths::Combine(ProjectDir, TEXT("Content"), TEXT("Oodle")), false, &OoodleDictionaryFilter);
+
 	FFileFilter ShaderCacheFilter = FFileFilter()
 		.IncludeExtension(TEXT("ushadercache"))
 		.IncludeExtension(TEXT("upipelinecache"));
