IoStore: Fixed non-deterministic shader classification and ordering

diff --git a/Engine/Source/Developer/IoStoreUtilities/Private/IoStoreUtilities.cpp b/Engine/Source/Developer/IoStoreUtilities/Private/IoStoreUtilities.cpp
--- a/Engine/Source/Developer/IoStoreUtilities/Private/IoStoreUtilities.cpp
+++ b/Engine/Source/Developer/IoStoreUtilities/Private/IoStoreUtilities.cpp
@@ -1861,7 +1861,14 @@ void ProcessShaderLibraries(const FIoStoreArguments& Arguments, TArray<FContaine
 						OutShaders.Add(ShaderInfo);
 						ChunkIdToShaderInfoMap.Add(ShaderChunkId, ShaderInfo);
 					}
-					ShaderInfo->TypeInContainer.Add(ContainerTarget, ShaderType);
+
+					const FShaderInfo::EShaderType* CurrentShaderTypeInContainer = ShaderInfo->TypeInContainer.Find(ContainerTarget);
+					if (!CurrentShaderTypeInContainer || *CurrentShaderTypeInContainer != FShaderInfo::Global)
+					{
+						// If a shader is both global and shared consider it to be global
+						ShaderInfo->TypeInContainer.Add(ContainerTarget, ShaderType);
+					}
+					
 					ContainerShaderLibraryShaders.Add(ShaderInfo);
 				}
 				for (const TTuple<FSHAHash, TSet<FName>>& ShaderMapAssetAssociation : ShaderMapAssetAssociations)
