Several IoStore determinism fixes:
-- Shaders must be sorted prior to inlining, as they are inserted in to a sorted list.
-- Break ties sorting shaders on disk with chunkid (hash)
-- Sort the shader map hashes in the package store entry
-- Pick the lowest LoadOrderFactor instead of the first for shaders that are used in more than one package.
-- Fix -WriteFinalOrder to actually write the final order as well as all chunks.

diff --git a/Engine/Source/Developer/IoStoreUtilities/Private/IoStoreUtilities.cpp b/Engine/Source/Developer/IoStoreUtilities/Private/IoStoreUtilities.cpp
--- a/Engine/Source/Developer/IoStoreUtilities/Private/IoStoreUtilities.cpp
+++ b/Engine/Source/Developer/IoStoreUtilities/Private/IoStoreUtilities.cpp
@@ -443,6 +443,24 @@ struct FShaderInfo
 	uint64 DiskLayoutOrder = MAX_uint64;
 	TSet<struct FLegacyCookedPackage*> ReferencedByPackages;
 	TMap<FContainerTargetSpec*, EShaderType> TypeInContainer;
+
+	// This is used to ensure build determinism and such must be stable
+	// across builds.
+	static bool Sort(const FShaderInfo* A, const FShaderInfo* B)
+	{
+		if (A->DiskLayoutOrder == B->DiskLayoutOrder)
+		{
+			if (A->LoadOrderFactor == B->LoadOrderFactor)
+			{
+				// Shader chunk IDs are the hash of the shader so this is consistent across builds.
+				uint64 AHash = *(uint64*)A->ChunkId.GetData();
+				uint64 BHash = *(uint64*)B->ChunkId.GetData();
+				return AHash < BHash;
+			}
+			return A->LoadOrderFactor < B->LoadOrderFactor;
+		}
+		return A->DiskLayoutOrder < B->DiskLayoutOrder;
+	}
 };
 
 struct FLegacyCookedPackage
@@ -1216,6 +1234,10 @@ static void CreateDiskLayout(
 			if (TargetFile->ChunkType == EContainerChunkType::PackageData)
 			{
 				TArray<FContainerTargetFile*, TInlineAllocator<1024>> PackageInlineShaders;
+
+				// Since we are inserting in to a sorted array (on disk order), we have to be stably sorted
+				// beforehand
+				Algo::Sort(TargetFile->Package->Shaders, FShaderInfo::Sort);
 				for (FShaderInfo* Shader : TargetFile->Package->Shaders)
 				{
 					check(Shader->ReferencedByPackages.Num() > 0);
@@ -1248,21 +1270,7 @@ static void CreateDiskLayout(
 			if (!Shaders.IsEmpty())
 			{
 				TArray<FShaderInfo*> SortedShaders = Shaders.Array();
-				for (FShaderInfo* Shader : SortedShaders)
-				{
-					for (FLegacyCookedPackage* Package : Shader->ReferencedByPackages)
-					{
-						Shader->DiskLayoutOrder = FMath::Min(Shader->DiskLayoutOrder, Package->DiskLayoutOrder);
-					}
-				}
-				Algo::Sort(SortedShaders, [](const FShaderInfo* A, const FShaderInfo* B)
-				{
-					if (A->DiskLayoutOrder == B->DiskLayoutOrder)
-					{
-						return A->LoadOrderFactor < B->LoadOrderFactor;
-					}
-					return A->DiskLayoutOrder < B->DiskLayoutOrder;
-				});
+				Algo::Sort(SortedShaders, FShaderInfo::Sort);
 				TArray<FContainerTargetFile*> ShaderTargetFiles;
 				ShaderTargetFiles.Reserve(SortedShaders.Num());
 				for (const FShaderInfo* ShaderInfo : SortedShaders)
@@ -1862,6 +1870,16 @@ void ProcessShaderLibraries(const FIoStoreArguments& Arguments, TArray<FContaine
 						OutShaders.Add(ShaderInfo);
 						ChunkIdToShaderInfoMap.Add(ShaderChunkId, ShaderInfo);
 					}
+					else
+					{
+						// If we already exist, then we have two separate LoadOrderFactors,
+						// which one we got first affects build determinism. Take the lower.
+						if (CodeChunk.Get<2>() < ShaderInfo->LoadOrderFactor)
+						{
+							ShaderInfo->LoadOrderFactor = CodeChunk.Get<2>();
+						}
+					}
+
 
 					const FShaderInfo::EShaderType* CurrentShaderTypeInContainer = ShaderInfo->TypeInContainer.Find(ContainerTarget);
 					if (!CurrentShaderTypeInContainer || *CurrentShaderTypeInContainer != FShaderInfo::Global)
@@ -1885,10 +1903,13 @@ void ProcessShaderLibraries(const FIoStoreArguments& Arguments, TArray<FContaine
 				{
 					ShaderChunkIdsByShaderMapHash.Add(ShaderMap.Key, MoveTemp(ShaderMap.Value));
 				}
-			}
-		}
-	}
+			} // end if containerchunktype shadercodelibrary
+		} // end foreach targetfile
+	} // end foreach container
 
+	// 1. Update ShaderInfos with which packages we reference.
+	// 2. Add to packages which shaders we use.
+	// 3. Add to PackageStore what shaders we use.
 	for (FContainerTargetSpec* ContainerTarget : ContainerTargets)
 	{
 		for (FLegacyCookedPackage* Package : ContainerTarget->Packages)
@@ -3364,13 +3385,29 @@ int32 CreateTarget(const FIoStoreArguments& Arguments, const FIoStoreWriterSetti
 			{
 				IoOrderListArchive->SetIsTextFormat(true);
 
+				// The TargetFiles list is not the order written to disk - FIoStoreWriter sorts
+				// on the IdealOrder prior to writing.
+				TArray<const FContainerTargetFile*> SortedTargetFiles;
+				SortedTargetFiles.Reserve(ContainerTarget->TargetFiles.Num());
 				for (const FContainerTargetFile& TargetFile : ContainerTarget->TargetFiles)
 				{
-					if (TargetFile.Package)
+					SortedTargetFiles.Add(&TargetFile);
+				}
+
+				Algo::SortBy(SortedTargetFiles, [](const FContainerTargetFile* TargetFile) { return TargetFile->IdealOrder; });
+
+				for (const FContainerTargetFile* TargetFile : SortedTargetFiles)
+				{
+					TStringBuilder<256> Line;
+					Line.Append(LexToString(TargetFile->ChunkId));
+
+					if (TargetFile->Package)
 					{
-						FString Line = FString::Printf(TEXT("%s"), *TargetFile.Package->FileName);
-						IoOrderListArchive->Logf(TEXT("%s"), *Line);
+						Line.Append(TEXT(" "));
+						Line.Append(TargetFile->Package->FileName);
 					}
+
+					IoOrderListArchive->Logf(TEXT("%s"), Line.ToString());
 				}
 
 				IoOrderListArchive->Close();
diff --git a/Engine/Source/Developer/IoStoreUtilities/Private/PackageStoreOptimizer.cpp b/Engine/Source/Developer/IoStoreUtilities/Private/PackageStoreOptimizer.cpp
--- a/Engine/Source/Developer/IoStoreUtilities/Private/PackageStoreOptimizer.cpp
+++ b/Engine/Source/Developer/IoStoreUtilities/Private/PackageStoreOptimizer.cpp
@@ -1816,6 +1816,9 @@ FPackageStoreEntryResource FPackageStoreOptimizer::CreatePackageStoreEntry(const
 	Result.ExportInfo.ExportBundleCount = Package->GraphData.ExportBundles.Num();
 	Result.ImportedPackageIds = Package->ImportedPackageIds;
 	Result.ShaderMapHashes = Package->ShaderMapHashes.Array();
+	// Package->ShaderMapHashes is a TSet and is unsorted - we want this sorted as it's serialized
+	// in to the ContainerHeader, and if not sorted leads to staging non-determinism.
+	Algo::Sort(Result.ShaderMapHashes);
 	if (OptionalSegmentPackage)
 	{
 		Result.OptionalSegmentExportInfo.ExportCount = OptionalSegmentPackage->Exports.Num();
diff --git a/Engine/Source/Runtime/Core/Private/IO/IoStore.cpp b/Engine/Source/Runtime/Core/Private/IO/IoStore.cpp
--- a/Engine/Source/Runtime/Core/Private/IO/IoStore.cpp
+++ b/Engine/Source/Runtime/Core/Private/IO/IoStore.cpp
@@ -64,6 +64,16 @@ FString LexToString(const EIoChunkType Type)
 	}
 }
 
+FString LexToString(const FIoChunkId& Id)
+{
+	FString Output;
+	TArray<TCHAR, FString::AllocatorType>& CharArray = Output.GetCharArray();
+	CharArray.AddUninitialized(sizeof(FIoChunkId) * 2 + 1);
+	UE::String::BytesToHexLower(Id.Id, CharArray.GetData());
+	CharArray.Last() = TCHAR('\0');
+	return Output;
+}
+
 template<typename ArrayType>
 bool WriteArray(IFileHandle* FileHandle, const ArrayType& Array)
 {
diff --git a/Engine/Source/Runtime/Core/Public/IO/IoDispatcher.h b/Engine/Source/Runtime/Core/Public/IO/IoDispatcher.h
--- a/Engine/Source/Runtime/Core/Public/IO/IoDispatcher.h
+++ b/Engine/Source/Runtime/Core/Public/IO/IoDispatcher.h
@@ -691,6 +691,8 @@ public:
 		return Builder;
 	}
 
+	friend CORE_API FString LexToString(const FIoChunkId& Id);
+
 	inline bool operator ==(const FIoChunkId& Rhs) const
 	{
 		return 0 == FMemory::Memcmp(Id, Rhs.Id, sizeof Id);
diff --git a/Engine/Source/Runtime/Core/Public/IO/PackageStore.h b/Engine/Source/Runtime/Core/Public/IO/PackageStore.h
--- a/Engine/Source/Runtime/Core/Public/IO/PackageStore.h
+++ b/Engine/Source/Runtime/Core/Public/IO/PackageStore.h
@@ -84,7 +84,7 @@ struct FPackageStoreEntryResource
 	FPackageStoreExportInfo ExportInfo;
 	/** Imported package IDs. */
 	TArray<FPackageId> ImportedPackageIds;
-	/** Referenced shader map hashes. */
+	/** Referenced shader map hashes - must be sorted in order to preserve determinism across builds. */
 	TArray<FSHAHash> ShaderMapHashes;
 	/** The editor data package export information. */
 	FPackageStoreExportInfo OptionalSegmentExportInfo;
