- Fixed typo on printf, I swear I had compiled the editor, I guess not sigh

diff --git a/Engine/Source/Developer/IoStoreUtilities/Private/IoStoreUtilities.cpp b/Engine/Source/Developer/IoStoreUtilities/Private/IoStoreUtilities.cpp
--- a/Engine/Source/Developer/IoStoreUtilities/Private/IoStoreUtilities.cpp
+++ b/Engine/Source/Developer/IoStoreUtilities/Private/IoStoreUtilities.cpp
@@ -2442,7 +2442,7 @@ void InitializeContainerTargetsAndPackages(
 					ContainerTarget->OptionalSegmentOutputPath = FPaths::Combine(ContainerSource.OptionalOutputPath, FPaths::GetCleanFilename(ContainerTarget->OutputPath) + FPackagePath::GetOptionalSegmentExtensionModifier());
 				}
 
-				UE_LOG(LogIoStore, Display, TEXT("Saving optional container to: '%s'", *ContainerTarget->OptionalSegmentOutputPath));
+				UE_LOG(LogIoStore, Display, TEXT("Saving optional container to: '%s'"), *ContainerTarget->OptionalSegmentOutputPath);
 			}
 		}
 	}
