Fix crash in package optimizer when recursively resolving exports which an import outer in the chain
Removed default parameter to the function that obfuscated the omission

diff --git a/Engine/Source/Developer/IoStoreUtilities/Internal/PackageStoreOptimizer.h b/Engine/Source/Developer/IoStoreUtilities/Internal/PackageStoreOptimizer.h
--- a/Engine/Source/Developer/IoStoreUtilities/Internal/PackageStoreOptimizer.h
+++ b/Engine/Source/Developer/IoStoreUtilities/Internal/PackageStoreOptimizer.h
@@ -395,9 +395,9 @@ private:
 	void ResolveImport(FPackageStorePackage::FUnresolvedImport* Imports, const FObjectImport* ObjectImports, int32 LocalImportIndex) const;
 	void ProcessImports(const FCookedHeaderData& CookedHeaderData, FPackageStorePackage* Package, TArray<FPackageStorePackage::FUnresolvedImport>& UnresolvedImports) const;
 	void ProcessImports(const FPackageStoreHeaderData& PackageStoreHeaderData, FPackageStorePackage* Package) const;
-	void ResolveExport(FPackageStorePackage::FExport* Exports, const FObjectExport* ObjectExports, const int32 LocalExportIndex, const FName& PackageName, FPackageStorePackage::FUnresolvedImport* Imports = nullptr, const FObjectImport* ObjectImports = nullptr) const;
+	void ResolveExport(FPackageStorePackage::FExport* Exports, const FObjectExport* ObjectExports, const int32 LocalExportIndex, const FName& PackageName, FPackageStorePackage::FUnresolvedImport* Imports, const FObjectImport* ObjectImports) const;
 	void ResolveExport(FPackageStorePackage::FExport* Exports, const int32 LocalExportIndex, const FName& PackageName) const;
-	void ProcessExports(const FCookedHeaderData& CookedHeaderData, FPackageStorePackage* Package, FPackageStorePackage::FUnresolvedImport* Imports = nullptr) const;
+	void ProcessExports(const FCookedHeaderData& CookedHeaderData, FPackageStorePackage* Package, FPackageStorePackage::FUnresolvedImport* Imports) const;
 	void ProcessExports(const FPackageStoreHeaderData& PackageStoreHeaderData, FPackageStorePackage* Package) const;
 	void ProcessPreloadDependencies(const FCookedHeaderData& CookedHeaderData, FPackageStorePackage* Package) const;
 	void ProcessPreloadDependencies(const FPackageStoreHeaderData& PackageStoreHeaderData, FPackageStorePackage* Package) const;
diff --git a/Engine/Source/Developer/IoStoreUtilities/Private/PackageStoreOptimizer.cpp b/Engine/Source/Developer/IoStoreUtilities/Private/PackageStoreOptimizer.cpp
--- a/Engine/Source/Developer/IoStoreUtilities/Private/PackageStoreOptimizer.cpp
+++ b/Engine/Source/Developer/IoStoreUtilities/Private/PackageStoreOptimizer.cpp
@@ -512,7 +512,7 @@ void FPackageStoreOptimizer::ResolveExport(
 			if (ObjectExport->OuterIndex.IsExport())
 			{
 				int32 OuterExportIndex = ObjectExport->OuterIndex.ToExport();
-				ResolveExport(Exports, ObjectExports, OuterExportIndex, PackageName);
+				ResolveExport(Exports, ObjectExports, OuterExportIndex, PackageName, Imports, ObjectImports);
 				OuterName = &Exports[OuterExportIndex].FullName;
 			}
 			else
