ZenStoreWriter: Fix assertion "Trying to add package that is already pending" if the cooker commits a package a second time while the first commit of the package is queued for asynchronous completion aand has not yet reached CommitPackageInternal.

This is necessary because as of CL 20151503, the cooker can commit the same package multiple times if the first save failed due to a timeout.

#rnx

diff --git a/Engine/Source/Developer/IoStoreUtilities/Private/IoStoreUtilities.cpp b/Engine/Source/Developer/IoStoreUtilities/Private/IoStoreUtilities.cpp
--- a/Engine/Source/Developer/IoStoreUtilities/Private/IoStoreUtilities.cpp
+++ b/Engine/Source/Developer/IoStoreUtilities/Private/IoStoreUtilities.cpp
@@ -5744,6 +5744,7 @@ int32 Staged2Zen(const FString& BuildPath, const FKeyChain& KeyChain, const FStr
 		IPackageWriter::FCommitPackageInfo CommitInfo;
 		CommitInfo.PackageName = PackageInfo.PackageName;
 		CommitInfo.WriteOptions = IPackageWriter::EWriteOptions::Write;
+		CommitInfo.Status = IPackageWriter::ECommitStatus::Success;
 		ZenStoreWriter->CommitPackage(MoveTemp(CommitInfo));
 
 		int32 LocalUploadCount = UploadCount.IncrementExchange() + 1;
diff --git a/Engine/Source/Developer/IoStoreUtilities/Private/ZenStoreWriter.cpp b/Engine/Source/Developer/IoStoreUtilities/Private/ZenStoreWriter.cpp
--- a/Engine/Source/Developer/IoStoreUtilities/Private/ZenStoreWriter.cpp
+++ b/Engine/Source/Developer/IoStoreUtilities/Private/ZenStoreWriter.cpp
@@ -58,6 +58,12 @@ FMD5Hash IoHashToMD5(const FIoHash& IoHash)
 	return Hash;
 }
 
+struct FZenStoreWriter::FZenCommitInfo
+{
+	IPackageWriter::FCommitPackageInfo CommitInfo;
+	TUniquePtr<FPendingPackageState> PackageState;
+};
+
 class FZenStoreWriter::FCommitQueue
 {
 public:
@@ -65,18 +71,18 @@ public:
 		: NewCommitEvent(EEventMode::AutoReset)
 		, QueueEmptyEvent(EEventMode::ManualReset) {}
 	
-	void Enqueue(const FCommitPackageInfo& Commit)
+	void Enqueue(FZenCommitInfo&& Commit)
 	{
 		{
 			FScopeLock _(&QueueCriticalSection);
-			Queue.Enqueue(Commit);
+			Queue.Enqueue(MoveTemp(Commit));
 			QueueNum++;
 		}
 
 		NewCommitEvent->Trigger();
 	}
 
-	bool BlockAndDequeue(FCommitPackageInfo& OutCommit)
+	bool BlockAndDequeue(FZenCommitInfo& OutCommit)
 	{
 		for (;;)
 		{
@@ -132,7 +138,7 @@ private:
 	FEventRef					NewCommitEvent;
 	FEventRef					QueueEmptyEvent;
 	FCriticalSection			QueueCriticalSection;
-	TQueue<FCommitPackageInfo>	Queue;
+	TQueue<FZenCommitInfo>		Queue;
 	int32						QueueNum = 0;
 	TAtomic<bool>				bCompleteAdding{false};
 };
@@ -517,7 +523,7 @@ void FZenStoreWriter::BeginCook()
 			for(;;)
 			{
 				TRACE_CPUPROFILER_EVENT_SCOPE(FZenStoreWriter::WaitingOnCooker);
-				FCommitPackageInfo Commit;
+				FZenCommitInfo Commit;
 				if (!CommitQueue->BlockAndDequeue(Commit))
 				{
 					break;
@@ -583,6 +589,12 @@ bool FZenStoreWriter::IsReservedOplogKey(FUtf8StringView Key)
 
 void FZenStoreWriter::CommitPackage(FCommitPackageInfo&& Info)
 {
+	if (Info.Status == ECommitStatus::Canceled)
+	{
+		RemovePendingPackage(Info.PackageName);
+		return;
+	}
+
 	UE::SavePackageUtilities::IncrementOutstandingAsyncWrites();
 
 	// If we are computing hashes, we need to allocate where the hashes will go.
@@ -592,7 +604,7 @@ void FZenStoreWriter::CommitPackage(FCommitPackageInfo&& Info)
 		FPendingPackageState& ExistingState = GetPendingPackage(Info.PackageName);
 		ExistingState.PackageHashes = new FPackageHashes();
 
-		if (Info.bSucceeded)
+		if (Info.Status == ECommitStatus::Success)
 		{
 			// Only record hashes for successful saves. A single package can be saved unsuccessfully multiple times
 			// during a cook if its rejected for being only referenced by editor-only references and we keep finding
@@ -608,21 +620,24 @@ void FZenStoreWriter::CommitPackage(FCommitPackageInfo&& Info)
 		}
 	}
 
+	TUniquePtr<FPendingPackageState> PackageState = RemovePendingPackage(Info.PackageName);
+	FZenCommitInfo ZenCommitInfo{ Forward<FCommitPackageInfo>(Info), MoveTemp(PackageState) };
 	if (FPlatformProcess::SupportsMultithreading())
 	{
-		CommitQueue->Enqueue(Info);
+		CommitQueue->Enqueue(MoveTemp(ZenCommitInfo));
 	}
 	else
 	{
-		CommitPackageInternal(Forward<FCommitPackageInfo>(Info));
+		CommitPackageInternal(MoveTemp(ZenCommitInfo));
 	}
 }
 
-void FZenStoreWriter::CommitPackageInternal(FCommitPackageInfo&& CommitInfo)
+void FZenStoreWriter::CommitPackageInternal(FZenCommitInfo&& ZenCommitInfo)
 {
 	TRACE_CPUPROFILER_EVENT_SCOPE(FZenStoreWriter::CommitPackage);
+	FCommitPackageInfo& CommitInfo = ZenCommitInfo.CommitInfo;
 	
-	TUniquePtr<FPendingPackageState> PackageState = RemovePendingPackage(CommitInfo.PackageName);
+	TUniquePtr<FPendingPackageState> PackageState = MoveTemp(ZenCommitInfo.PackageState);
 	checkf(PackageState.IsValid(), TEXT("Trying to commit non-pending package '%s'"), *CommitInfo.PackageName.ToString());
 
 	IPackageStoreWriter::FCommitEventArgs CommitEventArgs;
@@ -634,7 +649,7 @@ void FZenStoreWriter::CommitPackageInternal(FCommitPackageInfo&& CommitInfo)
 	const bool bComputeHash			= EnumHasAnyFlags(CommitInfo.WriteOptions, EWriteOptions::ComputeHash);
 	const bool bWritePackage		= EnumHasAnyFlags(CommitInfo.WriteOptions, EWriteOptions::Write);
 
-	if (CommitInfo.bSucceeded && bWritePackage)
+	if (CommitInfo.Status == ECommitStatus::Success && bWritePackage)
 	{
 		FPackageDataEntry& PkgData = PackageState->PackageData;
 		checkf(PkgData.IsValid, TEXT("CommitPackage called with bSucceeded but without first calling WritePackageData"))
@@ -807,7 +822,7 @@ void FZenStoreWriter::CommitPackageInternal(FCommitPackageInfo&& CommitInfo)
 		TIoStatusOr<uint64> Status = HttpClient->AppendOp(MoveTemp(OplogEntry));
 		UE_CLOG(!Status.IsOk(), LogZenStoreWriter, Error, TEXT("Failed to commit oplog entry '%s' to Zen"), *CommitInfo.PackageName.ToString());
 	}
-	else if (CommitInfo.bSucceeded && bComputeHash)
+	else if (CommitInfo.Status == ECommitStatus::Success && bComputeHash)
 	{
 		FPackageDataEntry& PkgData = PackageState->PackageData;
 		checkf(PkgData.IsValid, TEXT("CommitPackage called with bSucceeded but without first calling WritePackageData"));
diff --git a/Engine/Source/Developer/IoStoreUtilities/Public/ZenStoreWriter.h b/Engine/Source/Developer/IoStoreUtilities/Public/ZenStoreWriter.h
--- a/Engine/Source/Developer/IoStoreUtilities/Public/ZenStoreWriter.h
+++ b/Engine/Source/Developer/IoStoreUtilities/Public/ZenStoreWriter.h
@@ -143,7 +143,8 @@ private:
 	void CreateProjectMetaData(FCbPackage& Pkg, FCbWriter& PackageObj, bool bGenerateContainerHeader);
 	void BroadcastCommit(IPackageStoreWriter::FCommitEventArgs& EventArgs);
 	void BroadcastMarkUpToDate(IPackageStoreWriter::FMarkUpToDateEventArgs& EventArgs);
-	void CommitPackageInternal(FCommitPackageInfo&& CommitInfo);
+	struct FZenCommitInfo;
+	void CommitPackageInternal(FZenCommitInfo&& CommitInfo);
 	FCbAttachment CreateAttachment(FSharedBuffer Buffer);
 	FCbAttachment CreateAttachment(FIoBuffer Buffer);
 
diff --git a/Engine/Source/Editor/UnrealEd/Private/CookOnTheFlyServer.cpp b/Engine/Source/Editor/UnrealEd/Private/CookOnTheFlyServer.cpp
--- a/Engine/Source/Editor/UnrealEd/Private/CookOnTheFlyServer.cpp
+++ b/Engine/Source/Editor/UnrealEd/Private/CookOnTheFlyServer.cpp
@@ -3398,7 +3398,14 @@ void UCookOnTheFlyServer::PumpSaves(UE::Cook::FTickStackData& StackData, uint32
 			// so we call ReleaseCookedPlatformData and ClearObjectCache to clear it and recache on next attempt.
 			ReleaseCookedPlatformData(PackageData, false/* bCompletedSave */);
 			PackageData.ClearObjectCache();
-			SaveQueue.Add(&PackageData);
+			if (PackageData.GetIsUrgent())
+			{
+				SaveQueue.AddFront(&PackageData);
+			}
+			else
+			{
+				SaveQueue.Add(&PackageData);
+			}
 			continue;
 		}
 
@@ -4485,7 +4492,18 @@ void FSaveCookedPackageContext::FinishPlatform()
 		}
 
 		ICookedPackageWriter::FCommitPackageInfo Info;
-		Info.bSucceeded = bSuccessful;
+		if (bSuccessful)
+		{
+			Info.Status = IPackageWriter::ECommitStatus::Success;
+		}
+		else if (SavePackageResult.Result == ESavePackageResult::Timeout)
+		{
+			Info.Status = IPackageWriter::ECommitStatus::Canceled;
+		}
+		else
+		{
+			Info.Status = IPackageWriter::ECommitStatus::Error;
+		}
 		Info.PackageName = Package->GetFName();
 		PRAGMA_DISABLE_DEPRECATION_WARNINGS
 		Info.PackageGuid = AssetPackageData->PackageGuid;
@@ -4508,8 +4526,7 @@ void FSaveCookedPackageContext::FinishPlatform()
 		Generator.UpdateAssetRegistryPackageData(*Package, SavePackageResult, MoveTemp(*ArchiveCookContext.GetCookTagList()));
 	}
 
-	// In success or failure we want to mark the package as cooked, unless the failure was due to being referenced only
-	// by editor-only data; we may need to save it later when new content loads it through non editor-only references.
+	// If not retrying, mark the package as cooked, either successfully or with failure
 	bool bIsRetryErrorCode = IsRetryErrorCode(SavePackageResult.Result);
 	if (!bIsRetryErrorCode)
 	{
@@ -4624,7 +4641,7 @@ void FSaveCookedPackageContext::FinishPackage()
 			}
 		}
 	}
-	else
+	else if (bReferencedOnlyByEditorOnlyData)
 	{
 		COTFS.PackageTracker->UncookedEditorOnlyPackages.AddUnique(Package->GetFName());
 	}
@@ -9232,7 +9249,14 @@ uint32 UCookOnTheFlyServer::FullLoadAndSave(uint32& CookedPackageCount)
 						check(AssetPackageData);
 
 						ICookedPackageWriter::FCommitPackageInfo CommitInfo;
-						CommitInfo.bSucceeded = SaveResult.IsSuccessful();
+						if (SaveResult.IsSuccessful())
+						{
+							CommitInfo.Status = IPackageWriter::ECommitStatus::Success;
+						}
+						else
+						{
+							CommitInfo.Status = IPackageWriter::ECommitStatus::Error;
+						}
 						CommitInfo.PackageName = Package->GetFName();
 						PRAGMA_DISABLE_DEPRECATION_WARNINGS;
 						CommitInfo.PackageGuid = AssetPackageData->PackageGuid;
diff --git a/Engine/Source/Editor/UnrealEd/Private/Cooker/DiffPackageWriter.cpp b/Engine/Source/Editor/UnrealEd/Private/Cooker/DiffPackageWriter.cpp
--- a/Engine/Source/Editor/UnrealEd/Private/Cooker/DiffPackageWriter.cpp
+++ b/Engine/Source/Editor/UnrealEd/Private/Cooker/DiffPackageWriter.cpp
@@ -192,6 +192,11 @@ bool FDiffPackageWriter::IsAnotherSaveNeeded(FSavePackageResultStruct& PreviousR
 {
 	checkf(!Inner->IsAnotherSaveNeeded(PreviousResult, SaveArgs),
 		TEXT("DiffPackageWriter does not support an Inner that needs multiple saves."));
+	if (PreviousResult == ESavePackageResult::Timeout)
+	{
+		return false;
+	}
+
 	// When looking for deterministic cook issues, first serialize the package to memory and do a simple diff with the
 	// existing package. If the simple memory diff was not identical, collect callstacks for all Serialize calls and
 	// dump differences to log
@@ -205,7 +210,7 @@ bool FDiffPackageWriter::IsAnotherSaveNeeded(FSavePackageResultStruct& PreviousR
 			// The contract with the Inner is that Begin is paired with a single commit;
 			// send the old commit and the new begin
 			FCommitPackageInfo CommitInfo;
-			CommitInfo.bSucceeded = true;
+			CommitInfo.Status = IPackageWriter::ECommitStatus::Success;
 			CommitInfo.PackageName = BeginInfo.PackageName;
 			CommitInfo.WriteOptions = EWriteOptions::None;
 			Inner->CommitPackage(MoveTemp(CommitInfo));
@@ -317,6 +322,13 @@ bool FLinkerDiffPackageWriter::IsAnotherSaveNeeded(FSavePackageResultStruct& Pre
 {
 	checkf(!Inner->IsAnotherSaveNeeded(PreviousResult, SaveArgs),
 		TEXT("LinkerDiffPackageWriter does not support an Inner that needs multiple saves."));
+	if (PreviousResult == ESavePackageResult::Timeout)
+	{
+		OtherResult.LinkerSave.Reset();
+		PreviousResult.LinkerSave.Reset();
+		return false;
+	}
+
 	if (!bHasStartedSecondSave)
 	{
 		bHasStartedSecondSave = true;
@@ -328,7 +340,14 @@ bool FLinkerDiffPackageWriter::IsAnotherSaveNeeded(FSavePackageResultStruct& Pre
 		// The contract with the Inner is that every Begin is paired with a single commit.
 		// Send the old commit and the new begin.
 		IPackageWriter::FCommitPackageInfo CommitInfo;
-		CommitInfo.bSucceeded = OtherResult == ESavePackageResult::Success;
+		if (OtherResult == ESavePackageResult::Success)
+		{
+			CommitInfo.Status = IPackageWriter::ECommitStatus::Success;
+		}
+		else
+		{
+			CommitInfo.Status = IPackageWriter::ECommitStatus::Error;
+		}
 		CommitInfo.PackageName = BeginInfo.PackageName;
 		CommitInfo.WriteOptions = IPackageWriter::EWriteOptions::None;
 		Inner->CommitPackage(MoveTemp(CommitInfo));
diff --git a/Engine/Source/Editor/UnrealEd/Private/Cooker/LooseCookedPackageWriter.cpp b/Engine/Source/Editor/UnrealEd/Private/Cooker/LooseCookedPackageWriter.cpp
--- a/Engine/Source/Editor/UnrealEd/Private/Cooker/LooseCookedPackageWriter.cpp
+++ b/Engine/Source/Editor/UnrealEd/Private/Cooker/LooseCookedPackageWriter.cpp
@@ -58,11 +58,14 @@ void FLooseCookedPackageWriter::CommitPackageInternal(FPackageWriterRecords::FPa
 	const FCommitPackageInfo& Info)
 {
 	FRecord& Record = static_cast<FRecord&>(BaseRecord);
-	if (Info.bSucceeded)
+	if (Info.Status == IPackageWriter::ECommitStatus::Success)
 	{
 		AsyncSave(Record, Info);
 	}
-	UpdateManifest(Record);
+	if (Info.Status != IPackageWriter::ECommitStatus::Canceled)
+	{
+		UpdateManifest(Record);
+	}
 }
 
 void FLooseCookedPackageWriter::AsyncSave(FRecord& Record, const FCommitPackageInfo& Info)
diff --git a/Engine/Source/Editor/UnrealEd/Private/EditorDomain/EditorDomainUtils.cpp b/Engine/Source/Editor/UnrealEd/Private/EditorDomain/EditorDomainUtils.cpp
--- a/Engine/Source/Editor/UnrealEd/Private/EditorDomain/EditorDomainUtils.cpp
+++ b/Engine/Source/Editor/UnrealEd/Private/EditorDomain/EditorDomainUtils.cpp
@@ -1498,7 +1498,7 @@ protected:
 	{
 		// CommitPackage is called below with these options
 		check(Info.Attachments.Num() == 0);
-		check(Info.bSucceeded);
+		check(Info.Status == IPackageWriter::ECommitStatus::Success);
 		check(Info.WriteOptions == IPackageWriter::EWriteOptions::Write);
 		if (Record.AdditionalFiles.Num() > 0)
 		{
@@ -1700,7 +1700,7 @@ bool TrySavePackage(UPackage* Package)
 	}
 
 	ICookedPackageWriter::FCommitPackageInfo Info;
-	Info.bSucceeded = true;
+	Info.Status = IPackageWriter::ECommitStatus::Success;
 	Info.PackageName = Package->GetFName();
 	Info.WriteOptions = IPackageWriter::EWriteOptions::Write;
 	PackageWriter->CommitPackage(MoveTemp(Info));
diff --git a/Engine/Source/Runtime/Core/Public/Serialization/PackageWriter.h b/Engine/Source/Runtime/Core/Public/Serialization/PackageWriter.h
--- a/Engine/Source/Runtime/Core/Public/Serialization/PackageWriter.h
+++ b/Engine/Source/Runtime/Core/Public/Serialization/PackageWriter.h
@@ -73,12 +73,18 @@ public:
 		ComputeHash = 0x04,
 		SaveForDiff = 0x08,
 	};
+	enum class ECommitStatus
+	{
+		Success,
+		Canceled,
+		Error
+	};
 	struct FCommitPackageInfo
 	{
 		FName PackageName;
 		FGuid PackageGuid;
 		TArray<FCommitAttachmentInfo> Attachments;
-		bool bSucceeded = false;
+		ECommitStatus Status;
 		EWriteOptions WriteOptions;
 	};
 
diff --git a/Engine/Source/Runtime/CoreUObject/Internal/Serialization/PackageWriterToSharedBuffer.h b/Engine/Source/Runtime/CoreUObject/Internal/Serialization/PackageWriterToSharedBuffer.h
--- a/Engine/Source/Runtime/CoreUObject/Internal/Serialization/PackageWriterToSharedBuffer.h
+++ b/Engine/Source/Runtime/CoreUObject/Internal/Serialization/PackageWriterToSharedBuffer.h
@@ -57,7 +57,7 @@ public:
 
 		/** Always valid during CommitPackageInternal */
 		IPackageWriter::FBeginPackageInfo Begin;
-		/** Always valid during CommitPackageInternal if Info.bSucceeded */
+		/** Always valid during CommitPackageInternal if Info.Status == Success */
 		TArray<FWritePackage> Packages;
 		TArray<FBulkData> BulkDatas;
 		TArray<FAdditionalFile> AdditionalFiles;
diff --git a/Engine/Source/Runtime/CoreUObject/Private/Serialization/PackageWriterToSharedBuffer.cpp b/Engine/Source/Runtime/CoreUObject/Private/Serialization/PackageWriterToSharedBuffer.cpp
--- a/Engine/Source/Runtime/CoreUObject/Private/Serialization/PackageWriterToSharedBuffer.cpp
+++ b/Engine/Source/Runtime/CoreUObject/Private/Serialization/PackageWriterToSharedBuffer.cpp
@@ -104,9 +104,9 @@ TUniquePtr<FPackageWriterRecords::FPackage> FPackageWriterRecords::FindAndRemove
 
 void FPackageWriterRecords::ValidateCommit(FPackage& Record, const IPackageWriter::FCommitPackageInfo& Info) const
 {
-	checkf(Info.bSucceeded == false || Record.Packages.Num() > 0,
+	checkf(Info.Status != IPackageWriter::ECommitStatus::Success || Record.Packages.Num() > 0,
 		TEXT("IPackageWriter->WritePackageData must be called before Commit if the Package save was successful."));
-	checkf(Info.bSucceeded == false || Record.Packages.FindByPredicate([](const FPackageWriterRecords::FWritePackage& Package) { return Package.Info.MultiOutputIndex == 0; }),
+	checkf(Info.Status != IPackageWriter::ECommitStatus::Success || Record.Packages.FindByPredicate([](const FPackageWriterRecords::FWritePackage& Package) { return Package.Info.MultiOutputIndex == 0; }),
 		TEXT("SavePackage must provide output 0 when saving multioutput packages."));
 	uint8 HasBulkDataType[IPackageWriter::FBulkDataInfo::NumTypes]{};
 	for (FBulkData& BulkRecord : Record.BulkDatas)
