Avoid compiling in desktop-only Zen functionality for Stadia due to the use of PLATORM_UNIX instead of PLATFORM_LINUX.

diff --git a/Engine/Source/Developer/IoStoreUtilities/Private/ZenStoreHttpClient.cpp b/Engine/Source/Developer/IoStoreUtilities/Private/ZenStoreHttpClient.cpp
--- a/Engine/Source/Developer/IoStoreUtilities/Private/ZenStoreHttpClient.cpp
+++ b/Engine/Source/Developer/IoStoreUtilities/Private/ZenStoreHttpClient.cpp
@@ -10,7 +10,7 @@
 #include "Misc/CommandLine.h"
 #include "HAL/FileManager.h"
 
-#if PLATFORM_WINDOWS || PLATFORM_UNIX || PLATFORM_MAC
+#if PLATFORM_DESKTOP
 
 #include "Async/Async.h"
 #include "HAL/PlatformFileManager.h"
diff --git a/Engine/Source/Developer/Zen/Public/ZenGlobals.h b/Engine/Source/Developer/Zen/Public/ZenGlobals.h
--- a/Engine/Source/Developer/Zen/Public/ZenGlobals.h
+++ b/Engine/Source/Developer/Zen/Public/ZenGlobals.h
@@ -5,7 +5,7 @@
 #include "HAL/Platform.h"
 
 #ifndef UE_WITH_ZEN
-#	if PLATFORM_WINDOWS || PLATFORM_UNIX || PLATFORM_MAC
+#	if PLATFORM_DESKTOP
 #		define UE_WITH_ZEN 1
 #	else
 #		define UE_WITH_ZEN 0
