IoStore: Removed warning about missing packages not getting assigned a load order

diff --git a/Engine/Source/Developer/IoStoreUtilities/Private/IoStoreUtilities.cpp b/Engine/Source/Developer/IoStoreUtilities/Private/IoStoreUtilities.cpp
--- a/Engine/Source/Developer/IoStoreUtilities/Private/IoStoreUtilities.cpp
+++ b/Engine/Source/Developer/IoStoreUtilities/Private/IoStoreUtilities.cpp
@@ -1101,11 +1101,6 @@ static void AssignPackagesDiskOrder(
 	SortedPackages.Reserve(Packages.Num());
 	for (FLegacyCookedPackage* Package : Packages)
 	{
-		if (!Package->OptimizedPackage->GetExportBundleCount())
-		{
-			continue;
-		}
-
 		// Default to the fallback order map 
 		// Reverse the bundle load order for the fallback map (so that packages are considered before their imports)
 		const FFileOrderMap* UsedOrderMap = &FallbackOrderMap;
@@ -1265,7 +1260,7 @@ static void AssignPackagesDiskOrder(
 	}
 	UE_LOG(LogIoStore, Display, TEXT("Ordered %d packages using fallback bundle order"), AssignedPackages.Num() - LastAssignedCount);
 
-	check(AssignedPackages.Num() <= Packages.Num());
+	check(AssignedPackages.Num() == Packages.Num());
 	
 	if (bClusterByOrderFilePriority)
 	{
@@ -1291,18 +1286,6 @@ static void AssignPackagesDiskOrder(
 	}
 	
 	ClusterStatsCsv.Close();
-	if (AssignedPackages.Num() != Packages.Num())
-	{
-		UE_LOG(LogIoStore, Warning, TEXT(" %d/%d packages not assigned order"), Packages.Num()-AssignedPackages.Num(), Packages.Num());
-	
-		for (FLegacyCookedPackage* Package : Packages)
-		{
-			if (!AssignedPackages.Contains(Package))
-			{
-				Package->DiskLayoutOrder = LayoutIndex++;
-			}
-		}
-	}
 }
 
 static void CreateDiskLayout(
@@ -1610,6 +1593,7 @@ static void ParsePackageAssetsFromFiles(TArray<FLegacyCookedPackage*>& Packages,
 			}
 			else
 			{
+				UE_LOG(LogIoStore, Display, TEXT("Including package %s without a .uasset file. Excluded by PakFileRules?"), *Package->PackageName.ToString());
 				Package->OptimizedPackage = PackageStoreOptimizer.CreateMissingPackage(Package->PackageName);
 			}
 			check(Package->OptimizedPackage->GetId() == Package->GlobalPackageId);
