IoStore: Add support for staging and loading auto optional packages

diff --git a/Engine/Source/Developer/IoStoreUtilities/Internal/PackageStoreOptimizer.h b/Engine/Source/Developer/IoStoreUtilities/Internal/PackageStoreOptimizer.h
--- a/Engine/Source/Developer/IoStoreUtilities/Internal/PackageStoreOptimizer.h
+++ b/Engine/Source/Developer/IoStoreUtilities/Internal/PackageStoreOptimizer.h
@@ -154,6 +154,11 @@ public:
 		return ShaderMapHashes;
 	}
 
+	const bool HasEditorData() const
+	{
+		return (PackageFlags & PKG_FilterEditorOnly) == 0;
+	}
+
 private:
 	struct FExportGraphNode;
 
@@ -340,13 +345,8 @@ public:
 	void FinalizePackage(FPackageStorePackage* Package);
 	FIoBuffer CreatePackageBuffer(const FPackageStorePackage* Package, const FIoBuffer& CookedExportsBuffer, TArray<FFileRegion>* InOutFileRegions) const;
 	FPackageStoreEntryResource CreatePackageStoreEntry(const FPackageStorePackage* Package, const FPackageStorePackage* OptionalSegmentPackage) const;
-	enum EContainerHeaderInclusionFilter
-	{
-		IncludeNonOptionalSegments = 0x1,
-		IncludeOptionalSegments = 0x2,
-		IncludeAllSegments = IncludeNonOptionalSegments | IncludeOptionalSegments
-	};
-	FIoContainerHeader CreateContainerHeader(const FIoContainerId& ContainerId, TArrayView<const FPackageStoreEntryResource> PackageStoreEntries, EContainerHeaderInclusionFilter InclusionFilter) const;
+	FIoContainerHeader CreateContainerHeader(const FIoContainerId& ContainerId, TArrayView<const FPackageStoreEntryResource> PackageStoreEntries) const;
+	FIoContainerHeader CreateOptionalContainerHeader(const FIoContainerId& ContainerId, TArrayView<const FPackageStoreEntryResource> PackageStoreEntries) const;
 	IOSTOREUTILITIES_API FIoBuffer CreateScriptObjectsBuffer() const;
 	void LoadScriptObjectsBuffer(const FIoBuffer& ScriptObjectsBuffer);
 	void ProcessRedirects(const TMap<FPackageId, FPackageStorePackage*>& PackagesMap, bool bIsBuildingDLC) const;
@@ -410,6 +410,7 @@ private:
 	void FinalizePackageHeader(FPackageStorePackage* Package) const;
 	void FindScriptObjectsRecursive(FPackageObjectIndex OuterIndex, UObject* Object);
 	void FindScriptObjects();
+	FIoContainerHeader CreateContainerHeaderInternal(const FIoContainerId& ContainerId, TArrayView<const FPackageStoreEntryResource> PackageStoreEntries, bool bIsOptional) const;
 
 	TMap<FPackageObjectIndex, FScriptObjectData> ScriptObjectsMap;
 	uint64 TotalPackageCount = 0;
diff --git a/Engine/Source/Developer/IoStoreUtilities/Private/IoStoreUtilities.cpp b/Engine/Source/Developer/IoStoreUtilities/Private/IoStoreUtilities.cpp
--- a/Engine/Source/Developer/IoStoreUtilities/Private/IoStoreUtilities.cpp
+++ b/Engine/Source/Developer/IoStoreUtilities/Private/IoStoreUtilities.cpp
@@ -3407,7 +3407,13 @@ int32 CreateTarget(const FIoStoreArguments& Arguments, const FIoStoreWriterSetti
 					WriteOptions.FileName = TargetFile.DestinationPath;
 					if (TargetFile.ChunkType == EContainerChunkType::OptionalSegmentBulkData)
 					{
-						ContainerTarget->OptionalSegmentIoStoreWriter->Append(TargetFile.ChunkId, WriteRequestManager.Read(TargetFile), WriteOptions);
+						FIoChunkId ChunkId = TargetFile.ChunkId;
+						if (TargetFile.Package->OptimizedOptionalSegmentPackage && TargetFile.Package->OptimizedOptionalSegmentPackage->HasEditorData())
+						{
+							// Auto optional packages replace the non-optional part when the container is mounted
+							ChunkId = CreateIoChunkId(TargetFile.Package->GlobalPackageId.Value(), 0, EIoChunkType::BulkData);
+						}
+						ContainerTarget->OptionalSegmentIoStoreWriter->Append(ChunkId, WriteRequestManager.Read(TargetFile), WriteOptions);
 					}
 					else
 					{
@@ -3468,7 +3474,13 @@ int32 CreateTarget(const FIoStoreArguments& Arguments, const FIoStoreWriterSetti
 					{
 						check(ContainerTarget->OptionalSegmentIoStoreWriter);
 						check(TargetFile.Package->OptimizedOptionalSegmentPackage);
-						ContainerTarget->OptionalSegmentIoStoreWriter->Append(TargetFile.ChunkId, WriteRequestManager.Read(TargetFile), WriteOptions);
+						FIoChunkId ChunkId = TargetFile.ChunkId;
+						if (TargetFile.Package->OptimizedOptionalSegmentPackage->HasEditorData())
+						{
+							// Auto optional packages replace the non-optional part when the container is mounted
+							ChunkId = CreateIoChunkId(TargetFile.Package->GlobalPackageId.Value(), 0, EIoChunkType::ExportBundleData);
+						}
+						ContainerTarget->OptionalSegmentIoStoreWriter->Append(ChunkId, WriteRequestManager.Read(TargetFile), WriteOptions);
 					}
 					else
 					{
@@ -3494,12 +3506,12 @@ int32 CreateTarget(const FIoStoreArguments& Arguments, const FIoStoreWriterSetti
 					WriteOptions);
 			};
 
-			ContainerTarget->Header = PackageStoreOptimizer.CreateContainerHeader(ContainerTarget->ContainerId, PackageStoreEntries, FPackageStoreOptimizer::IncludeNonOptionalSegments);
+			ContainerTarget->Header = PackageStoreOptimizer.CreateContainerHeader(ContainerTarget->ContainerId, PackageStoreEntries);
 			WriteContainerHeaderChunk(ContainerTarget->Header, ContainerTarget->IoStoreWriter.Get());
 
 			if (ContainerTarget->OptionalSegmentIoStoreWriter)
 			{
-				ContainerTarget->OptionalSegmentHeader = PackageStoreOptimizer.CreateContainerHeader(ContainerTarget->ContainerId, PackageStoreEntries, FPackageStoreOptimizer::IncludeOptionalSegments);
+				ContainerTarget->OptionalSegmentHeader = PackageStoreOptimizer.CreateOptionalContainerHeader(ContainerTarget->ContainerId, PackageStoreEntries);
 				WriteContainerHeaderChunk(ContainerTarget->OptionalSegmentHeader, ContainerTarget->OptionalSegmentIoStoreWriter.Get());
 			}
 		}
diff --git a/Engine/Source/Developer/IoStoreUtilities/Private/PackageStoreOptimizer.cpp b/Engine/Source/Developer/IoStoreUtilities/Private/PackageStoreOptimizer.cpp
--- a/Engine/Source/Developer/IoStoreUtilities/Private/PackageStoreOptimizer.cpp
+++ b/Engine/Source/Developer/IoStoreUtilities/Private/PackageStoreOptimizer.cpp
@@ -1809,6 +1809,13 @@ FPackageStoreEntryResource FPackageStoreOptimizer::CreatePackageStoreEntry(const
 {
 	FPackageStoreEntryResource Result;
 	Result.Flags = Package->bIsRedirected ? EPackageStoreEntryFlags::Redirected : EPackageStoreEntryFlags::None;
+
+	if (OptionalSegmentPackage && OptionalSegmentPackage->HasEditorData())
+	{
+		// AutoOptional packages are saved with editor data included
+		Result.Flags |= EPackageStoreEntryFlags::AutoOptional;
+	}
+	
 	Result.PackageName = Package->Name;
 	Result.SourcePackageName = Package->SourceName;
 	Result.Region = FName(*Package->Region);
@@ -1828,27 +1835,35 @@ FPackageStoreEntryResource FPackageStoreOptimizer::CreatePackageStoreEntry(const
 	return Result;
 }
 
-FIoContainerHeader FPackageStoreOptimizer::CreateContainerHeader(const FIoContainerId& ContainerId, TArrayView<const FPackageStoreEntryResource> PackageStoreEntries, EContainerHeaderInclusionFilter InclusionFilter) const
+FIoContainerHeader FPackageStoreOptimizer::CreateContainerHeaderInternal(const FIoContainerId& ContainerId, TArrayView<const FPackageStoreEntryResource> PackageStoreEntries, bool bIsOptional) const
 {
-	const bool bIncludeNonOptionalSegments = (InclusionFilter & IncludeNonOptionalSegments) > 0;
-	const bool bIncludeOptionalSegments = (InclusionFilter & IncludeOptionalSegments) > 0;
-
 	FIoContainerHeader Header;
 	Header.ContainerId = ContainerId;
 	
-	const int32 NonOptionalSegmentStoreEntriesCount = bIncludeNonOptionalSegments ? PackageStoreEntries.Num() : 0;
-
+	int32 NonOptionalSegmentStoreEntriesCount = 0;
 	int32 OptionalSegmentStoreEntriesCount = 0;
-	if (bIncludeOptionalSegments)
+	if (bIsOptional)
 	{
 		for (const FPackageStoreEntryResource& Entry : PackageStoreEntries)
 		{
 			if (Entry.OptionalSegmentExportInfo.ExportCount)
 			{
-				++OptionalSegmentStoreEntriesCount;
+				if (Entry.IsAutoOptional())
+				{
+					// Auto optional packages fully replace the non-optional segment
+					++NonOptionalSegmentStoreEntriesCount;
+				}
+				else
+				{
+					++OptionalSegmentStoreEntriesCount;
+				}
 			}
 		}
 	}
+	else
+	{
+		NonOptionalSegmentStoreEntriesCount = PackageStoreEntries.Num();
+	}
 
 	struct FStoreEntriesWriter
 	{
@@ -1907,70 +1922,85 @@ FIoContainerHeader FPackageStoreOptimizer::CreateContainerHeader(const FIoContai
 	FPackageStoreNameMapBuilder RedirectsNameMapBuilder;
 	RedirectsNameMapBuilder.SetNameMapType(FMappedName::EType::Container);
 	TSet<FPackageId> AllLocalizedPackages;
-	for (const FPackageStoreEntryResource* Entry : SortedPackageStoreEntries)
+	if (bIsOptional)
 	{
-		if (bIncludeOptionalSegments && Entry->OptionalSegmentExportInfo.ExportCount)
+		for (const FPackageStoreEntryResource* Entry : SortedPackageStoreEntries)
 		{
-			Header.OptionalSegmentPackageIds.Add(Entry->GetPackageId());
-
-			FPackageStoreExportInfo OptionalSegmentExportInfo = Entry->OptionalSegmentExportInfo;
-			OptionalSegmentStoreEntriesWriter.StoreTocArchive << OptionalSegmentExportInfo;
-
-			const TArray<FPackageId>& OptionalSegmentImportedPackageIds = Entry->OptionalSegmentImportedPackageIds;
-			SerializePackageEntryCArrayHeader(OptionalSegmentStoreEntriesWriter, OptionalSegmentImportedPackageIds.Num());
-			for (FPackageId OptionalSegmentImportedPackageId : OptionalSegmentImportedPackageIds)
+			if (Entry->OptionalSegmentExportInfo.ExportCount)
 			{
-				check(OptionalSegmentImportedPackageId.IsValid());
-				OptionalSegmentStoreEntriesWriter.StoreDataArchive << OptionalSegmentImportedPackageId;
-			}
+				FStoreEntriesWriter* TargetEntriesWriter;
+				if (Entry->IsAutoOptional())
+				{
+					Header.PackageIds.Add(Entry->GetPackageId());
+					TargetEntriesWriter = &StoreEntriesWriter;
+				}
+				else
+				{
+					Header.OptionalSegmentPackageIds.Add(Entry->GetPackageId());
+					TargetEntriesWriter = &OptionalSegmentStoreEntriesWriter;
+				}
 
-			// ShaderMapHashes is N/A for optional segments
-			SerializePackageEntryCArrayHeader(OptionalSegmentStoreEntriesWriter, 0);
-		}
+				// OptionalStoreEntry
+				FPackageStoreExportInfo OptionalSegmentExportInfo = Entry->OptionalSegmentExportInfo;
+				TargetEntriesWriter->StoreTocArchive << OptionalSegmentExportInfo;
 
-		if (!bIncludeNonOptionalSegments)
-		{
-			continue;
-		}
+				// OptionalImportedPackages
+				const TArray<FPackageId>& OptionalSegmentImportedPackageIds = Entry->OptionalSegmentImportedPackageIds;
+				SerializePackageEntryCArrayHeader(*TargetEntriesWriter, OptionalSegmentImportedPackageIds.Num());
+				for (FPackageId OptionalSegmentImportedPackageId : OptionalSegmentImportedPackageIds)
+				{
+					check(OptionalSegmentImportedPackageId.IsValid());
+					TargetEntriesWriter->StoreDataArchive << OptionalSegmentImportedPackageId;
+				}
 
-		Header.PackageIds.Add(Entry->GetPackageId());
-		if (Entry->IsRedirected())
+				// ShaderMapHashes is N/A for optional segments
+				SerializePackageEntryCArrayHeader(*TargetEntriesWriter, 0);
+			}
+		}
+	}
+	else
+	{
+		for (const FPackageStoreEntryResource* Entry : SortedPackageStoreEntries)
 		{
-			RedirectsNameMapBuilder.MarkNameAsReferenced(Entry->GetSourcePackageName());
-			FMappedName MappedSourcePackageName = RedirectsNameMapBuilder.MapName(Entry->GetSourcePackageName());
-			if (!Entry->Region.IsNone())
+			Header.PackageIds.Add(Entry->GetPackageId());
+			if (Entry->IsRedirected())
 			{
-				if (!AllLocalizedPackages.Contains(Entry->GetSourcePackageId()))
+				RedirectsNameMapBuilder.MarkNameAsReferenced(Entry->GetSourcePackageName());
+				FMappedName MappedSourcePackageName = RedirectsNameMapBuilder.MapName(Entry->GetSourcePackageName());
+				if (!Entry->Region.IsNone())
+				{
+					if (!AllLocalizedPackages.Contains(Entry->GetSourcePackageId()))
+					{
+						Header.LocalizedPackages.Add({ Entry->GetSourcePackageId(), MappedSourcePackageName });
+						AllLocalizedPackages.Add(Entry->GetSourcePackageId());
+					}
+				}
+				else
 				{
-					Header.LocalizedPackages.Add({ Entry->GetSourcePackageId(), MappedSourcePackageName });
-					AllLocalizedPackages.Add(Entry->GetSourcePackageId());
+					Header.PackageRedirects.Add({ Entry->GetSourcePackageId(), Entry->GetPackageId(), MappedSourcePackageName });
 				}
 			}
-			else
+
+			// StoreEntries
+			FPackageStoreExportInfo ExportInfo = Entry->ExportInfo;
+			StoreEntriesWriter.StoreTocArchive << ExportInfo;
+
+			// ImportedPackages
+			const TArray<FPackageId>& ImportedPackageIds = Entry->ImportedPackageIds;
+			SerializePackageEntryCArrayHeader(StoreEntriesWriter, ImportedPackageIds.Num());
+			for (FPackageId ImportedPackageId : ImportedPackageIds)
 			{
-				Header.PackageRedirects.Add({ Entry->GetSourcePackageId(), Entry->GetPackageId(), MappedSourcePackageName });
+				check(ImportedPackageId.IsValid());
+				StoreEntriesWriter.StoreDataArchive << ImportedPackageId;
+			}
+
+			// ShaderMapHashes
+			const TArray<FSHAHash>& ShaderMapHashes = Entry->ShaderMapHashes;
+			SerializePackageEntryCArrayHeader(StoreEntriesWriter, ShaderMapHashes.Num());
+			for (const FSHAHash& ShaderMapHash : ShaderMapHashes)
+			{
+				StoreEntriesWriter.StoreDataArchive << const_cast<FSHAHash&>(ShaderMapHash);
 			}
-		}
-		
-		// StoreEntries
-		FPackageStoreExportInfo ExportInfo = Entry->ExportInfo;
-		StoreEntriesWriter.StoreTocArchive << ExportInfo;
-		
-		// ImportedPackages
-		const TArray<FPackageId>& ImportedPackageIds = Entry->ImportedPackageIds;
-		SerializePackageEntryCArrayHeader(StoreEntriesWriter, ImportedPackageIds.Num());
-		for (FPackageId ImportedPackageId : ImportedPackageIds)
-		{
-			check(ImportedPackageId.IsValid());
-			StoreEntriesWriter.StoreDataArchive << ImportedPackageId;
-		}
-		
-		// ShaderMapHashes
-		const TArray<FSHAHash>& ShaderMapHashes = Entry->ShaderMapHashes;
-		SerializePackageEntryCArrayHeader(StoreEntriesWriter, ShaderMapHashes.Num());
-		for (const FSHAHash& ShaderMapHash : ShaderMapHashes)
-		{
-			StoreEntriesWriter.StoreDataArchive << const_cast<FSHAHash&>(ShaderMapHash);
 		}
 	}
 	Header.RedirectsNameMap = RedirectsNameMapBuilder.GetNameMap();
@@ -1980,3 +2010,13 @@ FIoContainerHeader FPackageStoreOptimizer::CreateContainerHeader(const FIoContai
 
 	return Header;
 }
+
+FIoContainerHeader FPackageStoreOptimizer::CreateContainerHeader(const FIoContainerId& ContainerId, TArrayView<const FPackageStoreEntryResource> PackageStoreEntries) const
+{
+	return CreateContainerHeaderInternal(ContainerId, PackageStoreEntries, false);
+}
+
+FIoContainerHeader FPackageStoreOptimizer::CreateOptionalContainerHeader(const FIoContainerId& ContainerId, TArrayView<const FPackageStoreEntryResource> PackageStoreEntries) const
+{
+	return CreateContainerHeaderInternal(ContainerId, PackageStoreEntries, true);
+}
\ No newline at end of file
diff --git a/Engine/Source/Developer/IoStoreUtilities/Private/ZenStoreWriter.cpp b/Engine/Source/Developer/IoStoreUtilities/Private/ZenStoreWriter.cpp
--- a/Engine/Source/Developer/IoStoreUtilities/Private/ZenStoreWriter.cpp
+++ b/Engine/Source/Developer/IoStoreUtilities/Private/ZenStoreWriter.cpp
@@ -1103,7 +1103,7 @@ void FZenStoreWriter::CreateProjectMetaData(FCbPackage& Pkg, FCbWriter& PackageO
 			FIoBuffer HeaderBuffer;
 
 			{
-				FIoContainerHeader Header = PackageStoreOptimizer->CreateContainerHeader(ContainerId, PackageStoreEntries, FPackageStoreOptimizer::IncludeAllSegments);
+				FIoContainerHeader Header = PackageStoreOptimizer->CreateContainerHeader(ContainerId, PackageStoreEntries);
 				FLargeMemoryWriter HeaderAr(0, true);
 				HeaderAr << Header;
 				int64 DataSize = HeaderAr.TotalSize();
diff --git a/Engine/Source/Runtime/Core/Public/IO/PackageStore.h b/Engine/Source/Runtime/Core/Public/IO/PackageStore.h
--- a/Engine/Source/Runtime/Core/Public/IO/PackageStore.h
+++ b/Engine/Source/Runtime/Core/Public/IO/PackageStore.h
@@ -60,6 +60,7 @@ enum class EPackageStoreEntryFlags : uint32
 {
 	None		= 0,
 	Redirected	= 0x01,
+	AutoOptional= 0x02,
 };
 ENUM_CLASS_FLAGS(EPackageStoreEntryFlags);
 
@@ -114,6 +115,12 @@ struct FPackageStoreEntryResource
 		return EnumHasAnyFlags(Flags, EPackageStoreEntryFlags::Redirected); 
 	}
 
+	/** Returns whether this package was saved as auto optional */
+	bool IsAutoOptional() const
+	{
+		return EnumHasAnyFlags(Flags, EPackageStoreEntryFlags::AutoOptional);
+	}
+
 	CORE_API friend FArchive& operator<<(FArchive& Ar, FPackageStoreEntryResource& PackageStoreEntry);
 	
 	CORE_API friend FCbWriter& operator<<(FCbWriter& Writer, const FPackageStoreEntryResource& PackageStoreEntry);
diff --git a/Engine/Source/Runtime/CoreUObject/Private/Serialization/AsyncLoading2.cpp b/Engine/Source/Runtime/CoreUObject/Private/Serialization/AsyncLoading2.cpp
--- a/Engine/Source/Runtime/CoreUObject/Private/Serialization/AsyncLoading2.cpp
+++ b/Engine/Source/Runtime/CoreUObject/Private/Serialization/AsyncLoading2.cpp
@@ -4016,7 +4016,9 @@ EAsyncPackageState::Type FAsyncPackage2::Event_ProcessExportBundle(FAsyncLoading
 
 	auto FilterExport = [](const EExportFilterFlags FilterFlags) -> bool
 	{
-#if UE_SERVER
+#if WITH_EDITOR
+		return false;
+#elif UE_SERVER
 		return !!(static_cast<uint32>(FilterFlags) & static_cast<uint32>(EExportFilterFlags::NotForServer));
 #elif !WITH_SERVER_CODE
 		return !!(static_cast<uint32>(FilterFlags) & static_cast<uint32>(EExportFilterFlags::NotForClient));
