warn if the previous container files could not be found when creating a new release version - otherwise the patch delta can be significantly higher on some platforms.

#rnx

diff --git a/Engine/Source/Developer/IoStoreUtilities/Private/IoStoreUtilities.cpp b/Engine/Source/Developer/IoStoreUtilities/Private/IoStoreUtilities.cpp
--- a/Engine/Source/Developer/IoStoreUtilities/Private/IoStoreUtilities.cpp
+++ b/Engine/Source/Developer/IoStoreUtilities/Private/IoStoreUtilities.cpp
@@ -6333,6 +6333,8 @@ bool ParseContainerGenerationArguments(FIoStoreArguments& Arguments, FIoStoreWri
 					PatchSourceContainerFile = PatchSourceContainersDirectory / PatchSourceContainerFile;
 					FPaths::NormalizeFilename(PatchSourceContainerFile);
 				}
+
+				UE_CLOG(ContainerSpec.PatchSourceContainerFiles.Num() == 0, LogIoStore, Warning, TEXT("%s - No patch source files found. This may affect the patch delta size on some platforms because the container file order may change."), *PatchSourceWildcard);
 			}
 
 			ContainerSpec.bGenerateDiffPatch = FParse::Param(*Command, TEXT("GenerateDiffPatch"));
