Log AR load errors (turns out, not actually logged when it's missing on disk!)

diff --git a/Engine/Source/Developer/IoStoreUtilities/Private/IoStoreUtilities.cpp b/Engine/Source/Developer/IoStoreUtilities/Private/IoStoreUtilities.cpp
--- a/Engine/Source/Developer/IoStoreUtilities/Private/IoStoreUtilities.cpp
+++ b/Engine/Source/Developer/IoStoreUtilities/Private/IoStoreUtilities.cpp
@@ -3152,7 +3152,8 @@ int32 DoAssetRegistryWritebackAfterStage(const FString& InAssetRegistryFileName,
 	FAssetRegistryState AssetRegistry;
 	if (LoadAssetRegistry(InAssetRegistryFileName, AssetRegistry) == false)
 	{
-		return 1; // already logged
+		UE_LOG(LogIoStore, Error, TEXT("Unabled to open source asset registry: %s"), *InAssetRegistryFileName);
+		return 1;
 	}
 
 	// Grab all containers in the directory
