IoStore: Downgraded failure to obtain package name from file name to a warning

diff --git a/Engine/Source/Developer/IoStoreUtilities/Private/IoStoreUtilities.cpp b/Engine/Source/Developer/IoStoreUtilities/Private/IoStoreUtilities.cpp
--- a/Engine/Source/Developer/IoStoreUtilities/Private/IoStoreUtilities.cpp
+++ b/Engine/Source/Developer/IoStoreUtilities/Private/IoStoreUtilities.cpp
@@ -1436,7 +1436,7 @@ FContainerTargetSpec* AddContainer(
 	return ContainerTargetSpec;
 }
 
-FLegacyCookedPackage* FindOrAddPackage(
+FLegacyCookedPackage& FindOrAddPackage(
 	const FIoStoreArguments& Arguments,
 	const FName& PackageName,
 	TArray<FLegacyCookedPackage*>& Packages,
@@ -1465,10 +1465,10 @@ FLegacyCookedPackage* FindOrAddPackage(
 		PackageIdMap.Add(PackageId, Package);
 	}
 
-	return Package;
+	return *Package;
 }
 
-FLegacyCookedPackage* FindOrAddPackage(
+FLegacyCookedPackage* FindOrAddPackageFromFileName(
 	const FIoStoreArguments& Arguments,
 	const TCHAR* FileName,
 	TArray<FLegacyCookedPackage*>& Packages,
@@ -1478,11 +1478,10 @@ FLegacyCookedPackage* FindOrAddPackage(
 	FName PackageName = Arguments.PackageStore->GetPackageNameFromFileName(FileName);
 	if (PackageName.IsNone())
 	{
-		UE_LOG(LogIoStore, Fatal, TEXT("Failed to obtain package name from file name '%s'"), FileName);
 		return nullptr;
 	}
 
-	return FindOrAddPackage(Arguments, PackageName, Packages, PackageNameMap, PackageIdMap);
+	return &FindOrAddPackage(Arguments, PackageName, Packages, PackageNameMap, PackageIdMap);
 }
 
 static void ParsePackageAssetsFromFiles(TArray<FLegacyCookedPackage*>& Packages, const FPackageStoreOptimizer& PackageStoreOptimizer)
@@ -2179,7 +2178,12 @@ void InitializeContainerTargetsAndPackages(
 		}
 		else
 		{
-			OutTargetFile.Package = FindOrAddPackage(Arguments, *SourceFile.NormalizedPath, Packages, PackageNameMap, PackageIdMap);
+			OutTargetFile.Package = FindOrAddPackageFromFileName(Arguments, *SourceFile.NormalizedPath, Packages, PackageNameMap, PackageIdMap);
+			if (!OutTargetFile.Package)
+			{
+				UE_LOG(LogIoStore, Warning, TEXT("Failed to obtain package name from file name '%s'"), *SourceFile.NormalizedPath);
+				return false;
+			}
 
 			switch (CookedFileStatData->FileType)
 			{
@@ -2293,7 +2297,7 @@ void InitializeContainerTargetsAndPackages(
 			return false;
 		}
 
-		OutTargetFile.Package = FindOrAddPackage(Arguments, PackageName, Packages, PackageNameMap, PackageIdMap);
+		OutTargetFile.Package = &FindOrAddPackage(Arguments, PackageName, Packages, PackageNameMap, PackageIdMap);
 		OutTargetFile.Package->PackageStoreEntry = PackageStore.GetPackageStoreEntry(OutTargetFile.Package->GlobalPackageId);
 		if (!OutTargetFile.Package->PackageStoreEntry)
 		{
