remove warning about missing utoc files when generating patches.

- the default ordering is deterministic based on which packages are being imported so the warning is not appropriate for all circumstances.

#rnx

diff --git a/Engine/Source/Developer/IoStoreUtilities/Private/IoStoreUtilities.cpp b/Engine/Source/Developer/IoStoreUtilities/Private/IoStoreUtilities.cpp
--- a/Engine/Source/Developer/IoStoreUtilities/Private/IoStoreUtilities.cpp
+++ b/Engine/Source/Developer/IoStoreUtilities/Private/IoStoreUtilities.cpp
@@ -6338,8 +6338,6 @@ bool ParseContainerGenerationArguments(FIoStoreArguments& Arguments, FIoStoreWri
 					PatchSourceContainerFile = PatchSourceContainersDirectory / PatchSourceContainerFile;
 					FPaths::NormalizeFilename(PatchSourceContainerFile);
 				}
-
-				UE_CLOG(ContainerSpec.PatchSourceContainerFiles.Num() == 0, LogIoStore, Warning, TEXT("%s - No patch source files found. This may affect the patch delta size on some platforms because the container file order may change."), *PatchSourceWildcard);
 			}
 
 			ContainerSpec.bGenerateDiffPatch = FParse::Param(*Command, TEXT("GenerateDiffPatch"));
