Make sure we wait for any async commits in ZenStoreWriter to complete before sending EndCook.
Previous implementation only waited for the commit to be picked up by the worker thread.
Simplified CommitQueue event logic.

diff --git a/Engine/Source/Developer/IoStoreUtilities/Private/ZenStoreWriter.cpp b/Engine/Source/Developer/IoStoreUtilities/Private/ZenStoreWriter.cpp
--- a/Engine/Source/Developer/IoStoreUtilities/Private/ZenStoreWriter.cpp
+++ b/Engine/Source/Developer/IoStoreUtilities/Private/ZenStoreWriter.cpp
@@ -68,18 +68,22 @@ class FZenStoreWriter::FCommitQueue
 {
 public:
 	FCommitQueue()
-		: NewCommitEvent(EEventMode::AutoReset)
-		, QueueEmptyEvent(EEventMode::ManualReset) {}
+		: NewCommitEvent(EEventMode::AutoReset) {}
 	
 	void Enqueue(FZenCommitInfo&& Commit)
 	{
+		bool TriggerEvent = false;
 		{
 			FScopeLock _(&QueueCriticalSection);
+			TriggerEvent = Queue.IsEmpty();
 			Queue.Enqueue(MoveTemp(Commit));
-			QueueNum++;
+			
 		}
 
-		NewCommitEvent->Trigger();
+		if (TriggerEvent)
+		{
+			NewCommitEvent->Trigger();
+		}
 	}
 
 	bool BlockAndDequeue(FZenCommitInfo& OutCommit)
@@ -90,10 +94,8 @@ public:
 				FScopeLock _(&QueueCriticalSection);
 				if (Queue.Dequeue(OutCommit))
 				{
-					QueueNum--;
 					return true;
 				}
-				QueueEmptyEvent->Trigger();
 			}
 
 			if (bCompleteAdding)
@@ -116,30 +118,10 @@ public:
 		bCompleteAdding = false;
 	}
 
-	void WaitUntilEmpty()
-	{
-		bool bWait = false;
-		{
-			FScopeLock _(&QueueCriticalSection);
-			if (QueueNum > 0)
-			{
-				QueueEmptyEvent->Reset();
-				bWait = true;
-			}
-		}
-
-		if (bWait)
-		{
-			QueueEmptyEvent->Wait();
-		}
-	}
-
 private:
 	FEventRef					NewCommitEvent;
-	FEventRef					QueueEmptyEvent;
 	FCriticalSection			QueueCriticalSection;
 	TQueue<FZenCommitInfo>		Queue;
-	int32						QueueNum = 0;
 	TAtomic<bool>				bCompleteAdding{false};
 };
 
@@ -537,10 +519,10 @@ void FZenStoreWriter::BeginCook()
 void FZenStoreWriter::EndCook()
 {
 	UE_LOG(LogZenStoreWriter, Display, TEXT("Flushing..."));
-	CommitQueue->WaitUntilEmpty();
 	
 	CommitQueue->CompleteAdding();
-	
+	CommitThread.Wait();
+
 	FCbPackage Pkg;
 	FCbWriter PackageObj;
 	
