Pass along path to project file when zen server is running locally so we can detect when it is removed and garbage-collect any projectstore entries on the zen side.
Allow Cook to remote Zen, if we detect that we are targetting a zenservice on a different machine we change the local file references to attachment based references.

diff --git a/Engine/Source/Developer/IoStoreUtilities/Internal/ZenStoreHttpClient.h b/Engine/Source/Developer/IoStoreUtilities/Internal/ZenStoreHttpClient.h
--- a/Engine/Source/Developer/IoStoreUtilities/Internal/ZenStoreHttpClient.h
+++ b/Engine/Source/Developer/IoStoreUtilities/Internal/ZenStoreHttpClient.h
@@ -33,7 +33,8 @@ public:
 	~FZenStoreHttpClient();
 
 	bool TryCreateProject(FStringView InProjectId, FStringView InOplogId, FStringView ServerRoot, 
-					FStringView EngineRoot, FStringView ProjectRoot);
+					FStringView EngineRoot, FStringView ProjectRoot,
+					FStringView ProjectFilePath);
 	bool TryCreateOplog(FStringView InProjectId, FStringView InOplogId, bool bFullBuild);
 
 	void InitializeReadOnly(FStringView InProjectId, FStringView InOplogId);
diff --git a/Engine/Source/Developer/IoStoreUtilities/Private/ZenFileSystemManifest.h b/Engine/Source/Developer/IoStoreUtilities/Private/ZenFileSystemManifest.h
--- a/Engine/Source/Developer/IoStoreUtilities/Private/ZenFileSystemManifest.h
+++ b/Engine/Source/Developer/IoStoreUtilities/Private/ZenFileSystemManifest.h
@@ -39,6 +39,11 @@ public:
 		return Entries.Num();
 	}
 
+	const FString& ServerRootPath() const
+	{
+		return ServerRoot;
+	}
+
 private:
 	static void GetExtensionDirs(TArray<FString>& OutExtensionDirs, const TCHAR* BaseDir, const TCHAR* SubDir, const TArray<FString>& PlatformDirectoryNames);
 
diff --git a/Engine/Source/Developer/IoStoreUtilities/Private/ZenStoreHttpClient.cpp b/Engine/Source/Developer/IoStoreUtilities/Private/ZenStoreHttpClient.cpp
--- a/Engine/Source/Developer/IoStoreUtilities/Private/ZenStoreHttpClient.cpp
+++ b/Engine/Source/Developer/IoStoreUtilities/Private/ZenStoreHttpClient.cpp
@@ -58,7 +58,8 @@ FZenStoreHttpClient::TryCreateProject(FStringView InProjectId,
 	FStringView InOplogId, 
 	FStringView ServerRoot,
 	FStringView EngineRoot,
-	FStringView ProjectRoot)
+	FStringView ProjectRoot,
+	FStringView ProjectFilePath)
 {
 	TRACE_CPUPROFILER_EVENT_SCOPE(ZenStoreHttp_Initialize);
 
@@ -93,6 +94,7 @@ FZenStoreHttpClient::TryCreateProject(FStringView InProjectId,
 			ProjInfo << "root" << ServerRoot;
 			ProjInfo << "engine" << EngineRoot;
 			ProjInfo << "project" << ProjectRoot;
+			ProjInfo << "projectfile" << ProjectFilePath;
 			ProjInfo.EndObject();
 
 			Res = Request->PerformBlockingPost(ProjectUri, ProjInfo.Save().AsObject());
@@ -631,7 +633,8 @@ FZenStoreHttpClient::~FZenStoreHttpClient()
 }
 
 bool FZenStoreHttpClient::TryCreateProject(FStringView InProjectId, FStringView InOplogId, FStringView ServerRoot,
-	FStringView EngineRoot,	FStringView ProjectRoot)
+	FStringView EngineRoot,	FStringView ProjectRoot,
+	FStringView ProjectFilePath)
 {
 	return false;
 }
diff --git a/Engine/Source/Developer/IoStoreUtilities/Private/ZenStoreWriter.cpp b/Engine/Source/Developer/IoStoreUtilities/Private/ZenStoreWriter.cpp
--- a/Engine/Source/Developer/IoStoreUtilities/Private/ZenStoreWriter.cpp
+++ b/Engine/Source/Developer/IoStoreUtilities/Private/ZenStoreWriter.cpp
@@ -166,18 +166,25 @@ FZenStoreWriter::FZenStoreWriter(
 	
 	HttpClient = MakeUnique<UE::FZenStoreHttpClient>();
 
+#if UE_WITH_ZEN
+	IsLocalConnection = HttpClient->GetZenServiceInstance().IsServiceRunningLocally();
+#endif
+
 	FString RootDir = FPaths::RootDir();
 	FString EngineDir = FPaths::EngineDir();
 	FPaths::NormalizeDirectoryName(EngineDir);
 	FString ProjectDir = FPaths::ProjectDir();
 	FPaths::NormalizeDirectoryName(ProjectDir);
+	FString ProjectPath = FPaths::GetProjectFilePath();
+	FPaths::NormalizeFilename(ProjectPath);
 
 	IPlatformFile& PlatformFile = FPlatformFileManager::Get().GetPlatformFile();
 	FString AbsServerRoot = PlatformFile.ConvertToAbsolutePathForExternalAppForRead(*RootDir);
 	FString AbsEngineDir = PlatformFile.ConvertToAbsolutePathForExternalAppForRead(*EngineDir);
 	FString AbsProjectDir = PlatformFile.ConvertToAbsolutePathForExternalAppForRead(*ProjectDir);
+	FString ProjectFilePath = PlatformFile.ConvertToAbsolutePathForExternalAppForRead(*ProjectPath);
 
-	HttpClient->TryCreateProject(ProjectId, OplogId, AbsServerRoot, AbsEngineDir, AbsProjectDir);
+	HttpClient->TryCreateProject(ProjectId, OplogId, AbsServerRoot, AbsEngineDir, AbsProjectDir, IsLocalConnection ? ProjectFilePath : FStringView());
 
 	PackageStoreOptimizer->Initialize();
 
@@ -1045,12 +1052,34 @@ void FZenStoreWriter::CreateProjectMetaData(FCbPackage& Pkg, FCbWriter& PackageO
 			{
 				FCbObjectId FileOid = ToObjectId(NewEntry.FileChunkId);
 
-				PackageObj.BeginObject();
-				PackageObj << "id" << FileOid;
-				PackageObj << "data" << FIoHash::Zero;
-				PackageObj << "serverpath" << NewEntry.ServerPath;
-				PackageObj << "clientpath" << NewEntry.ClientPath;
-				PackageObj.EndObject();
+				if (IsLocalConnection)
+				{
+					PackageObj.BeginObject();
+					PackageObj << "id" << FileOid;
+					PackageObj << "data" << FIoHash::Zero;
+					PackageObj << "serverpath" << NewEntry.ServerPath;
+					PackageObj << "clientpath" << NewEntry.ClientPath;
+					PackageObj.EndObject();
+				}
+				else
+				{
+					const FString AbsPath = ZenFileSystemManifest->ServerRootPath() / NewEntry.ServerPath;
+					TArray<uint8> FileBuffer;
+					FFileHelper::LoadFileToArray(FileBuffer, *AbsPath);
+					if (FileBuffer.Num())
+					{
+						FCbAttachment FileAttachment = CreateAttachment(FIoBuffer(FIoBuffer::Clone, FileBuffer.GetData(), FileBuffer.Num()));
+
+						PackageObj.BeginObject();
+						PackageObj << "id" << FileOid;
+						PackageObj << "data" << FileAttachment;
+						PackageObj << "serverpath" << NewEntry.ServerPath;
+						PackageObj << "clientpath" << NewEntry.ClientPath;
+						PackageObj.EndObject();
+
+						Pkg.AddAttachment(FileAttachment);
+					}
+				}
 			}
 
 			PackageObj.EndArray();
diff --git a/Engine/Source/Developer/IoStoreUtilities/Public/ZenStoreWriter.h b/Engine/Source/Developer/IoStoreUtilities/Public/ZenStoreWriter.h
--- a/Engine/Source/Developer/IoStoreUtilities/Public/ZenStoreWriter.h
+++ b/Engine/Source/Developer/IoStoreUtilities/Public/ZenStoreWriter.h
@@ -174,6 +174,7 @@ private:
 	FCriticalSection								PackagesCriticalSection;
 	TMap<FName, TUniquePtr<FPendingPackageState>>	PendingPackages;
 	TUniquePtr<UE::FZenStoreHttpClient>	HttpClient;
+	bool IsLocalConnection = true;
 
 	const ITargetPlatform&				TargetPlatform;
 	const FName							TargetPlatformFName;
diff --git a/Engine/Source/Editor/UnrealEd/Private/TargetDomain/TargetDomainUtils.cpp b/Engine/Source/Editor/UnrealEd/Private/TargetDomain/TargetDomainUtils.cpp
--- a/Engine/Source/Editor/UnrealEd/Private/TargetDomain/TargetDomainUtils.cpp
+++ b/Engine/Source/Editor/UnrealEd/Private/TargetDomain/TargetDomainUtils.cpp
@@ -430,16 +430,20 @@ FEditorDomainOplog::FEditorDomainOplog()
 	FPaths::NormalizeDirectoryName(EngineDir);
 	FString ProjectDir = FPaths::ProjectDir();
 	FPaths::NormalizeDirectoryName(ProjectDir);
+	FString ProjectPath = FPaths::GetProjectFilePath();
+	FPaths::NormalizeFilename(ProjectPath);
 
 	IPlatformFile& PlatformFile = FPlatformFileManager::Get().GetPlatformFile();
 	FString AbsServerRoot = PlatformFile.ConvertToAbsolutePathForExternalAppForRead(*RootDir);
 	FString AbsEngineDir = PlatformFile.ConvertToAbsolutePathForExternalAppForRead(*EngineDir);
 	FString AbsProjectDir = PlatformFile.ConvertToAbsolutePathForExternalAppForRead(*ProjectDir);
+	FString ProjectFilePath = PlatformFile.ConvertToAbsolutePathForExternalAppForRead(*ProjectPath);
 
 #if UE_WITH_ZEN
 	if (UE::Zen::IsDefaultServicePresent())
 	{
-		HttpClient.TryCreateProject(ProjectId, OplogId, AbsServerRoot, AbsEngineDir, AbsProjectDir);
+		bool IsLocalConnection = HttpClient.GetZenServiceInstance().IsServiceRunningLocally();
+		HttpClient.TryCreateProject(ProjectId, OplogId, AbsServerRoot, AbsEngineDir, AbsProjectDir, IsLocalConnection ? ProjectFilePath : FStringView());
 		HttpClient.TryCreateOplog(ProjectId, OplogId, false /* bFullBuild */);
 	}
 #endif
