ZenCook: Use the oplog for package store entries instead of writing a container header

diff --git a/Engine/Source/Developer/IoStoreUtilities/Private/IoStoreUtilities.cpp b/Engine/Source/Developer/IoStoreUtilities/Private/IoStoreUtilities.cpp
--- a/Engine/Source/Developer/IoStoreUtilities/Private/IoStoreUtilities.cpp
+++ b/Engine/Source/Developer/IoStoreUtilities/Private/IoStoreUtilities.cpp
@@ -659,39 +659,32 @@ public:
 
 		if (const FPackageStoreManifest::FZenServerInfo* ZenServerInfo = PackageStoreManifest.ReadZenServerInfo())
 		{
-			DataSource = MakeUnique<FZenStoreDataSource>(*ZenServerInfo);
+			ZenStoreClient = MakeUnique<UE::FZenStoreHttpClient>(UE::Zen::FServiceSettings(ZenServerInfo->Settings));
+			ZenStoreClient->InitializeReadOnly(ZenServerInfo->ProjectId, ZenServerInfo->OplogId);
 		}
 		
-		if (DataSource)
+		if (ZenStoreClient)
 		{
-			FIoContainerId ContainerId = FIoContainerId::FromName(TEXT("global"));
-
-			TIoStatusOr<FIoBuffer> HeaderBuffer = DataSource->ReadChunk(CreateIoChunkId(ContainerId.Value(), 0, EIoChunkType::ContainerHeader), FIoReadOptions());
-			if (!HeaderBuffer.IsOk())
-			{
-				return FIoStatusBuilder(EIoErrorCode::ReadError) << TEXT("Failed to read container header");
-			}
-
-			FMemoryReaderView Ar(MakeArrayView(HeaderBuffer.ValueOrDie().Data(), HeaderBuffer.ValueOrDie().DataSize()));
-			FIoContainerHeader ContainerHeader;
-			Ar << ContainerHeader;
-
-			PackageEntries.Reserve(ContainerHeader.PackageIds.Num());
-			TArrayView<const FFilePackageStoreEntry> FilePackageEntries = MakeArrayView<const FFilePackageStoreEntry>(reinterpret_cast<const FFilePackageStoreEntry*>(ContainerHeader.StoreEntries.GetData()), ContainerHeader.PackageIds.Num());
-
-			int32 Idx = 0;
-			for (const FFilePackageStoreEntry& FilePackageEntry : FilePackageEntries)
-			{
-				FPackageStoreEntryResource& Entry = PackageEntries.AddDefaulted_GetRef();
-
-				Entry.ExportInfo = FPackageStoreExportInfo
+			TFuture<FIoStatus> FutureOplogStatus = ZenStoreClient->GetOplog().Next([this](TIoStatusOr<FCbObject> OplogStatus)
 				{
-					FilePackageEntry.ExportCount,
-					FilePackageEntry.ExportBundleCount
-				};
-				Entry.ImportedPackageIds = MakeArrayView<const FPackageId>(FilePackageEntry.ImportedPackages.Data(), FilePackageEntry.ImportedPackages.Num());
+					if (!OplogStatus.IsOk())
+					{
+						return OplogStatus.Status();
+					}
+					FCbObject Oplog = OplogStatus.ConsumeValueOrDie();
+					for (FCbField& OplogEntry : Oplog["entries"].AsArray())
+					{
+						FCbObject OplogObj = OplogEntry.AsObject();
+						FPackageStoreEntryResource Entry = FPackageStoreEntryResource::FromCbObject(OplogObj["packagestoreentry"].AsObject());
+						FPackageId PackageId = Entry.GetPackageId();
+						PackageIdToEntry.Add(PackageId, MoveTemp(Entry));
+					}
+					return FIoStatus::Ok;
+				});
 
-				PackageIdToEntry.Add(ContainerHeader.PackageIds[Idx++], &Entry);
+			if (!FutureOplogStatus.Get().IsOk())
+			{
+				return FutureOplogStatus.Get();
 			}
 		}
 		
@@ -747,22 +740,23 @@ public:
 
 	const FPackageStoreEntryResource* GetPackageStoreEntry(FPackageId PackageId) const
 	{
-		return PackageIdToEntry.FindRef(PackageId);
+		return PackageIdToEntry.Find(PackageId);
 	}
 
-	bool HasDataSource() const
+	bool HasZenStoreClient() const
 	{
-		return DataSource.IsValid();
+		return ZenStoreClient.IsValid();
 	}
 
 	TIoStatusOr<uint64> GetChunkSize(const FIoChunkId& ChunkId)
 	{
-		return DataSource->GetChunkSize(ChunkId);
+		return ZenStoreClient->GetChunkSize(ChunkId);
 	}
 
 	TIoStatusOr<FIoBuffer> ReadChunk(const FIoChunkId& ChunkId)
 	{
-		return DataSource->ReadChunk(ChunkId, FIoReadOptions());
+		FIoReadOptions ReadOptions;
+		return ZenStoreClient->ReadChunk(ChunkId, ReadOptions.GetOffset(), ReadOptions.GetSize());
 	}
 
 	void ReadChunkAsync(const FIoChunkId& ChunkId, TFunction<void(TIoStatusOr<FIoBuffer>)>&& Callback)
@@ -771,8 +765,8 @@ public:
 			: public FNonAbandonableTask
 		{
 		public:
-			FReadChunkTask(IDataSource* InDataSource, const FIoChunkId& InChunkId, TFunction<void(TIoStatusOr<FIoBuffer>)>&& InCallback)
-				: DataSource(InDataSource)
+			FReadChunkTask(UE::FZenStoreHttpClient* InZenStoreClient, const FIoChunkId& InChunkId, TFunction<void(TIoStatusOr<FIoBuffer>)>&& InCallback)
+				: ZenStoreClient(InZenStoreClient)
 				, ChunkId(InChunkId)
 				, Callback(MoveTemp(InCallback))
 			{
@@ -780,7 +774,8 @@ public:
 
 			void DoWork()
 			{
-				Callback(DataSource->ReadChunk(ChunkId, FIoReadOptions()));
+				FIoReadOptions ReadOptions;
+				Callback(ZenStoreClient->ReadChunk(ChunkId, ReadOptions.GetOffset(), ReadOptions.GetSize()));
 			}
 
 			TStatId GetStatId() const
@@ -789,22 +784,22 @@ public:
 			}
 
 		private:
-			IDataSource* DataSource;
+			UE::FZenStoreHttpClient* ZenStoreClient;
 			FIoChunkId ChunkId;
 			TFunction<void(TIoStatusOr<FIoBuffer>)> Callback;
 		};
 
-		(new FAutoDeleteAsyncTask<FReadChunkTask>(DataSource.Get(), ChunkId, MoveTemp(Callback)))->StartBackgroundTask();
+		(new FAutoDeleteAsyncTask<FReadChunkTask>(ZenStoreClient.Get(), ChunkId, MoveTemp(Callback)))->StartBackgroundTask();
 	}
 	
-	TIoStatusOr<FIoBuffer> ReadPackageHeader(FPackageId PackageId)
+	TIoStatusOr<FIoBuffer> ReadPackageHeaderFromZen(FPackageId PackageId)
 	{
-		if (const FPackageStoreEntryResource* Entry = PackageIdToEntry.FindRef(PackageId))
+		if (const FPackageStoreEntryResource* Entry = PackageIdToEntry.Find(PackageId))
 		{
 			FIoReadOptions ReadOptions;
 			ReadOptions.SetRange(0, 64 << 10);
 			
-			TIoStatusOr<FIoBuffer> Status = DataSource->ReadChunk(CreateIoChunkId(PackageId.Value(), 0, EIoChunkType::ExportBundleData), ReadOptions);
+			TIoStatusOr<FIoBuffer> Status = ZenStoreClient->ReadChunk(CreateIoChunkId(PackageId.Value(), 0, EIoChunkType::ExportBundleData), ReadOptions.GetOffset(), ReadOptions.GetSize());
 			if (!Status.IsOk())
 			{
 				return Status;
@@ -816,7 +811,7 @@ public:
 			{
 				ReadOptions.SetRange(0, HeaderSize);
 
-				Status = DataSource->ReadChunk(CreateIoChunkId(PackageId.Value(), 0, EIoChunkType::ExportBundleData), ReadOptions);
+				Status = ZenStoreClient->ReadChunk(CreateIoChunkId(PackageId.Value(), 0, EIoChunkType::ExportBundleData), ReadOptions.GetOffset(), ReadOptions.GetSize());
 				if (!Status.IsOk())
 				{
 					return Status;
@@ -831,44 +826,10 @@ public:
 	}
 
 private:
-	class IDataSource
-	{
-	public:
-		virtual ~IDataSource() = default;
-		virtual TIoStatusOr<uint64> GetChunkSize(const FIoChunkId& ChunkId) = 0;
-		virtual TIoStatusOr<FIoBuffer> ReadChunk(const FIoChunkId& ChunkId, const FIoReadOptions& ReadOptions) = 0;
-	};
-
-	class FZenStoreDataSource
-		: public IDataSource
-	{
-	public:
-		FZenStoreDataSource(const FPackageStoreManifest::FZenServerInfo& ZenServerInfo)
-		{
-			ZenStoreClient = MakeUnique<UE::FZenStoreHttpClient>(UE::Zen::FServiceSettings(ZenServerInfo.Settings));
-
-			ZenStoreClient->InitializeReadOnly(ZenServerInfo.ProjectId, ZenServerInfo.OplogId);
-		}
-
-		virtual TIoStatusOr<uint64> GetChunkSize(const FIoChunkId& ChunkId) override
-		{
-			return ZenStoreClient->GetChunkSize(ChunkId);
-		}
-
-		virtual TIoStatusOr<FIoBuffer> ReadChunk(const FIoChunkId& ChunkId, const FIoReadOptions& ReadOptions) override
-		{
-			return ZenStoreClient->ReadChunk(ChunkId, ReadOptions.GetOffset(), ReadOptions.GetSize());
-		}
-
-	private:
-		TUniquePtr<UE::FZenStoreHttpClient> ZenStoreClient;
-	};
-	
-	TUniquePtr<IDataSource> DataSource;
+	TUniquePtr<UE::FZenStoreHttpClient> ZenStoreClient;
 	FPackageStoreManifest PackageStoreManifest;
 	TArray<FPackageStoreManifest::FPackageInfo> ManifestPackageInfos;
-	TArray<FPackageStoreEntryResource> PackageEntries;
-	TMap<FPackageId, const FPackageStoreEntryResource*> PackageIdToEntry;
+	TMap<FPackageId, FPackageStoreEntryResource> PackageIdToEntry;
 	TMap<FString, FIoChunkId> FilenameToChunkIdMap;
 	TMap<FIoChunkId, FName> ChunkIdToPackageNameMap;
 	TMap<FString, FName> FilenameToPackageNameMap;
@@ -1611,9 +1572,9 @@ static void ParsePackageAssetsFromFiles(TArray<FLegacyCookedPackage*>& Packages,
 	FMemory::Free(OptionalSegmentUAssetMemory);
 }
 
-static void ParsePackageAssetsFromPackageStore(FCookedPackageStore& PackageStore, TArray<FLegacyCookedPackage*>& Packages, const FPackageStoreOptimizer& PackageStoreOptimizer)
+static void ParsePackageAssetsFromZen(FCookedPackageStore& PackageStore, TArray<FLegacyCookedPackage*>& Packages, const FPackageStoreOptimizer& PackageStoreOptimizer)
 {
-	IOSTORE_CPU_SCOPE(ParsePackageAssetsFromPackageStore);
+	IOSTORE_CPU_SCOPE(ParsePackageAssetsFromZen);
 	UE_LOG(LogIoStore, Display, TEXT("Parsing packages..."));
 
 	const int32 TotalPackageCount = Packages.Num();
@@ -1624,7 +1585,7 @@ static void ParsePackageAssetsFromPackageStore(FCookedPackageStore& PackageStore
 		&Packages](int32 Index)
 	{
 		FLegacyCookedPackage* Package = Packages[Index];
-		TIoStatusOr<FIoBuffer> HeaderBuffer = PackageStore.ReadPackageHeader(Package->GlobalPackageId);
+		TIoStatusOr<FIoBuffer> HeaderBuffer = PackageStore.ReadPackageHeaderFromZen(Package->GlobalPackageId);
 		Package->UExpSize = Package->UAssetSize - HeaderBuffer.ValueOrDie().DataSize();
 		Package->UAssetSize = HeaderBuffer.ValueOrDie().DataSize();
 		Package->OptimizedPackage = PackageStoreOptimizer.CreatePackageFromPackageStoreHeader(Package->PackageName, HeaderBuffer.ValueOrDie(), *Package->PackageStoreEntry);
@@ -2259,7 +2220,7 @@ void InitializeContainerTargetsAndPackages(
 		return true;
 	};
 
-	auto CreateTargetFileFromPackageStore = [
+	auto CreateTargetFileFromZen = [
 		&Arguments,
 		&Packages,
 		&PackageNameMap,
@@ -2390,8 +2351,8 @@ void InitializeContainerTargetsAndPackages(
 			for (const FContainerSourceFile& SourceFile : ContainerSource.SourceFiles)
 			{
 				FContainerTargetFile TargetFile;
-				bool bIsValidTargetFile = Arguments.PackageStore->HasDataSource()
-					? CreateTargetFileFromPackageStore(SourceFile, TargetFile)
+				bool bIsValidTargetFile = Arguments.PackageStore->HasZenStoreClient()
+					? CreateTargetFileFromZen(SourceFile, TargetFile)
 					: CreateTargetFileFromCookedFile(SourceFile, TargetFile);
 
 				if (!bIsValidTargetFile)
@@ -2660,9 +2621,9 @@ public:
 		{
 			return new FInMemoryWriteRequest(*this, InTargetFile);
 		}
-		else if (PackageStore->HasDataSource())
+		else if (PackageStore->HasZenStoreClient())
 		{
-			return new FCookedPackageStoreWriteRequest(*this, InTargetFile);
+			return new FZenWriteRequest(*this, InTargetFile);
 		}
 		else
 		{
@@ -2792,12 +2753,11 @@ private:
 		}
 	};
 
-	// Used when cooking directly to I/O store container file
-	class FCookedPackageStoreWriteRequest
+	class FZenWriteRequest
 		: public FWriteContainerTargetFileRequest
 	{
 	public:
-		FCookedPackageStoreWriteRequest(FIoStoreWriteRequestManager& InManager,const FContainerTargetFile& InTargetFile)
+		FZenWriteRequest(FIoStoreWriteRequestManager& InManager,const FContainerTargetFile& InTargetFile)
 			: FWriteContainerTargetFileRequest(InManager, InTargetFile) {}
 
 		virtual void LoadSourceBufferAsync() override
@@ -3379,9 +3339,9 @@ int32 CreateTarget(const FIoStoreArguments& Arguments, const FIoStoreWriterSetti
 		}
 	}
 
-	if (Arguments.PackageStore->HasDataSource())
+	if (Arguments.PackageStore->HasZenStoreClient())
 	{
-		ParsePackageAssetsFromPackageStore(*Arguments.PackageStore, Packages, PackageStoreOptimizer);
+		ParsePackageAssetsFromZen(*Arguments.PackageStore, Packages, PackageStoreOptimizer);
 	}
 	else
 	{
@@ -6686,7 +6646,7 @@ int32 CreateIoStoreContainerFiles(const TCHAR* CmdLine)
 	}
 	else
 	{
-		UE_CLOG(!Arguments.PackageStore->HasDataSource(), LogIoStore, Fatal, TEXT("Expected -ScriptObjects=<path to script objects file> argument"));
+		UE_CLOG(!Arguments.PackageStore->HasZenStoreClient(), LogIoStore, Fatal, TEXT("Expected -ScriptObjects=<path to script objects file> argument"));
 		TIoStatusOr<FIoBuffer> Status = Arguments.PackageStore->ReadChunk(CreateIoChunkId(0, 0, EIoChunkType::ScriptObjects));
 		UE_CLOG(!Status.IsOk(), LogIoStore, Fatal, TEXT("Failed reading script objects chunk '%s'"), *Status.Status().ToString());
 		Arguments.ScriptObjects = MakeUnique<FIoBuffer>(Status.ConsumeValueOrDie());
diff --git a/Engine/Source/Developer/IoStoreUtilities/Private/ZenStoreWriter.cpp b/Engine/Source/Developer/IoStoreUtilities/Private/ZenStoreWriter.cpp
--- a/Engine/Source/Developer/IoStoreUtilities/Private/ZenStoreWriter.cpp
+++ b/Engine/Source/Developer/IoStoreUtilities/Private/ZenStoreWriter.cpp
@@ -492,8 +492,7 @@ void FZenStoreWriter::BeginCook()
 		PackageObj.BeginObject();
 		PackageObj << "key" << "CookOnTheFly";
 
-		const bool bGenerateContainerHeader = false;
-		CreateProjectMetaData(Pkg, PackageObj, bGenerateContainerHeader);
+		CreateProjectMetaData(Pkg, PackageObj);
 
 		PackageObj.EndObject();
 		FCbObject Obj = PackageObj.Save().AsObject();
@@ -536,8 +535,7 @@ void FZenStoreWriter::EndCook()
 	PackageObj.BeginObject();
 	PackageObj << "key" << "EndCook";
 
-	const bool bGenerateContainerHeader = true;
-	CreateProjectMetaData(Pkg, PackageObj, bGenerateContainerHeader);
+	CreateProjectMetaData(Pkg, PackageObj);
 
 	PackageObj.EndObject();
 	FCbObject Obj = PackageObj.Save().AsObject();
@@ -1032,7 +1030,7 @@ void FZenStoreWriter::MarkPackagesUpToDate(TArrayView<const FName> UpToDatePacka
 	}
 }
 
-void FZenStoreWriter::CreateProjectMetaData(FCbPackage& Pkg, FCbWriter& PackageObj, bool bGenerateContainerHeader)
+void FZenStoreWriter::CreateProjectMetaData(FCbPackage& Pkg, FCbWriter& PackageObj)
 {
 	// File Manifest
 	{
@@ -1107,30 +1105,6 @@ void FZenStoreWriter::CreateProjectMetaData(FCbPackage& Pkg, FCbWriter& PackageO
 		PackageObj << "data" << ScriptAttachment;
 		PackageObj.EndObject();
 
-		// Generate Container Header
-		if (bGenerateContainerHeader)
-		{
-			FCbObjectId HeaderOid = ToObjectId(CreateIoChunkId(ContainerId.Value(), 0, EIoChunkType::ContainerHeader));
-			FIoBuffer HeaderBuffer;
-
-			{
-				FIoContainerHeader Header = PackageStoreOptimizer->CreateContainerHeader(ContainerId, PackageStoreEntries);
-				FLargeMemoryWriter HeaderAr(0, true);
-				HeaderAr << Header;
-				int64 DataSize = HeaderAr.TotalSize();
-				HeaderBuffer = FIoBuffer(FIoBuffer::AssumeOwnership, HeaderAr.ReleaseOwnership(), DataSize);
-			}
-
-			FCbAttachment HeaderAttachment = CreateAttachment(HeaderBuffer);
-			Pkg.AddAttachment(HeaderAttachment);
-			
-			PackageObj.BeginObject();
-			PackageObj << "id" << HeaderOid;
-			PackageObj << "name" << "ContainerHeader";
-			PackageObj << "data" << HeaderAttachment;
-			PackageObj.EndObject();
-		}
-
 		PackageObj.EndArray();	// End of Meta array
 	}
 }
diff --git a/Engine/Source/Developer/IoStoreUtilities/Public/ZenStoreWriter.h b/Engine/Source/Developer/IoStoreUtilities/Public/ZenStoreWriter.h
--- a/Engine/Source/Developer/IoStoreUtilities/Public/ZenStoreWriter.h
+++ b/Engine/Source/Developer/IoStoreUtilities/Public/ZenStoreWriter.h
@@ -15,7 +15,6 @@
 #include "Delegates/Delegate.h"
 #include "HAL/CriticalSection.h"
 #include "HAL/Platform.h"
-#include "IO/IoContainerId.h"
 #include "IO/IoDispatcher.h"
 #include "IO/PackageStore.h"
 #include "Misc/AssertionMacros.h"
@@ -162,7 +161,7 @@ private:
 		return PendingPackages.FindAndRemoveChecked(PackageName);
 	}
 
-	void CreateProjectMetaData(FCbPackage& Pkg, FCbWriter& PackageObj, bool bGenerateContainerHeader);
+	void CreateProjectMetaData(FCbPackage& Pkg, FCbWriter& PackageObj);
 	void BroadcastCommit(IPackageStoreWriter::FCommitEventArgs& EventArgs);
 	void BroadcastMarkUpToDate(IPackageStoreWriter::FMarkUpToDateEventArgs& EventArgs);
 	struct FZenCommitInfo;
@@ -182,7 +181,6 @@ private:
 	FString								OplogId;
 	FString								OutputPath;
 	FString								MetadataDirectoryPath;
-	FIoContainerId						ContainerId = FIoContainerId::FromName(TEXT("global"));
 	TMap<FName, TRefCountPtr<FPackageHashes>> AllPackageHashes;
 
 	FPackageStoreManifest				PackageStoreManifest;
diff --git a/Engine/Source/Runtime/StorageServerClient/Private/StorageServerConnection.cpp b/Engine/Source/Runtime/StorageServerClient/Private/StorageServerConnection.cpp
--- a/Engine/Source/Runtime/StorageServerClient/Private/StorageServerConnection.cpp
+++ b/Engine/Source/Runtime/StorageServerClient/Private/StorageServerConnection.cpp
@@ -14,6 +14,7 @@
 #include "ProfilingDebugging/CountersTrace.h"
 #include "SocketSubsystem.h"
 #include "Sockets.h"
+#include "IO/PackageStore.h"
 
 #if !UE_BUILD_SHIPPING
 
@@ -533,6 +534,34 @@ int32 FStorageServerConnection::HandshakeRequest(TArrayView<const TSharedPtr<FIn
 	return -1;
 }
 
+void FStorageServerConnection::PackageStoreRequest(TFunctionRef<void(FPackageStoreEntryResource&&)> Callback)
+{
+	TAnsiStringBuilder<256> ResourceBuilder;
+	ResourceBuilder.Append(OplogPath).Append("/entries");
+	FStorageServerRequest Request("GET", *ResourceBuilder, Hostname, EStorageServerContentType::CbObject);
+	FSocket* Socket = Request.Send(*this);
+	if (!Socket)
+	{
+		UE_LOG(LogStorageServerConnection, Fatal, TEXT("Failed to send oplog request to storage server at %s."), *ServerAddr->ToString(true));
+		return;
+	}
+	FStorageServerResponse Response(*this, *Socket);
+	if (Response.IsOk())
+	{
+		FCbObject ResponseObj = Response.GetResponseObject();
+		for (FCbField& OplogEntry : ResponseObj["entries"].AsArray())
+		{
+			FCbObject OplogObj = OplogEntry.AsObject();
+			FPackageStoreEntryResource Entry = FPackageStoreEntryResource::FromCbObject(OplogObj["packagestoreentry"].AsObject());
+			Callback(MoveTemp(Entry));
+		}
+	}
+	else
+	{
+		UE_LOG(LogStorageServerConnection, Fatal, TEXT("Failed to read oplog from storage server at %s. '%s'"), *ServerAddr->ToString(true), *Response.GetErrorMessage());
+	}
+}
+
 void FStorageServerConnection::FileManifestRequest(TFunctionRef<void(FIoChunkId Id, FStringView Path)> Callback)
 {
 	TAnsiStringBuilder<256> ResourceBuilder;
diff --git a/Engine/Source/Runtime/StorageServerClient/Private/StorageServerPackageStore.cpp b/Engine/Source/Runtime/StorageServerClient/Private/StorageServerPackageStore.cpp
--- a/Engine/Source/Runtime/StorageServerClient/Private/StorageServerPackageStore.cpp
+++ b/Engine/Source/Runtime/StorageServerClient/Private/StorageServerPackageStore.cpp
@@ -1,7 +1,6 @@
 // Copyright Epic Games, Inc. All Rights Reserved.
 
 #include "StorageServerPackageStore.h"
-#include "IO/IoContainerHeader.h"
 #include "IO/IoDispatcher.h"
 #include "StorageServerConnection.h"
 #include "Serialization/MemoryReader.h"
@@ -10,40 +9,23 @@
 
 FStorageServerPackageStoreBackend::FStorageServerPackageStoreBackend(FStorageServerConnection& Connection)
 {
-	FIoChunkId HeaderChunkId = CreateIoChunkId(FIoContainerId::FromName(TEXT("global")).Value(), 0, EIoChunkType::ContainerHeader);
-	Connection.ReadChunkRequest(HeaderChunkId, 0, uint64(-1), [this](FStorageServerResponse& Response)
-	{
-		FIoBuffer Chunk;
-		if (Response.SerializeChunk(Chunk))
+	Connection.PackageStoreRequest([this](FPackageStoreEntryResource&& PackageStoreEntryResource)
 		{
-			FMemoryReaderView Ar(MakeArrayView(reinterpret_cast<const uint8*>(Chunk.Data()), Chunk.DataSize()));
-			
-			FIoContainerHeader ContainerHeader;
-			Ar << ContainerHeader;
-			StoreEntriesData = MoveTemp(ContainerHeader.StoreEntries);
-			StoreEntriesMap.Reserve(ContainerHeader.PackageIds.Num());
-			TArrayView<const FFilePackageStoreEntry> StoreEntries(reinterpret_cast<const FFilePackageStoreEntry*>(StoreEntriesData.GetData()), ContainerHeader.PackageIds.Num());
-			int32 Index = 0;
-			for (const FFilePackageStoreEntry& StoreEntry : StoreEntries)
-			{
-				const FPackageId& PackageId = ContainerHeader.PackageIds[Index];
-				check(PackageId.IsValid());
-				StoreEntriesMap.FindOrAdd(PackageId, &StoreEntry);
-				++Index;
-			}
-		}
-	});
+			FStoreEntry& LocalEntry = StoreEntriesMap.FindOrAdd(PackageStoreEntryResource.GetPackageId());
+			LocalEntry.ExportInfo = MoveTemp(PackageStoreEntryResource.ExportInfo);
+			LocalEntry.ImportedPackages = MoveTemp(PackageStoreEntryResource.ImportedPackageIds);
+			LocalEntry.ShaderMapHashes = MoveTemp(PackageStoreEntryResource.ShaderMapHashes);
+		});
 }
 
 EPackageStoreEntryStatus FStorageServerPackageStoreBackend::GetPackageStoreEntry(FPackageId PackageId, FPackageStoreEntry& OutPackageStoreEntry)
 {
-	const FFilePackageStoreEntry* FindEntry = StoreEntriesMap.FindRef(PackageId);
+	const FStoreEntry* FindEntry = StoreEntriesMap.Find(PackageId);
 	if (FindEntry)
 	{
-		OutPackageStoreEntry.ExportInfo.ExportCount = FindEntry->ExportCount;
-		OutPackageStoreEntry.ExportInfo.ExportBundleCount = FindEntry->ExportBundleCount;
-		OutPackageStoreEntry.ImportedPackageIds = MakeArrayView(FindEntry->ImportedPackages.Data(), FindEntry->ImportedPackages.Num());
-		OutPackageStoreEntry.ShaderMapHashes = MakeArrayView(FindEntry->ShaderMapHashes.Data(), FindEntry->ShaderMapHashes.Num());
+		OutPackageStoreEntry.ExportInfo = FindEntry->ExportInfo;
+		OutPackageStoreEntry.ImportedPackageIds = FindEntry->ImportedPackages;
+		OutPackageStoreEntry.ShaderMapHashes = FindEntry->ShaderMapHashes;
 		return EPackageStoreEntryStatus::Ok;
 	}
 	else
diff --git a/Engine/Source/Runtime/StorageServerClient/Private/StorageServerPackageStore.h b/Engine/Source/Runtime/StorageServerClient/Private/StorageServerPackageStore.h
--- a/Engine/Source/Runtime/StorageServerClient/Private/StorageServerPackageStore.h
+++ b/Engine/Source/Runtime/StorageServerClient/Private/StorageServerPackageStore.h
@@ -36,8 +36,13 @@ public:
 	}
 
 private:
-	TArray<uint8> StoreEntriesData;
-	TMap<FPackageId, const FFilePackageStoreEntry*> StoreEntriesMap;
+	struct FStoreEntry
+	{
+		FPackageStoreExportInfo ExportInfo;
+		TArray<FPackageId> ImportedPackages;
+		TArray<FSHAHash> ShaderMapHashes;
+	};
+	TMap<FPackageId, FStoreEntry> StoreEntriesMap;
 };
 
 #endif
\ No newline at end of file
diff --git a/Engine/Source/Runtime/StorageServerClient/Public/StorageServerConnection.h b/Engine/Source/Runtime/StorageServerClient/Public/StorageServerConnection.h
--- a/Engine/Source/Runtime/StorageServerClient/Public/StorageServerConnection.h
+++ b/Engine/Source/Runtime/StorageServerClient/Public/StorageServerConnection.h
@@ -18,6 +18,7 @@ class FSocket;
 class FStorageServerChunkBatchRequest;
 class FStorageServerConnection;
 class ISocketSubsystem;
+struct FPackageStoreEntryResource;
 
 enum class EStorageServerContentType : uint8
 {
@@ -181,6 +182,7 @@ public:
 
 	STORAGESERVERCLIENT_API bool Initialize(TArrayView<const FString> HostAddresses, int32 Port, const TCHAR* ProjectNameOverride = nullptr, const TCHAR* PlatformNameOverride = nullptr);
 
+	STORAGESERVERCLIENT_API void PackageStoreRequest(TFunctionRef<void(FPackageStoreEntryResource&&)> Callback);
 	STORAGESERVERCLIENT_API void FileManifestRequest(TFunctionRef<void(FIoChunkId Id, FStringView Path)> Callback);
 	STORAGESERVERCLIENT_API int64 ChunkSizeRequest(const FIoChunkId& ChunkId);
 	STORAGESERVERCLIENT_API bool ReadChunkRequest(const FIoChunkId& ChunkId, uint64 Offset, uint64 Size, TFunctionRef<void(FStorageServerResponse&)> OnResponse);
