I/O Store - added chunk type filter to -Diff command

#prefight 6389b0fc303395f6c9820284
#rnx

diff --git a/Engine/Source/Developer/IoStoreUtilities/Private/IoStoreUtilities.cpp b/Engine/Source/Developer/IoStoreUtilities/Private/IoStoreUtilities.cpp
--- a/Engine/Source/Developer/IoStoreUtilities/Private/IoStoreUtilities.cpp
+++ b/Engine/Source/Developer/IoStoreUtilities/Private/IoStoreUtilities.cpp
@@ -5074,12 +5074,33 @@ int32 Describe(
 	return 0;
 }
 
+enum class EChunkTypeFilter
+{
+	None,
+	PackageData,
+	BulkData
+};
+
+FString LexToString(EChunkTypeFilter Filter)
+{
+	switch(Filter)
+	{
+		case EChunkTypeFilter::PackageData:
+			return TEXT("PackageData");
+		case EChunkTypeFilter::BulkData:
+			return TEXT("BulkData");
+		default:
+			return TEXT("None");
+	}
+}
+
 static int32 Diff(
 	const FString& SourcePath,
 	const FKeyChain& SourceKeyChain,
 	const FString& TargetPath,
 	const FKeyChain& TargetKeyChain,
-	const FString& OutPath)
+	const FString& OutPath,
+	EChunkTypeFilter ChunkTypeFilter)
 {
 	struct FContainerChunkInfo
 	{
@@ -5103,7 +5124,7 @@ static int32 Diff(
 
 	using FContainers = TMap<FString, FContainerChunkInfo>;
 
-	auto ReadContainers = [](const FString& Directory, const FKeyChain& KeyChain, FContainers& OutContainers)
+	auto ReadContainers = [ChunkTypeFilter](const FString& Directory, const FKeyChain& KeyChain, FContainers& OutContainers)
 	{
 		TArray<FString> ContainerFileNames;
 		IFileManager::Get().FindFiles(ContainerFileNames, *(Directory / TEXT("*.utoc")), true, false);
@@ -5124,11 +5145,24 @@ static int32 Diff(
 			FContainerChunkInfo& ContainerChunkInfo = OutContainers.FindOrAdd(ContainerName);
 			ContainerChunkInfo.ContainerName = MoveTemp(ContainerName);
 
-			Reader->EnumerateChunks([&ContainerChunkInfo](const FIoStoreTocChunkInfo& ChunkInfo)
+			Reader->EnumerateChunks([&ContainerChunkInfo, ChunkTypeFilter](const FIoStoreTocChunkInfo& ChunkInfo)
 			{
-				ContainerChunkInfo.ChunkInfoById.Add(ChunkInfo.Id, ChunkInfo);
-				ContainerChunkInfo.UncompressedContainerSize += ChunkInfo.Size;
-				ContainerChunkInfo.CompressedContainerSize += ChunkInfo.CompressedSize;
+				const EIoChunkType ChunkType = ChunkInfo.Id.GetChunkType();
+
+				const bool bCompareChunk =
+					ChunkTypeFilter == EChunkTypeFilter::None
+					|| (ChunkTypeFilter == EChunkTypeFilter::PackageData
+							&& ChunkType == EIoChunkType::ExportBundleData)
+					|| (ChunkTypeFilter == EChunkTypeFilter::BulkData
+							&& (ChunkType == EIoChunkType::BulkData || ChunkType == EIoChunkType::OptionalBulkData || ChunkType == EIoChunkType::MemoryMappedBulkData));
+
+				if (bCompareChunk)
+				{
+					ContainerChunkInfo.ChunkInfoById.Add(ChunkInfo.Id, ChunkInfo);
+					ContainerChunkInfo.UncompressedContainerSize += ChunkInfo.Size;
+					ContainerChunkInfo.CompressedContainerSize += ChunkInfo.CompressedSize;
+				}
+
 				return true;
 			});
 		}
@@ -5237,6 +5271,7 @@ static int32 Diff(
 	OutputDevice->Logf(ELogVerbosity::Display, TEXT("------------------------------ Container Diff Summary ------------------------------"));
 	OutputDevice->Logf(ELogVerbosity::Display, TEXT("Source path '%s'"), *SourcePath);
 	OutputDevice->Logf(ELogVerbosity::Display, TEXT("Target path '%s'"), *TargetPath);
+	OutputDevice->Logf(ELogVerbosity::Display, TEXT("Chunk type filter '%s'"), *LexToString(ChunkTypeFilter));
 
 	OutputDevice->Logf(ELogVerbosity::Display, TEXT(""));
 	OutputDevice->Logf(ELogVerbosity::Display, TEXT("Source container file(s):"));
@@ -6549,7 +6584,17 @@ int32 CreateIoStoreContainerFiles(const TCHAR* CmdLine)
 			KeyChainUtilities::LoadKeyChainFromFile(CryptoKeysCacheFilename, TargetKeyChain);
 		}
 
-		return Diff(SourcePath, SourceKeyChain, TargetPath, TargetKeyChain, OutPath);
+		EChunkTypeFilter ChunkTypeFilter = EChunkTypeFilter::None;
+		if (FParse::Param(FCommandLine::Get(), TEXT("FilterBulkData")))
+		{
+			ChunkTypeFilter = EChunkTypeFilter::BulkData;
+		}
+		else if (FParse::Param(FCommandLine::Get(), TEXT("FilterPackageData")))
+		{
+			ChunkTypeFilter = EChunkTypeFilter::PackageData;
+		}
+
+		return Diff(SourcePath, SourceKeyChain, TargetPath, TargetKeyChain, OutPath, ChunkTypeFilter);
 	}
 	else if (FParse::Param(FCommandLine::Get(), TEXT("Staged2Zen")))
 	{
