Add marker file for project store oplog in Saved/Cooked/<platform> folder so Zen can check if an oplog should be garbage collected.

Companion Zen PR: https://github.com/EpicGames/zen/pull/209

diff --git a/Engine/Source/Developer/IoStoreUtilities/Internal/ZenStoreHttpClient.h b/Engine/Source/Developer/IoStoreUtilities/Internal/ZenStoreHttpClient.h
--- a/Engine/Source/Developer/IoStoreUtilities/Internal/ZenStoreHttpClient.h
+++ b/Engine/Source/Developer/IoStoreUtilities/Internal/ZenStoreHttpClient.h
@@ -35,7 +35,7 @@ public:
 	bool TryCreateProject(FStringView InProjectId, FStringView InOplogId, FStringView ServerRoot, 
 					FStringView EngineRoot, FStringView ProjectRoot,
 					FStringView ProjectFilePath);
-	bool TryCreateOplog(FStringView InProjectId, FStringView InOplogId, bool bFullBuild);
+	bool TryCreateOplog(FStringView InProjectId, FStringView InOplogId, FStringView InOplogLifetimeMarkerPath, bool bFullBuild);
 
 	void InitializeReadOnly(FStringView InProjectId, FStringView InOplogId);
 
diff --git a/Engine/Source/Developer/IoStoreUtilities/Private/ZenStoreHttpClient.cpp b/Engine/Source/Developer/IoStoreUtilities/Private/ZenStoreHttpClient.cpp
--- a/Engine/Source/Developer/IoStoreUtilities/Private/ZenStoreHttpClient.cpp
+++ b/Engine/Source/Developer/IoStoreUtilities/Private/ZenStoreHttpClient.cpp
@@ -129,7 +129,7 @@ FZenStoreHttpClient::TryCreateProject(FStringView InProjectId,
 	return bConnectionSucceeded;
 }
 
-bool FZenStoreHttpClient::TryCreateOplog(FStringView InProjectId, FStringView InOplogId, bool bFullBuild)
+bool FZenStoreHttpClient::TryCreateOplog(FStringView InProjectId, FStringView InOplogId, FStringView InOplogLifetimeMarkerPath, bool bFullBuild)
 {
 	if (!IsConnected())
 	{
@@ -157,10 +157,16 @@ bool FZenStoreHttpClient::TryCreateOplog(FStringView InProjectId, FStringView In
 	}
 	else
 	{
-		FMemoryView Payload;
+		FCbWriter Writer;
+		Writer.BeginObject();
+		Writer.AddString("gcpath", InOplogLifetimeMarkerPath);
+		Writer.EndObject();
+
+		FCbFieldIterator OplogCreateInfo = Writer.Save();
 
 		Request->Reset();
-		Res = Request->PerformBlockingPost(OplogPath, Payload);
+
+		Res = Request->PerformBlockingPost(OplogPath, OplogCreateInfo.AsObjectView());
 
 		if (Res != Zen::FZenHttpRequest::Result::Success)
 		{
@@ -639,7 +645,7 @@ bool FZenStoreHttpClient::TryCreateProject(FStringView InProjectId, FStringView
 	return false;
 }
 
-bool FZenStoreHttpClient::TryCreateOplog(FStringView InProjectId, FStringView InOplogId, bool bFullBuild)
+bool FZenStoreHttpClient::TryCreateOplog(FStringView InProjectId, FStringView InOplogId, FStringView InOplogLifetimeMarkerPath, bool bFullBuild)
 {
 	return false;
 }
diff --git a/Engine/Source/Developer/IoStoreUtilities/Private/ZenStoreWriter.cpp b/Engine/Source/Developer/IoStoreUtilities/Private/ZenStoreWriter.cpp
--- a/Engine/Source/Developer/IoStoreUtilities/Private/ZenStoreWriter.cpp
+++ b/Engine/Source/Developer/IoStoreUtilities/Private/ZenStoreWriter.cpp
@@ -364,7 +364,8 @@ void FZenStoreWriter::Initialize(const FCookInfo& Info)
 
 	if (!bInitialized)
 	{
-		if (Info.bFullBuild && !Info.bWorkerOnSharedSandbox)
+		bool CleanBuild = Info.bFullBuild && !Info.bWorkerOnSharedSandbox;
+		if (CleanBuild)
 		{
 			UE_LOG(LogZenStoreWriter, Display, TEXT("Deleting %s..."), *OutputPath);
 			const bool bRequireExists = false;
@@ -372,7 +373,16 @@ void FZenStoreWriter::Initialize(const FCookInfo& Info)
 			IFileManager::Get().DeleteDirectory(*OutputPath, bRequireExists, bTree);
 		}
 
-		bool bOplogEstablished = HttpClient->TryCreateOplog(ProjectId, OplogId, Info.bFullBuild);
+		FString OplogLifetimeMarkerPath = OutputPath / (ProjectId + TEXT(".") + OplogId + TEXT(".projectstore"));
+		TUniquePtr<FArchive> OplogMarker(IFileManager::Get().CreateFileWriter(*OplogLifetimeMarkerPath));
+
+		bool bOplogEstablished = HttpClient->TryCreateOplog(ProjectId, OplogId, OplogLifetimeMarkerPath, Info.bFullBuild);
+		OplogMarker.Reset();
+
+		if (!bOplogEstablished && CleanBuild)
+		{
+			IFileManager::Get().Delete(*OplogLifetimeMarkerPath);
+		}
 		UE_CLOG(!bOplogEstablished, LogZenStoreWriter, Fatal, TEXT("Failed to establish oplog on the ZenServer"));
 
 		if (!Info.bFullBuild)
diff --git a/Engine/Source/Editor/UnrealEd/Private/TargetDomain/TargetDomainUtils.cpp b/Engine/Source/Editor/UnrealEd/Private/TargetDomain/TargetDomainUtils.cpp
--- a/Engine/Source/Editor/UnrealEd/Private/TargetDomain/TargetDomainUtils.cpp
+++ b/Engine/Source/Editor/UnrealEd/Private/TargetDomain/TargetDomainUtils.cpp
@@ -444,7 +444,7 @@ FEditorDomainOplog::FEditorDomainOplog()
 	{
 		bool IsLocalConnection = HttpClient.GetZenServiceInstance().IsServiceRunningLocally();
 		HttpClient.TryCreateProject(ProjectId, OplogId, AbsServerRoot, AbsEngineDir, AbsProjectDir, IsLocalConnection ? ProjectFilePath : FStringView());
-		HttpClient.TryCreateOplog(ProjectId, OplogId, false /* bFullBuild */);
+		HttpClient.TryCreateOplog(ProjectId, OplogId, TEXT("") /*InOplogMarkerFile*/, false /* bFullBuild */);
 	}
 #endif
 }
