Removing storing the hashes of chunk types other than bulk and external file, because Zen stores packed data before its hit by the PackageOptimizer under the same chunk id, so it doesn't match the data that ends up in the iostore containers.

diff --git a/Engine/Source/Developer/IoStoreUtilities/Private/ZenStoreWriter.cpp b/Engine/Source/Developer/IoStoreUtilities/Private/ZenStoreWriter.cpp
--- a/Engine/Source/Developer/IoStoreUtilities/Private/ZenStoreWriter.cpp
+++ b/Engine/Source/Developer/IoStoreUtilities/Private/ZenStoreWriter.cpp
@@ -673,11 +673,6 @@ void FZenStoreWriter::CommitPackageInternal(FZenCommitInfo&& ZenCommitInfo)
 		
 		FMD5 PkgHashGen;
 		FCbPackage OplogEntry;
-		
-		if (bComputeHash)
-		{
-			PackageState->PackageHashes->ChunkHashes.Add(PkgData.Info.ChunkId, PkgData.CompressedPayload.Get().GetRawHash());
-		}
 
 		FCbAttachment PkgDataAttachment = FCbAttachment(PkgData.CompressedPayload.Get());
 		PkgHashGen.Update(PkgDataAttachment.GetHash().GetBytes(), sizeof(FIoHash::ByteArray));
@@ -835,7 +830,6 @@ void FZenStoreWriter::CommitPackageInternal(FZenCommitInfo&& ZenCommitInfo)
 		{
 			FCompressedBuffer Payload = PkgData.CompressedPayload.Get();
 			FIoHash IoHash = Payload.GetRawHash();
-			PackageState->PackageHashes->ChunkHashes.Add(PkgData.Info.ChunkId, IoHash);
 			PkgHashGen.Update(IoHash.GetBytes(), sizeof(FIoHash::ByteArray));
 		}
 		
