IoStoreWriterHashDatabase - add a way for the iostore writers to access hashes that have already been calculated for chunks, to save pak time.

diff --git a/Engine/Source/Developer/IoStoreUtilities/Private/IoStoreUtilities.cpp b/Engine/Source/Developer/IoStoreUtilities/Private/IoStoreUtilities.cpp
--- a/Engine/Source/Developer/IoStoreUtilities/Private/IoStoreUtilities.cpp
+++ b/Engine/Source/Developer/IoStoreUtilities/Private/IoStoreUtilities.cpp
@@ -132,7 +132,7 @@ public:
 			TUniquePtr<FIoStoreReader> Reader = CreateIoStoreReader(*ContainerFilePath, InDecryptionKeychain);
 			if (Reader.IsValid() == false)
 			{
-				UE_LOG(LogIoStore, Error, TEXT("Failed to open reference chunk container %s"), *InGlobalContainerFileName);
+				UE_LOG(LogIoStore, Error, TEXT("Failed to open reference chunk container %s"), *ContainerFilePath);
 				return false;
 			}
 			FReaderChunks& ReaderChunks = ChunkDatabase.FindOrAdd(Reader->GetContainerId());
@@ -867,6 +867,7 @@ struct FIoStoreArguments
 	FReleasedPackages ReleasedPackages;
 	TUniquePtr<FCookedPackageStore> PackageStore;
 	TUniquePtr<FIoBuffer> ScriptObjects;
+	bool bVerifyHashDatabase = false;
 	bool bSign = false;
 	bool bRemapPluginContentToGame = false;
 	bool bCreateDirectoryIndex = true;
@@ -3160,6 +3161,42 @@ int32 DoAssetRegistryWritebackAfterStage(const FString& InAssetRegistryFileName,
 	return SaveAssetRegistry(InAssetRegistryFileName, AssetRegistry, true) ? 0 : 1;
 }
 
+static bool FindAndLoadDevelopmentAssetRegistry(const FString& InCookedDir, FAssetRegistryState& OutAssetRegistry, FString* OutAssetRegistryFileName /*optional, set on success*/)
+{
+	TRACE_CPUPROFILER_EVENT_SCOPE(LoadingAssetRegistry);
+
+	// Look for the development registry. Should be in \\GameName\\Metadata\\DevelopmentAssetRegistry.bin, but we don't know what "GameName" is.
+	TArray<FString> PossibleAssetRegistryFiles;
+	IFileManager::Get().FindFilesRecursive(PossibleAssetRegistryFiles, *InCookedDir, GetDevelopmentAssetRegistryFilename(), true, false);
+
+	if (PossibleAssetRegistryFiles.Num() > 1)
+	{
+		UE_LOG(LogIoStore, Warning, TEXT("Found multiple possible development asset registries:"));
+		for (FString& Filename : PossibleAssetRegistryFiles)
+		{
+			UE_LOG(LogIoStore, Warning, TEXT("    %s"), *Filename);
+		}
+	}
+
+	if (PossibleAssetRegistryFiles.Num() == 0)
+	{
+		UE_LOG(LogIoStore, Error, TEXT("No development asset registry file found!"));
+		return false;
+	}
+
+	UE_LOG(LogIoStore, Display, TEXT("Using input asset registry: %s"), *PossibleAssetRegistryFiles[0]);
+	if (LoadAssetRegistry(PossibleAssetRegistryFiles[0], OutAssetRegistry) == false)
+	{
+		return false; // already logged
+	}
+
+	if (OutAssetRegistryFileName)
+	{
+		*OutAssetRegistryFileName = MoveTemp(PossibleAssetRegistryFiles[0]);
+	}
+	return true;
+}
+
 bool DoAssetRegistryWritebackDuringStage(EAssetRegistryWritebackMethod InMethod, const FString& InCookedDir, TArray<TSharedPtr<IIoStoreWriter>>& InIoStoreWriters)
 {
 	// This version called during container creation.
@@ -3170,37 +3207,10 @@ bool DoAssetRegistryWritebackDuringStage(EAssetRegistryWritebackMethod InMethod,
 	// The overwhelming majority of time for the asset registry writeback is loading and saving.
 	FString AssetRegistryFileName;
 	FAssetRegistryState AssetRegistry;
+	if (FindAndLoadDevelopmentAssetRegistry(InCookedDir, AssetRegistry, &AssetRegistryFileName) == false)
 	{
-		TRACE_CPUPROFILER_EVENT_SCOPE(LoadingAssetRegistry);
-
-		// Look for the development registry. Should be in \\GameName\\Metadata\\DevelopmentAssetRegistry.bin, but we don't know what "GameName" is.
-		TArray<FString> PossibleAssetRegistryFiles;
-		IFileManager::Get().FindFilesRecursive(PossibleAssetRegistryFiles, *InCookedDir, GetDevelopmentAssetRegistryFilename(), true, false);
-
-		if (PossibleAssetRegistryFiles.Num() > 1)
-		{
-			UE_LOG(LogIoStore, Warning, TEXT("Found multiple possible development asset registries:"));
-			for (FString& Filename : PossibleAssetRegistryFiles)
-			{
-				UE_LOG(LogIoStore, Warning, TEXT("    %s"), *Filename);
-			}
-		}
-
-		if (PossibleAssetRegistryFiles.Num() == 0)
-		{
-			UE_LOG(LogIoStore, Error, TEXT("No development asset registry file found!"));
-			return false;
-		}
-		else
-		{
-			AssetRegistryFileName = MoveTemp(PossibleAssetRegistryFiles[0]);
-			UE_LOG(LogIoStore, Display, TEXT("Using input asset registry: %s"), *AssetRegistryFileName);
-
-			if (LoadAssetRegistry(AssetRegistryFileName, AssetRegistry) == false)
-			{
-				return false; // already logged
-			}
-		}
+		// already logged
+		return false;
 	}
 
 	// Create a map off the package id to all of its chunks. 2 inline allocation
@@ -3251,6 +3261,59 @@ bool DoAssetRegistryWritebackDuringStage(EAssetRegistryWritebackMethod InMethod,
 	return true;
 }
 
+// Implements providing the chunk hashes that exist in the asset registry to the
+// iostore writer to avoid reading and hashing redundently.
+class FIoStoreHashDb : public IIoStoreWriterHashDatabase
+{
+public:
+	virtual ~FIoStoreHashDb() {}
+
+	TMap<FIoChunkId, FIoHash> Hashes;
+
+	bool Initialize(const FString& InCookedDir)
+	{
+		FString AssetRegistryFileName;
+		FAssetRegistryState AssetRegistry;
+		if (FindAndLoadDevelopmentAssetRegistry(InCookedDir, AssetRegistry, nullptr) == false)
+		{
+			// already logged
+			return false;
+		}
+
+		double StartTime = FPlatformTime::Seconds();
+
+		const TMap<FName, const FAssetPackageData*>& Packages = AssetRegistry.GetAssetPackageDataMap();
+		for (auto PackageIter : Packages)
+		{
+			for (const TPair<FIoChunkId, FIoHash>& HashIter : PackageIter.Value->ChunkHashes)
+			{
+				// For the moment, only bulk data types are added to teh asset registry - gate here so that
+				// we remember to verify all the hashes match when they eventually get added during cook.
+				if (HashIter.Key.GetChunkType() == EIoChunkType::BulkData ||
+					HashIter.Key.GetChunkType() == EIoChunkType::OptionalBulkData)
+				{
+					Hashes.Add(HashIter.Key, HashIter.Value);
+				}
+			}
+		}
+
+		double EndTime = FPlatformTime::Seconds();
+		UE_LOG(LogIoStore, Display, TEXT("Added %d hashes to the hash database, init took %f seconds"), Hashes.Num(), EndTime - StartTime);
+		return true;
+	}
+
+	virtual bool FindHashForChunkId(const FIoChunkId& ChunkId, FIoChunkHash& OutHash) const override
+	{
+		const FIoHash* Exists = Hashes.Find(ChunkId);
+		if (Exists)
+		{
+			OutHash = FIoChunkHash::CreateFromIoHash(*Exists);
+			return true;
+		}
+		return false;
+	}
+};
+
 int32 CreateTarget(const FIoStoreArguments& Arguments, const FIoStoreWriterSettings& GeneralIoWriterSettings)
 {
 	TGuardValue<int32> GuardAllowUnversionedContentInEditor(GAllowUnversionedContentInEditor, 1);
@@ -3270,6 +3333,16 @@ int32 CreateTarget(const FIoStoreArguments& Arguments, const FIoStoreWriterSetti
 		}
 	}
 
+	TSharedPtr<IIoStoreWriterHashDatabase> HashDatabase = MakeShared<FIoStoreHashDb>();
+	if (((FIoStoreHashDb&)*HashDatabase).Initialize(Arguments.CookedDir) == false)
+	{
+		UE_LOG(LogIoStore, Warning, TEXT("Unabled to initialize the hash database from the asset registry!"));
+	}
+	if (Arguments.bVerifyHashDatabase)
+	{
+		UE_LOG(LogIoStore, Display, TEXT("Hash database verification on: hashes will be checked for accuracy during this run."));
+	}
+
 	TArray<FLegacyCookedPackage*> Packages;
 	FPackageNameMap PackageNameMap;
 	FPackageIdMap PackageIdMap;
@@ -3330,6 +3403,7 @@ int32 CreateTarget(const FIoStoreArguments& Arguments, const FIoStoreWriterSetti
 				ContainerTarget->IoStoreWriter = IoStoreWriterContext->CreateContainer(*ContainerTarget->OutputPath, ContainerSettings);
 				ContainerTarget->IoStoreWriter->EnableDiskLayoutOrdering(ContainerTarget->PatchSourceReaders);
 				ContainerTarget->IoStoreWriter->SetReferenceChunkDatabase(ChunkDatabase);
+				ContainerTarget->IoStoreWriter->SetHashDatabase(HashDatabase, Arguments.bVerifyHashDatabase);
 				IoStoreWriters.Add(ContainerTarget->IoStoreWriter);
 				if (!ContainerTarget->OptionalSegmentOutputPath.IsEmpty())
 				{
@@ -3571,7 +3645,7 @@ int32 CreateTarget(const FIoStoreArguments& Arguments, const FIoStoreWriterSetti
 		}
 		else
 		{
-			UE_LOG(LogIoStore, Display, TEXT("Hashing chunks (%llu/%llu)..."), Progress.HashedChunksCount, Progress.TotalChunksCount);
+			UE_LOG(LogIoStore, Display, TEXT("Hashing chunks (%llu/%llu) %llu from hash database..."), Progress.HashedChunksCount, Progress.TotalChunksCount, Progress.HashDbChunksCount);
 		}
 	}
 	if (GeneralIoWriterSettings.bCompressionEnableDDC)
@@ -6327,6 +6401,11 @@ bool ParseContainerGenerationArguments(FIoStoreArguments& Arguments, FIoStoreWri
 		KeyChainUtilities::LoadKeyChainFromFile(CryptoKeysCacheFilename, Arguments.ReferenceChunkKeys);
 	}
 
+	// By default, we use any hashes in the asset registry that exist in order to avoid reading and hashing
+	// chunk unnecessarily. This flag causes us to read and hash anyway, and then ensure they match what is
+	// in the asset registry. It is very bad if this fails!
+	Arguments.bVerifyHashDatabase = FParse::Param(FCommandLine::Get(), TEXT("-verifyhashdatabase"));
+
 
 	uint64 PatchPaddingAlignment = 0;
 	if (ParseSizeArgument(FCommandLine::Get(), TEXT("-patchpaddingalign="), PatchPaddingAlignment))
diff --git a/Engine/Source/Runtime/Core/Private/IO/IoStore.cpp b/Engine/Source/Runtime/Core/Private/IO/IoStore.cpp
--- a/Engine/Source/Runtime/Core/Private/IO/IoStore.cpp
+++ b/Engine/Source/Runtime/Core/Private/IO/IoStore.cpp
@@ -350,6 +350,7 @@ public:
 	FIoStoreWriterContext::FProgress GetProgress() const
 	{
 		FIoStoreWriterContext::FProgress Progress;
+		Progress.HashDbChunksCount = HashDbChunksCount.Load();
 		Progress.TotalChunksCount = TotalChunksCount.Load();
 		Progress.HashedChunksCount = HashedChunksCount.Load();
 		Progress.CompressedChunksCount = CompressedChunksCount.Load();
@@ -440,6 +441,7 @@ private:
 	FIoStoreWriteQueue WriterQueue;
 	TAtomic<uint64> TotalChunksCount{ 0 };
 	TAtomic<uint64> HashedChunksCount{ 0 };
+	TAtomic<uint64> HashDbChunksCount{ 0 };
 	TAtomic<uint64> CompressedChunksCount{ 0 };
 	TAtomic<uint64> SerializedChunksCount{ 0 };
 	TAtomic<uint64> ScheduledCompressionTasksCount{ 0 };
@@ -657,6 +659,11 @@ public:
 	{
 		ReferenceChunkDatabase = InReferenceChunkDatabase;
 	}
+	void SetHashDatabase(TSharedPtr<IIoStoreWriterHashDatabase> InHashDatabase, bool bInVerifyHashDatabase)
+	{
+		HashDatabase = InHashDatabase;
+		bVerifyHashDatabase = bInVerifyHashDatabase;
+	}
 
 	void EnumerateChunks(TFunction<bool(const FIoStoreTocChunkInfo&)>&& Callback) const
 	{
@@ -758,6 +765,16 @@ public:
 
 	virtual void Append(const FIoChunkId& ChunkId, IIoStoreWriteRequest* Request, const FIoWriteOptions& WriteOptions) override
 	{
+		//
+		// This function sets up the sequence of events that takes a chunk from source data on disc
+		// to written to a container. The first thing that happens is the source data is read in order
+		// to hash it to detect whether or not it's modified as well as look up in reference databases.
+		// Load the data -> PrepareSourceBufferAsync
+		// Hash the data -> HashTask lambda
+		//
+		// The hash task itself doesn't continue to the next steps - the Flush() call
+		// waits for all hashes to be complete before kicking the next steps.
+		//
 		TRACE_CPUPROFILER_EVENT_SCOPE(AppendWriteRequest);
 		check(!bHasFlushed);
 		checkf(ChunkId.IsValid(), TEXT("ChunkId is not valid!"));
@@ -769,22 +786,54 @@ public:
 		Entries.Add(Entry);
 		Entry->ChunkId = ChunkId;
 		Entry->Options = WriteOptions;
-		Entry->Request = Request;
-		Entry->HashBarrier = FGraphEvent::CreateGraphEvent();
+		Entry->Request = Request;		
 		Entry->BeginCompressionBarrier = FGraphEvent::CreateGraphEvent();
 		Entry->FinishCompressionBarrier = FGraphEvent::CreateGraphEvent();
 		Entry->FinishEncryptionAndSigningBarrier = FGraphEvent::CreateGraphEvent();
 		Entry->BeginWriteBarrier = FGraphEvent::CreateGraphEvent();
 		Entry->WriteFinishedEvent = FGraphEvent::CreateGraphEvent();
 		
+		bool bTestHashValid = false;
+		FIoChunkHash TestHash;
+		// If we can get the hash without reading the whole thing and hashing it, do so
+		// to avoid the IO.
+		if (HashDatabase.IsValid() &&
+			HashDatabase->FindHashForChunkId(ChunkId, Entry->ChunkHash))
+		{
+			if (bVerifyHashDatabase == false)
+			{
+				// If we aren't validating then we just use it and bail.
+				WriterContext->HashDbChunksCount.IncrementExchange();
+				WriterContext->HashedChunksCount.IncrementExchange();
+				return;
+			}
+
+			// If we are validating, copy the hash out and run the normal path
+			// to check against it.
+			TestHash = Entry->ChunkHash;
+			bTestHashValid = true;
+		}
+		// Otherwise, we have to do the load & hash
+
+		Entry->HashBarrier = FGraphEvent::CreateGraphEvent();
+
 		FGraphEventArray HashPrereqs;
 		HashPrereqs.Add(Entry->HashBarrier);
-		Entry->HashTask = FFunctionGraphTask::CreateAndDispatchWhenReady([this, Entry]()
+		Entry->HashTask = FFunctionGraphTask::CreateAndDispatchWhenReady([this, Entry, TestHash, bTestHashValid]()
 		{
 			TRACE_CPUPROFILER_EVENT_SCOPE(HashChunk);
 			const FIoBuffer* SourceBuffer = Entry->Request->GetSourceBuffer();
 			Entry->ChunkHash = FIoChunkHash::HashBuffer(SourceBuffer->Data(), SourceBuffer->DataSize());
 			WriterContext->HashedChunksCount.IncrementExchange();
+
+			if (bVerifyHashDatabase && bTestHashValid)
+			{
+				if (TestHash != Entry->ChunkHash)
+				{
+					UE_LOG(LogIoStore, Warning, TEXT("HashDb Validation Failed: ChunkId %s has mismatching hash"), *LexToString(Entry->ChunkId));
+				}
+			}
+
 			// Release the source data buffer, it will be reloaded later when we start compressing the chunk
 			Entry->Request->FreeSourceBuffer();
 		}, TStatId(), &HashPrereqs, ENamedThreads::AnyHiPriThreadHiPriTask);
@@ -1779,6 +1828,8 @@ private:
 	bool						bHasFlushed = false;
 	bool						bHasResult = false;
 	TSharedPtr<IIoStoreWriterReferenceChunkDatabase> ReferenceChunkDatabase;
+	TSharedPtr<IIoStoreWriterHashDatabase> HashDatabase;
+	bool						bVerifyHashDatabase = false;
 
 
 	friend class FIoStoreWriterContextImpl;
@@ -1839,9 +1890,11 @@ void FIoStoreWriterContextImpl::Flush()
 
 	// Classically there were so few writers that this didn't need to be multi threaded, but it
 	// involves writing files, and with content on demand this ends up being thousands of iterations. 
-	// 10x speedup in testing.
 	double FinalizeStart = FPlatformTime::Seconds();	
-	ParallelFor(TEXT("IoStoreWriter::Finalize.PF"), IoStoreWriters.Num(), 1, [this](int Index){ IoStoreWriters[Index]->Finalize(); });
+	ParallelFor(TEXT("IoStoreWriter::Finalize.PF"), IoStoreWriters.Num(), 1, [this](int Index)
+	{ 
+		IoStoreWriters[Index]->Finalize(); 
+	});
 	double FinalizeEnd = FPlatformTime::Seconds();
 	UE_LOG(LogIoStore, Display, TEXT("Finalize took %.1f seconds for %d writers"), FinalizeEnd - FinalizeStart, IoStoreWriters.Num());
 }
diff --git a/Engine/Source/Runtime/Core/Public/IO/IoDispatcher.h b/Engine/Source/Runtime/Core/Public/IO/IoDispatcher.h
--- a/Engine/Source/Runtime/Core/Public/IO/IoDispatcher.h
+++ b/Engine/Source/Runtime/Core/Public/IO/IoDispatcher.h
@@ -645,15 +645,19 @@ public:
 		return BytesToHex(Hash, 20);
 	}
 
-	static FIoChunkHash HashBuffer(const void* Data, uint64 DataSize)
+	static FIoChunkHash CreateFromIoHash(const FIoHash& IoHash)
 	{
 		FIoChunkHash Result;
-		FIoHash IoHash = FIoHash::HashBuffer(Data, DataSize);
 		FMemory::Memcpy(Result.Hash, &IoHash, sizeof IoHash);
 		FMemory::Memset(Result.Hash + 20, 0, 12);
 		return Result;
 	}
 
+	static FIoChunkHash HashBuffer(const void* Data, uint64 DataSize)
+	{
+		return CreateFromIoHash(FIoHash::HashBuffer(Data, DataSize));
+	}
+
 private:
 	uint8	Hash[32];
 };
@@ -1041,6 +1045,8 @@ public:
 	{
 		uint64 TotalChunksCount = 0;
 		uint64 HashedChunksCount = 0;
+		// Number of chunks where we avoided reading and hashing, and instead used the result from the hashdb
+		uint64 HashDbChunksCount = 0;
 		uint64 CompressedChunksCount = 0;
 		uint64 SerializedChunksCount = 0;
 		uint64 ScheduledCompressionTasksCount = 0;
@@ -1149,6 +1155,19 @@ public:
 	CORE_API virtual bool RetrieveChunk(const TPair<FIoContainerId, FIoChunkHash>& InChunkKey, const FName& InCompressionMethod, uint64 InUncompressedSize, uint64 InNumChunkBlocks, TUniqueFunction<void(TIoStatusOr<FIoStoreCompressedReadResult>)> InCompletionCallback) = 0;
 };
 
+/**
+*	Allows the IIoStoreWriter to avoid loading and hashing chunks, saving pak/stage time, as the normal
+*	process involved loading the chunks, hashing them, freeing them, making some decisions, then loading
+*	them _again_ for compression/writting. It's completely fine for this to not have all available hashes,
+*	but they have to match when provided!
+*/
+class IIoStoreWriterHashDatabase
+{
+public:
+	CORE_API virtual ~IIoStoreWriterHashDatabase() = default;
+	CORE_API virtual bool FindHashForChunkId(const FIoChunkId& ChunkId, FIoChunkHash& OutHash) const = 0;
+};
+
 
 class IIoStoreWriter
 {
@@ -1160,6 +1179,7 @@ public:
 	*	from previous containers instead of recompressing input data. This must be set before any writes are appended.
 	*/
 	CORE_API virtual void SetReferenceChunkDatabase(TSharedPtr<IIoStoreWriterReferenceChunkDatabase> ReferenceChunkDatabase) = 0;
+	CORE_API virtual void SetHashDatabase(TSharedPtr<IIoStoreWriterHashDatabase> HashDatabase, bool bVerifyHashDatabase) = 0;
 	CORE_API virtual void EnableDiskLayoutOrdering(const TArray<TUniquePtr<FIoStoreReader>>& PatchSourceReaders = TArray<TUniquePtr<FIoStoreReader>>()) = 0;
 	CORE_API virtual void Append(const FIoChunkId& ChunkId, FIoBuffer Chunk, const FIoWriteOptions& WriteOptions, uint64 OrderHint = MAX_uint64) = 0;
 	CORE_API virtual void Append(const FIoChunkId& ChunkId, IIoStoreWriteRequest* Request, const FIoWriteOptions& WriteOptions) = 0;
