Missing asset registry for hash database is now a warning, not an error.

diff --git a/Engine/Source/Developer/IoStoreUtilities/Private/IoStoreUtilities.cpp b/Engine/Source/Developer/IoStoreUtilities/Private/IoStoreUtilities.cpp
--- a/Engine/Source/Developer/IoStoreUtilities/Private/IoStoreUtilities.cpp
+++ b/Engine/Source/Developer/IoStoreUtilities/Private/IoStoreUtilities.cpp
@@ -3161,7 +3161,7 @@ int32 DoAssetRegistryWritebackAfterStage(const FString& InAssetRegistryFileName,
 	return SaveAssetRegistry(InAssetRegistryFileName, AssetRegistry, true) ? 0 : 1;
 }
 
-static bool FindAndLoadDevelopmentAssetRegistry(const FString& InCookedDir, FAssetRegistryState& OutAssetRegistry, FString* OutAssetRegistryFileName /*optional, set on success*/)
+static bool FindAndLoadDevelopmentAssetRegistry(const FString& InCookedDir, bool bInRequired, FAssetRegistryState& OutAssetRegistry, FString* OutAssetRegistryFileName /*optional, set on success*/)
 {
 	TRACE_CPUPROFILER_EVENT_SCOPE(LoadingAssetRegistry);
 
@@ -3180,7 +3180,14 @@ static bool FindAndLoadDevelopmentAssetRegistry(const FString& InCookedDir, FAss
 
 	if (PossibleAssetRegistryFiles.Num() == 0)
 	{
-		UE_LOG(LogIoStore, Error, TEXT("No development asset registry file found!"));
+		if (bInRequired)
+		{
+			UE_LOG(LogIoStore, Error, TEXT("No development asset registry file found!"));
+		}
+		else
+		{
+			UE_LOG(LogIoStore, Display, TEXT("No development asset registry file found!"));
+		}
 		return false;
 	}
 
@@ -3207,7 +3214,7 @@ bool DoAssetRegistryWritebackDuringStage(EAssetRegistryWritebackMethod InMethod,
 	// The overwhelming majority of time for the asset registry writeback is loading and saving.
 	FString AssetRegistryFileName;
 	FAssetRegistryState AssetRegistry;
-	if (FindAndLoadDevelopmentAssetRegistry(InCookedDir, AssetRegistry, &AssetRegistryFileName) == false)
+	if (FindAndLoadDevelopmentAssetRegistry(InCookedDir, true, AssetRegistry, &AssetRegistryFileName) == false)
 	{
 		// already logged
 		return false;
@@ -3274,7 +3281,7 @@ public:
 	{
 		FString AssetRegistryFileName;
 		FAssetRegistryState AssetRegistry;
-		if (FindAndLoadDevelopmentAssetRegistry(InCookedDir, AssetRegistry, nullptr) == false)
+		if (FindAndLoadDevelopmentAssetRegistry(InCookedDir, false, AssetRegistry, nullptr) == false)
 		{
 			// already logged
 			return false;
