MPCook - PackageWriter and IOStore API changes to support MPCook.
Add FCookInfo argument to BeginCook and EndCook, the same data passed to initialize. It is provided in Begin/End as well since Writers need to consider it for Begin/End actions.
Add WriteMPCookMessageForPackage/TryReadMPCookMessageForPackage functions for replicating information from PackageWriters from CookWorker to CookDirector.
Add WritePackage/TryReadPackage to Manifest. Some extra information is collected when MPCook is active to support this replication.
Add CompactBinary serialization functions to FIoChunkId.

diff --git a/Engine/Source/Developer/IoStoreUtilities/Private/IoStoreUtilities.cpp b/Engine/Source/Developer/IoStoreUtilities/Private/IoStoreUtilities.cpp
--- a/Engine/Source/Developer/IoStoreUtilities/Private/IoStoreUtilities.cpp
+++ b/Engine/Source/Developer/IoStoreUtilities/Private/IoStoreUtilities.cpp
@@ -5794,7 +5794,7 @@ int32 Staged2Zen(const FString& BuildPath, const FKeyChain& KeyChain, const FStr
 	ICookedPackageWriter::FCookInfo CookInfo;
 	CookInfo.bFullBuild = true;
 	ZenStoreWriter->Initialize(CookInfo);
-	ZenStoreWriter->BeginCook();
+	ZenStoreWriter->BeginCook(CookInfo);
 	int32 LocalPackageIndex = 0;
 	TArray<FPackageInfo> PackagesArray;
 	CollectedData.Packages.GenerateValueArray(PackagesArray);
@@ -5838,7 +5838,7 @@ int32 Staged2Zen(const FString& BuildPath, const FKeyChain& KeyChain, const FStr
 	}, EParallelForFlags::ForceSingleThread); // Single threaded for now to limit memory usage
 
 	UE_LOG(LogIoStore, Display, TEXT("Waiting for uploads to finish..."));
-	ZenStoreWriter->EndCook();
+	ZenStoreWriter->EndCook(CookInfo);
 	return 0;
 }
 
diff --git a/Engine/Source/Developer/IoStoreUtilities/Private/PackageStoreManifest.cpp b/Engine/Source/Developer/IoStoreUtilities/Private/PackageStoreManifest.cpp
--- a/Engine/Source/Developer/IoStoreUtilities/Private/PackageStoreManifest.cpp
+++ b/Engine/Source/Developer/IoStoreUtilities/Private/PackageStoreManifest.cpp
@@ -1,11 +1,15 @@
 // Copyright Epic Games, Inc. All Rights Reserved.
 
 #include "PackageStoreManifest.h"
+#include "Serialization/CompactBinaryContainerSerialization.h"
+#include "Serialization/CompactBinarySerialization.h"
+#include "Serialization/CompactBinaryWriter.h"
 #include "Serialization/JsonWriter.h"
 #include "Serialization/JsonReader.h"
 #include "Serialization/JsonSerializer.h"
 #include "Dom/JsonObject.h"
 #include "HAL/LowLevelMemTracker.h"
+#include "IO/IoStore.h"
 #include "Misc/FileHelper.h"
 #include "Misc/ScopeLock.h"
 #include "Misc/Paths.h"
@@ -44,7 +48,14 @@ void FPackageStoreManifest::AddPackageData(FName PackageName, const FString& Fil
 	PackageInfo->ExportBundleChunkIds.Add(ChunkId);
 	if (!FileName.IsEmpty())
 	{
-		FileNameByChunkIdMap.Add(ChunkId, FileName);
+		if (!bTrackPackageData)
+		{
+			FileNameByChunkIdMap.Add(ChunkId, FileName);
+		}
+		else
+		{
+			PackageFileChunkIds.FindOrAdd(PackageName).Emplace(FileName, ChunkId);
+		}
 	}
 }
 
@@ -57,7 +68,14 @@ void FPackageStoreManifest::AddBulkData(FName PackageName, const FString& FileNa
 	PackageInfo->BulkDataChunkIds.Add(ChunkId);
 	if (!FileName.IsEmpty())
 	{
-		FileNameByChunkIdMap.Add(ChunkId, FileName);
+		if (!bTrackPackageData)
+		{
+			FileNameByChunkIdMap.Add(ChunkId, FileName);
+		}
+		else
+		{
+			PackageFileChunkIds.FindOrAdd(PackageName).Emplace(FileName, ChunkId);
+		}
 	}
 }
 
@@ -292,3 +310,55 @@ FPackageStoreManifest::FPackageInfo* FPackageStoreManifest::GetPackageInfo_NoLoc
 	return PackageInfoByNameMap.Find(PackageName);
 }
 
+void FPackageStoreManifest::SetTrackPackageData(bool bInTrackPackageData)
+{
+	FScopeLock Lock(&CriticalSection);
+	bTrackPackageData = bInTrackPackageData;
+}
+
+void FPackageStoreManifest::WritePackage(FCbWriter& Writer, FName PackageName)
+{
+	FPackageStoreManifest::FPackageInfo PackageInfo;
+	TArray<TPair<FString, FIoChunkId>> FileChunkIds;
+	bool bHasPackageInfo;
+
+	{
+		FScopeLock Lock(&CriticalSection);
+		bHasPackageInfo = PackageInfoByNameMap.RemoveAndCopyValue(PackageName, PackageInfo);
+		PackageFileChunkIds.RemoveAndCopyValue(PackageName, FileChunkIds);
+	}
+
+	UE_CLOG(!bHasPackageInfo, LogIoStore, Error,
+		TEXT("ClientTickPackage was called for package %s, but CommitPackage was not called for that package. We will not have up to date PackageStoreManifest data."),
+		*PackageName.ToString());
+
+	Writer.BeginObject();
+	Writer << "ExportBundleChunkIds" << PackageInfo.ExportBundleChunkIds;
+	Writer << "BulkDataChunkIds" << PackageInfo.BulkDataChunkIds;
+	Writer << "FileChunkIds" << FileChunkIds;
+	Writer.EndObject();
+}
+
+bool FPackageStoreManifest::TryReadPackage(FCbFieldView Field, FName PackageName)
+{
+	FPackageInfo PackageInfo;
+	TArray<TPair<FString, FIoChunkId>> FileChunkIds;
+
+	bool bOk = true;
+	PackageInfo.PackageName = PackageName;
+	bOk = LoadFromCompactBinary(Field["ExportBundleChunkIds"], PackageInfo.ExportBundleChunkIds) & bOk;
+	bOk = LoadFromCompactBinary(Field["BulkDataChunkIds"], PackageInfo.BulkDataChunkIds) & bOk;
+	bOk = LoadFromCompactBinary(Field["FileChunkIds"], FileChunkIds) & bOk;
+
+	{
+		FScopeLock Lock(&CriticalSection);
+		PackageInfoByNameMap.Add(PackageName, MoveTemp(PackageInfo));
+		for (TPair<FString, FIoChunkId>& Pair : FileChunkIds)
+		{
+			FileNameByChunkIdMap.Add(Pair.Value, Pair.Key);
+		}
+	}
+
+	return bOk;
+}
+
diff --git a/Engine/Source/Developer/IoStoreUtilities/Private/ZenStoreWriter.cpp b/Engine/Source/Developer/IoStoreUtilities/Private/ZenStoreWriter.cpp
--- a/Engine/Source/Developer/IoStoreUtilities/Private/ZenStoreWriter.cpp
+++ b/Engine/Source/Developer/IoStoreUtilities/Private/ZenStoreWriter.cpp
@@ -495,9 +495,16 @@ void FZenStoreWriter::Initialize(const FCookInfo& Info)
 	}
 }
 
-void FZenStoreWriter::BeginCook()
+void FZenStoreWriter::BeginCook(const FCookInfo& Info)
 {
-	PackageStoreManifest.Load(*(MetadataDirectoryPath / TEXT("packagestore.manifest")));
+	if (!Info.bWorkerOnSharedSandbox)
+	{
+		PackageStoreManifest.Load(*(MetadataDirectoryPath / TEXT("packagestore.manifest")));
+	}
+	else
+	{
+		PackageStoreManifest.SetTrackPackageData(true);
+	}
 	AllPackageHashes.Empty();
 
 	if (CookMode == ICookedPackageWriter::FCookInfo::CookOnTheFlyMode)
@@ -538,7 +545,7 @@ void FZenStoreWriter::BeginCook()
 	}
 }
 
-void FZenStoreWriter::EndCook()
+void FZenStoreWriter::EndCook(const FCookInfo& Info)
 {
 	UE_LOG(LogZenStoreWriter, Display, TEXT("Flushing..."));
 	
@@ -561,7 +568,10 @@ void FZenStoreWriter::EndCook()
 	TIoStatusOr<uint64> Status = HttpClient->EndBuildPass(Pkg);
 	UE_CLOG(!Status.IsOk(), LogZenStoreWriter, Fatal, TEXT("Failed to append OpLog and end the build pass"));
 
-	PackageStoreManifest.Save(*(MetadataDirectoryPath / TEXT("packagestore.manifest")));
+	if (!Info.bWorkerOnSharedSandbox)
+	{
+		PackageStoreManifest.Save(*(MetadataDirectoryPath / TEXT("packagestore.manifest")));
+	}
 
 	UE_LOG(LogZenStoreWriter, Display, TEXT("Input:\t%d Packages"), PackageStoreOptimizer->GetTotalPackageCount());
 	UE_LOG(LogZenStoreWriter, Display, TEXT("Output:\t%d Export bundles"), PackageStoreOptimizer->GetTotalExportBundleCount());
@@ -1040,6 +1050,21 @@ void FZenStoreWriter::MarkPackagesUpToDate(TArrayView<const FName> UpToDatePacka
 	}
 }
 
+FCbObject FZenStoreWriter::WriteMPCookMessageForPackage(FName PackageName)
+{
+	FCbWriter Writer;
+	Writer.BeginObject();
+	Writer.SetName("Manifest");
+	PackageStoreManifest.WritePackage(Writer, PackageName);
+	Writer.EndObject();
+	return Writer.Save().AsObject();
+}
+
+bool FZenStoreWriter::TryReadMPCookMessageForPackage(FName PackageName, FCbObjectView Message)
+{
+	return PackageStoreManifest.TryReadPackage(Message["Manifest"], PackageName);
+}
+
 void FZenStoreWriter::CreateProjectMetaData(FCbPackage& Pkg, FCbWriter& PackageObj)
 {
 	// File Manifest
diff --git a/Engine/Source/Developer/IoStoreUtilities/Public/PackageStoreManifest.h b/Engine/Source/Developer/IoStoreUtilities/Public/PackageStoreManifest.h
--- a/Engine/Source/Developer/IoStoreUtilities/Public/PackageStoreManifest.h
+++ b/Engine/Source/Developer/IoStoreUtilities/Public/PackageStoreManifest.h
@@ -14,6 +14,9 @@
 #include "Templates/UniquePtr.h"
 #include "UObject/NameTypes.h"
 
+class FCbFieldView;
+class FCbWriter;
+
 class FPackageStoreManifest
 {
 public:
@@ -53,6 +56,11 @@ public:
 	IOSTOREUTILITIES_API FZenServerInfo& EditZenServerInfo();
 	IOSTOREUTILITIES_API const FZenServerInfo* ReadZenServerInfo() const;
 
+	IOSTOREUTILITIES_API void SetTrackPackageData(bool bInTrackpackageData);
+	IOSTOREUTILITIES_API void WritePackage(FCbWriter& Writer, FName PackageName);
+	IOSTOREUTILITIES_API bool TryReadPackage(FCbFieldView Field, FName PackageName);
+
+
 private:
 	FPackageInfo* GetPackageInfo_NoLock(FName PackageName);
 
@@ -62,4 +70,11 @@ private:
 
 	TMap<FIoChunkId, FString> FileNameByChunkIdMap;
 	TUniquePtr<FZenServerInfo> ZenServerInfo;
+
+	/**
+	 * Transient data used during MPCook to find all Filenames used by a package. Used to replicate all data related to
+	 * the package to the CookDirector.
+	 */
+	TMap<FName, TArray<TPair<FString, FIoChunkId>>> PackageFileChunkIds;
+	bool bTrackPackageData = false;
 };
diff --git a/Engine/Source/Developer/IoStoreUtilities/Public/ZenStoreWriter.h b/Engine/Source/Developer/IoStoreUtilities/Public/ZenStoreWriter.h
--- a/Engine/Source/Developer/IoStoreUtilities/Public/ZenStoreWriter.h
+++ b/Engine/Source/Developer/IoStoreUtilities/Public/ZenStoreWriter.h
@@ -72,8 +72,8 @@ public:
 
 	IOSTOREUTILITIES_API virtual void WriteBulkData(const FBulkDataInfo& Info, const FIoBuffer& BulkData, const TArray<FFileRegion>& FileRegions) override;
 	IOSTOREUTILITIES_API virtual void Initialize(const FCookInfo& Info) override;
-	IOSTOREUTILITIES_API virtual void BeginCook() override;
-	IOSTOREUTILITIES_API virtual void EndCook() override;
+	IOSTOREUTILITIES_API virtual void BeginCook(const FCookInfo& Info) override;
+	IOSTOREUTILITIES_API virtual void EndCook(const FCookInfo& Info) override;
 
 	IOSTOREUTILITIES_API virtual void GetEntries(TFunction<void(TArrayView<const FPackageStoreEntryResource>, TArrayView<const FOplogCookInfo>)>&& Callback) override;
 
@@ -100,6 +100,9 @@ public:
 	IOSTOREUTILITIES_API virtual void RemoveCookedPackages(TArrayView<const FName> PackageNamesToRemove) override;
 	IOSTOREUTILITIES_API virtual void RemoveCookedPackages() override;
 	IOSTOREUTILITIES_API virtual void MarkPackagesUpToDate(TArrayView<const FName> UpToDatePackages) override;
+	IOSTOREUTILITIES_API virtual FCbObject WriteMPCookMessageForPackage(FName PackageName) override;
+	IOSTOREUTILITIES_API virtual bool TryReadMPCookMessageForPackage(FName PackageName, FCbObjectView Message) override;
+
 	IOSTOREUTILITIES_API virtual TMap<FName, TRefCountPtr<FPackageHashes>>& GetPackageHashes() override
 	{
 		return AllPackageHashes;
diff --git a/Engine/Source/Editor/UnrealEd/Private/CookOnTheFlyServer.cpp b/Engine/Source/Editor/UnrealEd/Private/CookOnTheFlyServer.cpp
--- a/Engine/Source/Editor/UnrealEd/Private/CookOnTheFlyServer.cpp
+++ b/Engine/Source/Editor/UnrealEd/Private/CookOnTheFlyServer.cpp
@@ -9157,6 +9157,8 @@ void UCookOnTheFlyServer::LoadBeginCookIterativeFlagsLocal(FBeginCookContext& Be
 			}
 		}
 		PlatformData->bFullBuild = PlatformContext.bFullBuild;
+		PlatformData->bIterateSharedBuild = PlatformContext.bIterateSharedBuild;
+		PlatformData->bWorkerOnSharedSandbox = PlatformContext.bWorkerOnSharedSandbox;
 	}
 }
 
@@ -9339,7 +9341,14 @@ void UCookOnTheFlyServer::FinalizePackageStore()
 	UE_LOG(LogCook, Display, TEXT("Finalize package store(s)..."));
 	for (const ITargetPlatform* TargetPlatform : PlatformManager->GetSessionPlatforms())
 	{
-		FindOrCreatePackageWriter(TargetPlatform).EndCook();
+		UE::Cook::FPlatformData* PlatformData = PlatformManager->GetPlatformData(TargetPlatform);
+		ICookedPackageWriter::FCookInfo CookInfo;
+		CookInfo.CookMode = IsDirectorCookOnTheFly() ? ICookedPackageWriter::FCookInfo::CookOnTheFlyMode : ICookedPackageWriter::FCookInfo::CookByTheBookMode;
+		CookInfo.bFullBuild = PlatformData->bFullBuild;
+		CookInfo.bIterateSharedBuild = PlatformData->bIterateSharedBuild;
+		CookInfo.bWorkerOnSharedSandbox = PlatformData->bWorkerOnSharedSandbox;
+
+		FindOrCreatePackageWriter(TargetPlatform).EndCook(CookInfo);
 	}
 	UE_LOG(LogCook, Display, TEXT("Done finalizing package store(s)"));
 }
@@ -10024,9 +10033,15 @@ void UCookOnTheFlyServer::RecordDLCPackagesFromBaseGame(FBeginCookContext& Begin
 
 void UCookOnTheFlyServer::BeginCookPackageWriters(FBeginCookContext& BeginContext)
 {
-	for (const ITargetPlatform* TargetPlatform : BeginContext.TargetPlatforms)
+	for (FBeginCookContextPlatform& Context : BeginContext.PlatformContexts)
 	{
-		FindOrCreatePackageWriter(TargetPlatform).BeginCook();
+		ICookedPackageWriter::FCookInfo CookInfo;
+		CookInfo.CookMode = IsDirectorCookOnTheFly() ? ICookedPackageWriter::FCookInfo::CookOnTheFlyMode : ICookedPackageWriter::FCookInfo::CookByTheBookMode;
+		CookInfo.bFullBuild = Context.bFullBuild;
+		CookInfo.bIterateSharedBuild = Context.bIterateSharedBuild;
+		CookInfo.bWorkerOnSharedSandbox = Context.bWorkerOnSharedSandbox;
+
+		FindOrCreatePackageWriter(Context.TargetPlatform).BeginCook(CookInfo);
 	}
 }
 
diff --git a/Engine/Source/Editor/UnrealEd/Private/Cooker/CookPlatformManager.h b/Engine/Source/Editor/UnrealEd/Private/Cooker/CookPlatformManager.h
--- a/Engine/Source/Editor/UnrealEd/Private/Cooker/CookPlatformManager.h
+++ b/Engine/Source/Editor/UnrealEd/Private/Cooker/CookPlatformManager.h
@@ -62,6 +62,12 @@ namespace Cook
 		/* Whether BeginCookSandbox specified a full, non-iterative build for this platform. */
 		bool bFullBuild = false;
 
+		/** If true we are cooking iteratively, from results in a shared build (e.g. from buildfarm) rather than from our previous cook. */
+		bool bIterateSharedBuild = false;
+
+		/** If true we are a CookWorker, and we are working on a Sandbox directory that has already been populated by a remote Director process. */
+		bool bWorkerOnSharedSandbox = false;
+
 		/*
 		 * The last FPlatformTime::GetSeconds() at which this platform was requested in a CookOnTheFly request.
 		 * If equal to 0, the platform was not requested in a CookOnTheFly since the last clear.
diff --git a/Engine/Source/Editor/UnrealEd/Private/Cooker/DiffPackageWriter.h b/Engine/Source/Editor/UnrealEd/Private/Cooker/DiffPackageWriter.h
--- a/Engine/Source/Editor/UnrealEd/Private/Cooker/DiffPackageWriter.h
+++ b/Engine/Source/Editor/UnrealEd/Private/Cooker/DiffPackageWriter.h
@@ -67,13 +67,13 @@ public:
 	{
 		Inner->Initialize(Info);
 	}
-	virtual void BeginCook() override
+	virtual void BeginCook(const FCookInfo& Info) override
 	{
-		Inner->BeginCook();
+		Inner->BeginCook(Info);
 	}
-	virtual void EndCook() override
+	virtual void EndCook(const FCookInfo& Info) override
 	{
-		Inner->EndCook();
+		Inner->EndCook(Info);
 	}
 	virtual TUniquePtr<FAssetRegistryState> LoadPreviousAssetRegistry() override
 	{
@@ -100,6 +100,14 @@ public:
 		Inner->UpdateSaveArguments(SaveArgs);
 	}
 	virtual bool IsAnotherSaveNeeded(FSavePackageResultStruct& PreviousResult, FSavePackageArgs& SaveArgs) override;
+	virtual FCbObject WriteMPCookMessageForPackage(FName PackageName) override
+	{
+		return Inner->WriteMPCookMessageForPackage(PackageName);
+	}
+	virtual bool TryReadMPCookMessageForPackage(FName PackageName, FCbObjectView Message) override
+	{
+		return Inner->TryReadMPCookMessageForPackage(PackageName, Message);
+	}
 	virtual TMap<FName, TRefCountPtr<FPackageHashes>>& GetPackageHashes() override
 	{
 		return Inner->GetPackageHashes();
@@ -198,13 +206,13 @@ public:
 	{
 		Inner->Initialize(Info);
 	}
-	virtual void BeginCook() override
+	virtual void BeginCook(const FCookInfo& Info) override
 	{
-		Inner->BeginCook();
+		Inner->BeginCook(Info);
 	}
-	virtual void EndCook() override
+	virtual void EndCook(const FCookInfo& Info) override
 	{
-		Inner->EndCook();
+		Inner->EndCook(Info);
 	}
 	virtual TUniquePtr<FAssetRegistryState> LoadPreviousAssetRegistry() override
 	{
@@ -228,6 +236,14 @@ public:
 	}
 	virtual void UpdateSaveArguments(FSavePackageArgs& SaveArgs) override;
 	virtual bool IsAnotherSaveNeeded(FSavePackageResultStruct& PreviousResult, FSavePackageArgs& SaveArgs) override;
+	virtual FCbObject WriteMPCookMessageForPackage(FName PackageName) override
+	{
+		return Inner->WriteMPCookMessageForPackage(PackageName);
+	}
+	virtual bool TryReadMPCookMessageForPackage(FName PackageName, FCbObjectView Message) override
+	{
+		return Inner->TryReadMPCookMessageForPackage(PackageName, Message);
+	}
 	virtual TMap<FName, TRefCountPtr<FPackageHashes>>& GetPackageHashes() override
 	{
 		return Inner->GetPackageHashes();
diff --git a/Engine/Source/Editor/UnrealEd/Private/Cooker/LooseCookedPackageWriter.cpp b/Engine/Source/Editor/UnrealEd/Private/Cooker/LooseCookedPackageWriter.cpp
--- a/Engine/Source/Editor/UnrealEd/Private/Cooker/LooseCookedPackageWriter.cpp
+++ b/Engine/Source/Editor/UnrealEd/Private/Cooker/LooseCookedPackageWriter.cpp
@@ -41,6 +41,7 @@
 #include "Serialization/Archive.h"
 #include "Serialization/ArchiveStackTrace.h"
 #include "Serialization/ArrayReader.h"
+#include "Serialization/CompactBinaryWriter.h"
 #include "Serialization/LargeMemoryWriter.h"
 #include "Serialization/MemoryWriter.h"
 #include "Tasks/Task.h"
@@ -515,15 +516,25 @@ void FLooseCookedPackageWriter::Initialize(const FCookInfo& Info)
 	}
 }
 
-void FLooseCookedPackageWriter::BeginCook()
+void FLooseCookedPackageWriter::BeginCook(const FCookInfo& Info)
 {
-	PackageStoreManifest.Load(*(MetadataDirectoryPath / TEXT("packagestore.manifest")));
+	if (!Info.bWorkerOnSharedSandbox)
+	{
+		PackageStoreManifest.Load(*(MetadataDirectoryPath / TEXT("packagestore.manifest")));
+	}
+	else
+	{
+		PackageStoreManifest.SetTrackPackageData(true);
+	}
 	AllPackageHashes.Empty();
 }
 
-void FLooseCookedPackageWriter::EndCook()
+void FLooseCookedPackageWriter::EndCook(const FCookInfo& Info)
 {
-	PackageStoreManifest.Save(*(MetadataDirectoryPath / TEXT("packagestore.manifest")));
+	if (!Info.bWorkerOnSharedSandbox)
+	{
+		PackageStoreManifest.Save(*(MetadataDirectoryPath / TEXT("packagestore.manifest")));
+	}
 }
 
 TUniquePtr<FAssetRegistryState> FLooseCookedPackageWriter::LoadPreviousAssetRegistry()
@@ -694,6 +705,21 @@ void FLooseCookedPackageWriter::MarkPackagesUpToDate(TArrayView<const FName> UpT
 {
 }
 
+FCbObject FLooseCookedPackageWriter::WriteMPCookMessageForPackage(FName PackageName)
+{
+	FCbWriter Writer;
+	Writer.BeginObject();
+	Writer.SetName("Manifest");
+	PackageStoreManifest.WritePackage(Writer, PackageName);
+	Writer.EndObject();
+	return Writer.Save().AsObject();
+}
+
+bool FLooseCookedPackageWriter::TryReadMPCookMessageForPackage(FName PackageName, FCbObjectView Message)
+{
+	return PackageStoreManifest.TryReadPackage(Message["Manifest"], PackageName);
+}
+
 void FLooseCookedPackageWriter::RemoveCookedPackages()
 {
 	DeleteSandboxDirectory();
diff --git a/Engine/Source/Editor/UnrealEd/Private/Cooker/LooseCookedPackageWriter.h b/Engine/Source/Editor/UnrealEd/Private/Cooker/LooseCookedPackageWriter.h
--- a/Engine/Source/Editor/UnrealEd/Private/Cooker/LooseCookedPackageWriter.h
+++ b/Engine/Source/Editor/UnrealEd/Private/Cooker/LooseCookedPackageWriter.h
@@ -55,13 +55,15 @@ public:
 
 	virtual FDateTime GetPreviousCookTime() const override;
 	virtual void Initialize(const FCookInfo& Info) override;
-	virtual void BeginCook() override;
-	virtual void EndCook() override;
+	virtual void BeginCook(const FCookInfo& Info) override;
+	virtual void EndCook(const FCookInfo& Info) override;
 	virtual TUniquePtr<FAssetRegistryState> LoadPreviousAssetRegistry() override;
 	virtual FCbObject GetOplogAttachment(FName PackageName, FUtf8StringView AttachmentKey) override;
 	virtual void RemoveCookedPackages(TArrayView<const FName> PackageNamesToRemove) override;
 	virtual void RemoveCookedPackages() override;
 	virtual void MarkPackagesUpToDate(TArrayView<const FName> UpToDatePackages) override;
+	virtual FCbObject WriteMPCookMessageForPackage(FName PackageName) override;
+	virtual bool TryReadMPCookMessageForPackage(FName PackageName, FCbObjectView Message) override;
 	virtual bool GetPreviousCookedBytes(const FPackageInfo& Info, FPreviousCookedBytesData& OutData) override;
 	virtual void CompleteExportsArchiveForDiff(const FPackageInfo& Info, FLargeMemoryWriter& ExportsArchive) override;
 	virtual void CommitPackageInternal(FPackageWriterRecords::FPackage&& BaseRecord,
diff --git a/Engine/Source/Editor/UnrealEd/Private/Cooker/WorkerRequestsRemote.cpp b/Engine/Source/Editor/UnrealEd/Private/Cooker/WorkerRequestsRemote.cpp
--- a/Engine/Source/Editor/UnrealEd/Private/Cooker/WorkerRequestsRemote.cpp
+++ b/Engine/Source/Editor/UnrealEd/Private/Cooker/WorkerRequestsRemote.cpp
@@ -157,6 +157,8 @@ void FWorkerRequestsRemote::GetBeginCookIterativeFlags(UCookOnTheFlyServer& COTF
 		PlatformContext.bIterateSharedBuild = false; // Not needed on CookWorkerClients
 		PlatformContext.bWorkerOnSharedSandbox = true;
 		PlatformData->bFullBuild = PlatformContext.bFullBuild;
+		PlatformData->bIterateSharedBuild = PlatformContext.bIterateSharedBuild;
+		PlatformData->bWorkerOnSharedSandbox = PlatformContext.bWorkerOnSharedSandbox;
 	}
 }
 
diff --git a/Engine/Source/Runtime/Core/Private/IO/IoChunkId.cpp b/Engine/Source/Runtime/Core/Private/IO/IoChunkId.cpp
--- a/Engine/Source/Runtime/Core/Private/IO/IoChunkId.cpp
+++ b/Engine/Source/Runtime/Core/Private/IO/IoChunkId.cpp
@@ -3,7 +3,9 @@
 #include "IO/IoChunkId.h"
 #include "IO/PackageId.h"
 #include "Hash/Blake3.h"
+#include "Memory/MemoryView.h"
 #include "Serialization/Archive.h"
+#include "Serialization/CompactBinaryWriter.h"
 
 FString LexToString(const EIoChunkType Type)
 {
@@ -52,6 +54,19 @@ FArchive& operator<<(FArchive& Ar, FIoChunkId& ChunkId)
 	return Ar;
 }
 
+FCbWriter& operator<<(FCbWriter& Writer, const FIoChunkId& ChunkId)
+{
+	Writer.AddObjectId(FCbObjectId(MakeMemoryView(&ChunkId.Id, sizeof FIoChunkId::Id)));
+	return Writer;
+}
+
+bool LoadFromCompactBinary(FCbFieldView Field, FIoChunkId& OutChunkId)
+{
+	FCbObjectId ObjectId = Field.AsObjectId();
+	OutChunkId.Set(ObjectId.GetView());
+	return !Field.HasError();
+}
+
 FIoChunkId CreatePackageDataChunkId(const FPackageId& PackageId)
 {
 	return CreateIoChunkId(PackageId.Value(), 0, EIoChunkType::ExportBundleData);
diff --git a/Engine/Source/Runtime/Core/Public/IO/IoChunkId.h b/Engine/Source/Runtime/Core/Public/IO/IoChunkId.h
--- a/Engine/Source/Runtime/Core/Public/IO/IoChunkId.h
+++ b/Engine/Source/Runtime/Core/Public/IO/IoChunkId.h
@@ -9,6 +9,8 @@
 #include "String/BytesToHex.h"
 
 class FArchive;
+class FCbFieldView;
+class FCbWriter;
 class FPackageId;
 
 /**
@@ -61,6 +63,8 @@ public:
 	}
 
 	friend CORE_API FArchive& operator<<(FArchive& Ar, FIoChunkId& ChunkId);
+	friend CORE_API FCbWriter& operator<<(FCbWriter& Writer, const FIoChunkId& ChunkId);
+	friend CORE_API bool LoadFromCompactBinary(FCbFieldView Field, FIoChunkId& OutChunkId);
 
 	template <typename CharType>
 	friend TStringBuilderBase<CharType>& operator<<(TStringBuilderBase<CharType>& Builder, const FIoChunkId& ChunkId)
diff --git a/Engine/Source/Runtime/Core/Public/Serialization/PackageWriter.h b/Engine/Source/Runtime/Core/Public/Serialization/PackageWriter.h
--- a/Engine/Source/Runtime/Core/Public/Serialization/PackageWriter.h
+++ b/Engine/Source/Runtime/Core/Public/Serialization/PackageWriter.h
@@ -15,6 +15,8 @@
 #include "Templates/UniquePtr.h"
 
 class FAssetRegistryState;
+class FCbObject;
+class FCbObjectView;
 class FLargeMemoryWriter;
 class ICookedPackageWriter;
 class IPackageStoreWriter;
@@ -255,11 +257,11 @@ public:
 		Package data may only be produced after BeginCook() has been called and
 		before EndCook() is called
 	  */
-	virtual void BeginCook() = 0;
+	virtual void BeginCook(const FCookInfo& Info) = 0;
 
 	/** Signal the end of a cooking pass.
 	  */
-	virtual void EndCook() = 0;
+	virtual void EndCook(const FCookInfo& Info) = 0;
 
 	struct FCookedPackageInfo
 	{
@@ -326,6 +328,17 @@ public:
 	{
 		return false;
 	}
+	/**
+	 * Create a CompactBinary Object message that replicates all of the package data from package save that is
+	 * collected in memory and written at end of cook rather than being written to disk during package save.
+	 * Used during MPCook to transfer this information from CookWorker to CookDirector. Called after CommitPackage,
+	 * and only on CookWorkers.
+	 */
+	virtual FCbObject WriteMPCookMessageForPackage(FName PackageName) = 0;
+
+	/** Read PackageData written by WriteMPCookMessageForPackage on a CookWorker. Called only on CookDirector. */
+	virtual bool TryReadMPCookMessageForPackage(FName PackageName, FCbObjectView Message) = 0;
+
 	/** Downcast function for ICookedPackageWriters that implement the IPackageStoreWriter inherited interface. */
 	virtual IPackageStoreWriter* AsPackageStoreWriter()
 	{
