MPCook: Sort the output of IoStore PackageStoreManifest when writing to disk to make it deterministic.
#rnx

diff --git a/Engine/Source/Developer/IoStoreUtilities/Private/PackageStoreManifest.cpp b/Engine/Source/Developer/IoStoreUtilities/Private/PackageStoreManifest.cpp
--- a/Engine/Source/Developer/IoStoreUtilities/Private/PackageStoreManifest.cpp
+++ b/Engine/Source/Developer/IoStoreUtilities/Private/PackageStoreManifest.cpp
@@ -107,18 +107,35 @@ FIoStatus FPackageStoreManifest::Save(const TCHAR* Filename) const
 	}
 	
 	Writer->WriteArrayStart(TEXT("Files"));
-	for (const auto& KV : FileNameByChunkIdMap)
+
+	// Convert FilePaths in ChunkIdMap to RelativePaths from the CookedOutput folder
+	// Sort by RelativePath for determinism
+	TArray<TPair<const FIoChunkId*, FString>> SortedFileNameByChunkIdMap;
+	SortedFileNameByChunkIdMap.Reserve(FileNameByChunkIdMap.Num());
+	for (const TPair<FIoChunkId, FString>& KV : FileNameByChunkIdMap)
 	{
-		Writer->WriteObjectStart();
 		FString RelativePath = KV.Value;
 		FPaths::MakePathRelativeTo(RelativePath, *CookedOutputPath);
-		Writer->WriteValue(TEXT("Path"), RelativePath);
-		Writer->WriteValue(TEXT("ChunkId"), ChunkIdToString(KV.Key));
+		SortedFileNameByChunkIdMap.Emplace(&KV.Key, MoveTemp(RelativePath));
+	}
+	SortedFileNameByChunkIdMap.Sort([]
+		(const TPair<const FIoChunkId*, FString>& A, const TPair<const FIoChunkId*, FString>& B)
+		{
+			return A.Value < B.Value;
+		});
+	for (const TPair<const FIoChunkId*, FString>& KV : SortedFileNameByChunkIdMap)
+	{
+		Writer->WriteObjectStart();
+		Writer->WriteValue(TEXT("Path"), KV.Value);
+		Writer->WriteValue(TEXT("ChunkId"), ChunkIdToString(*KV.Key));
 		Writer->WriteObjectEnd();
 	}
+	SortedFileNameByChunkIdMap.Empty();
 	Writer->WriteArrayEnd();
 
-	auto WritePackageInfoObject = [Writer, &ChunkIdToString](const FPackageInfo& PackageInfo, const TCHAR* Name = nullptr)
+	constexpr int32 ChunkIdStringsBufferSize = 10;
+	TArray<FString, TInlineAllocator<ChunkIdStringsBufferSize>> ChunkIdStringsBuffer;
+	auto WritePackageInfoObject = [Writer, &ChunkIdStringsBuffer](const FPackageInfo& PackageInfo, const TCHAR* Name = nullptr)
 	{
 		if (Name)
 		{
@@ -128,22 +145,54 @@ FIoStatus FPackageStoreManifest::Save(const TCHAR* Filename) const
 		{
 			Writer->WriteObjectStart();
 		}
+		auto AllocateChunkIdStrings = [&ChunkIdStringsBuffer](int32 Num)
+		{
+			if (ChunkIdStringsBuffer.Num() < Num)
+			{
+				if (ChunkIdStringsBuffer.Max() < Num)
+				{
+					ChunkIdStringsBuffer.SetNum(Num * 2 + ChunkIdStringsBufferSize, false /* bAllowShrinking */);
+				}
+				else
+				{
+					ChunkIdStringsBuffer.SetNum(ChunkIdStringsBuffer.Max(), false /* bAllowShrinking */);
+				}
+			}
+			return TArrayView<FString>(ChunkIdStringsBuffer.GetData(), Num);
+		};
 		Writer->WriteValue(TEXT("Name"), PackageInfo.PackageName.ToString());
 		if (!PackageInfo.ExportBundleChunkIds.IsEmpty())
 		{
-			Writer->WriteArrayStart(TEXT("ExportBundleChunkIds"));
+			// Determinism: Sort ExportBundleChunkIds by string
+			TArrayView<FString> ChunkIdStrings = AllocateChunkIdStrings(PackageInfo.ExportBundleChunkIds.Num());
+			int32 Index = 0;
 			for (const FIoChunkId& ChunkId : PackageInfo.ExportBundleChunkIds)
 			{
-				Writer->WriteValue(ChunkIdToString(ChunkId));
+				ChunkId.ToString(ChunkIdStrings[Index++]);
+			}
+			Algo::Sort(ChunkIdStrings);
+
+			Writer->WriteArrayStart(TEXT("ExportBundleChunkIds"));
+			for (const FString& ChunkIdString : ChunkIdStrings)
+			{
+				Writer->WriteValue(ChunkIdString);
 			}
 			Writer->WriteArrayEnd();
 		}
 		if (!PackageInfo.BulkDataChunkIds.IsEmpty())
 		{
-			Writer->WriteArrayStart(TEXT("BulkDataChunkIds"));
+			// Determinism: Sort BulkDataChunkIds by string
+			TArrayView<FString> ChunkIdStrings = AllocateChunkIdStrings(PackageInfo.BulkDataChunkIds.Num());
+			int32 Index = 0;
 			for (const FIoChunkId& ChunkId : PackageInfo.BulkDataChunkIds)
 			{
-				Writer->WriteValue(ChunkIdToString(ChunkId));
+				ChunkId.ToString(ChunkIdStrings[Index++]);
+			}
+
+			Writer->WriteArrayStart(TEXT("BulkDataChunkIds"));
+			for (const FString& ChunkIdString : ChunkIdStrings)
+			{
+				Writer->WriteValue(ChunkIdString);
 			}
 			Writer->WriteArrayEnd();
 		}
@@ -151,12 +200,24 @@ FIoStatus FPackageStoreManifest::Save(const TCHAR* Filename) const
 	};
 
 	Writer->WriteArrayStart(TEXT("Packages"));
-	for (const auto& PackageNameInfoPair : PackageInfoByNameMap)
+	// Sort PackageInfoByNameMap by PackageName
+	TArray<TPair<FName, const FPackageInfo*>> SortedPackageInfoByNameMap;
+	SortedPackageInfoByNameMap.Reserve(PackageInfoByNameMap.Num());
+	for (const TPair<FName, FPackageInfo>& KV : PackageInfoByNameMap)
+	{
+		SortedPackageInfoByNameMap.Emplace(KV.Key, &KV.Value);
+	}
+	SortedPackageInfoByNameMap.Sort([]
+	(const TPair<FName, const FPackageInfo*>& A, const TPair<FName, const FPackageInfo*>& B)
+		{
+			return A.Key.LexicalLess(B.Key);
+		});
+	for (const TPair<FName, const FPackageInfo*>& KV : SortedPackageInfoByNameMap)
 	{
-		const FPackageInfo& PackageInfo = PackageNameInfoPair.Value;
-		WritePackageInfoObject(PackageInfo);
+		WritePackageInfoObject(*KV.Value);
 	}
 	Writer->WriteArrayEnd();
+	ChunkIdStringsBuffer.Empty();
 
 	Writer->WriteObjectEnd();
 	Writer->Close();
diff --git a/Engine/Source/Runtime/Core/Private/IO/IoChunkId.cpp b/Engine/Source/Runtime/Core/Private/IO/IoChunkId.cpp
--- a/Engine/Source/Runtime/Core/Private/IO/IoChunkId.cpp
+++ b/Engine/Source/Runtime/Core/Private/IO/IoChunkId.cpp
@@ -41,11 +41,17 @@ FString LexToString(const EIoChunkType Type)
 FString LexToString(const FIoChunkId& Id)
 {
 	FString Output;
+	Id.ToString(Output);
+	return Output;
+}
+
+void FIoChunkId::ToString(FString& Output) const
+{
+	Output.Reset();
 	TArray<TCHAR, FString::AllocatorType>& CharArray = Output.GetCharArray();
 	CharArray.AddUninitialized(sizeof(FIoChunkId) * 2 + 1);
-	UE::String::BytesToHexLower(Id.Id, CharArray.GetData());
+	UE::String::BytesToHexLower(Id, CharArray.GetData());
 	CharArray.Last() = TCHAR('\0');
-	return Output;
 }
 
 FArchive& operator<<(FArchive& Ar, FIoChunkId& ChunkId)
diff --git a/Engine/Source/Runtime/Core/Public/IO/IoChunkId.h b/Engine/Source/Runtime/Core/Public/IO/IoChunkId.h
--- a/Engine/Source/Runtime/Core/Public/IO/IoChunkId.h
+++ b/Engine/Source/Runtime/Core/Public/IO/IoChunkId.h
@@ -73,6 +73,7 @@ public:
 		return Builder;
 	}
 
+	CORE_API void ToString(FString& Output) const;
 	friend CORE_API FString LexToString(const FIoChunkId& Id);
 
 	inline bool operator ==(const FIoChunkId& Rhs) const
