PackageStoreManifest: Fix compile error on older compilers that do not allow constexpr variables to be captured by default.
#rnx

diff --git a/Engine/Source/Developer/IoStoreUtilities/Private/PackageStoreManifest.cpp b/Engine/Source/Developer/IoStoreUtilities/Private/PackageStoreManifest.cpp
--- a/Engine/Source/Developer/IoStoreUtilities/Private/PackageStoreManifest.cpp
+++ b/Engine/Source/Developer/IoStoreUtilities/Private/PackageStoreManifest.cpp
@@ -135,7 +135,8 @@ FIoStatus FPackageStoreManifest::Save(const TCHAR* Filename) const
 
 	constexpr int32 ChunkIdStringsBufferSize = 10;
 	TArray<FString, TInlineAllocator<ChunkIdStringsBufferSize>> ChunkIdStringsBuffer;
-	auto WritePackageInfoObject = [Writer, &ChunkIdStringsBuffer](const FPackageInfo& PackageInfo, const TCHAR* Name = nullptr)
+	auto WritePackageInfoObject = [Writer, &ChunkIdStringsBuffer, ChunkIdStringsBufferSize]
+	(const FPackageInfo& PackageInfo, const TCHAR* Name = nullptr)
 	{
 		if (Name)
 		{
@@ -145,7 +146,7 @@ FIoStatus FPackageStoreManifest::Save(const TCHAR* Filename) const
 		{
 			Writer->WriteObjectStart();
 		}
-		auto AllocateChunkIdStrings = [&ChunkIdStringsBuffer](int32 Num)
+		auto AllocateChunkIdStrings = [&ChunkIdStringsBuffer, ChunkIdStringsBufferSize](int32 Num)
 		{
 			if (ChunkIdStringsBuffer.Num() < Num)
 			{
