Make attachment names written as temp-files in ZenStoreHttpClient uses a per-process unique salt

diff --git a/Engine/Source/Developer/IoStoreUtilities/Internal/ZenStoreHttpClient.h b/Engine/Source/Developer/IoStoreUtilities/Internal/ZenStoreHttpClient.h
--- a/Engine/Source/Developer/IoStoreUtilities/Internal/ZenStoreHttpClient.h
+++ b/Engine/Source/Developer/IoStoreUtilities/Internal/ZenStoreHttpClient.h
@@ -69,15 +69,28 @@ private:
 	TIoStatusOr<FIoBuffer> ReadOpLogUri(FStringBuilderBase& ChunkUri, uint64 Offset = 0, uint64 Size = ~0ull);
 
 	static const uint32 PoolEntryCount;
+	struct SaltGenerator
+	{
+		SaltGenerator();
+		inline int32_t Next()
+		{
+			const uint32_t A = ++GOpCounter;
+			return static_cast<int32_t>((A ^ (SaltBase + (A << 6) + (A >> 2))) & 0x7fffffffu);
+		}
+	private:
+		static std::atomic<uint32> GOpCounter;
+		const uint32_t SaltBase;
+	};
 #if UE_WITH_ZEN
 	UE::Zen::FScopeZenService ZenService;
 #endif
 	TUniquePtr<Zen::FZenHttpRequestPool> RequestPool;
+	SaltGenerator SaltGen;
 	FString OplogPath;
 	FString OplogNewEntryPath;
 	FString OplogPrepNewEntryPath;
 	FString TempDirPath;
-	uint64 StandaloneThresholdBytes = 1 * 1024 * 1024;
+	const uint64 StandaloneThresholdBytes = 1 * 1024 * 1024;
 	bool bAllowRead = false;
 	bool bAllowEdit = false;
 	bool bConnectionSucceeded = false;
diff --git a/Engine/Source/Developer/IoStoreUtilities/Private/ZenStoreHttpClient.cpp b/Engine/Source/Developer/IoStoreUtilities/Private/ZenStoreHttpClient.cpp
--- a/Engine/Source/Developer/IoStoreUtilities/Private/ZenStoreHttpClient.cpp
+++ b/Engine/Source/Developer/IoStoreUtilities/Private/ZenStoreHttpClient.cpp
@@ -32,6 +32,13 @@ namespace UE {
 
 const uint32 FZenStoreHttpClient::PoolEntryCount = 32;
 
+std::atomic<uint32> FZenStoreHttpClient::SaltGenerator::GOpCounter(0);
+
+FZenStoreHttpClient::SaltGenerator::SaltGenerator()
+	: SaltBase(FWindowsPlatformProcess::GetCurrentProcessId() + 0x9e3779b9u)
+{
+}
+
 FZenStoreHttpClient::FZenStoreHttpClient()
 {
 	RequestPool = MakeUnique<Zen::FZenHttpRequestPool>(ZenService.GetInstance().GetURL(), PoolEntryCount);
@@ -248,8 +255,6 @@ void FZenStoreHttpClient::InitializeReadOnly(FStringView InProjectId, FStringVie
 	bAllowRead = true;
 }
 
-static std::atomic<uint32> GOpCounter;
-
 TIoStatusOr<uint64> FZenStoreHttpClient::AppendOp(FCbPackage OpEntry)
 {
 	check(bAllowEdit);
@@ -260,7 +265,7 @@ TIoStatusOr<uint64> FZenStoreHttpClient::AppendOp(FCbPackage OpEntry)
 		TRACE_CPUPROFILER_EVENT_SCOPE(Zen_AppendOp_Async);
 		FLargeMemoryWriter SerializedPackage;
 
-		const uint32 Salt = ++GOpCounter;
+		const int32 Salt = SaltGen.Next();
 		bool bIsUsingTempFiles = false;
 
 		UE::Zen::FZenScopedRequestPtr Request(RequestPool.Get());
