Fix usage of FWindowsPlatformProcess in 23917462

diff --git a/Engine/Source/Developer/IoStoreUtilities/Private/ZenStoreHttpClient.cpp b/Engine/Source/Developer/IoStoreUtilities/Private/ZenStoreHttpClient.cpp
--- a/Engine/Source/Developer/IoStoreUtilities/Private/ZenStoreHttpClient.cpp
+++ b/Engine/Source/Developer/IoStoreUtilities/Private/ZenStoreHttpClient.cpp
@@ -35,7 +35,7 @@ const uint32 FZenStoreHttpClient::PoolEntryCount = 32;
 std::atomic<uint32> FZenStoreHttpClient::SaltGenerator::GOpCounter(0);
 
 FZenStoreHttpClient::SaltGenerator::SaltGenerator()
-	: SaltBase(FWindowsPlatformProcess::GetCurrentProcessId() + 0x9e3779b9u)
+	: SaltBase(FGenericPlatformProcess::GetCurrentProcessId() + 0x9e3779b9u)
 {
 }
 
