ZenCook: Remove unused fields "guid" and "disksize" from oplog

diff --git a/Engine/Source/Developer/IoStoreUtilities/Private/ZenStoreWriter.cpp b/Engine/Source/Developer/IoStoreUtilities/Private/ZenStoreWriter.cpp
--- a/Engine/Source/Developer/IoStoreUtilities/Private/ZenStoreWriter.cpp
+++ b/Engine/Source/Developer/IoStoreUtilities/Private/ZenStoreWriter.cpp
@@ -404,23 +404,15 @@ void FZenStoreWriter::Initialize(const FCookInfo& Info)
 						{
 							FCbObject OplogObj = OplogEntry.AsObject();
 
-							if (OplogObj["package"])
+							if (OplogObj["packagestoreentry"])
 							{
-								FCbObject PackageObj = OplogObj["package"].AsObject();
-
-								const FGuid PkgGuid = PackageObj["guid"].AsUuid();
-								const FIoHash PkgHash = PackageObj["data"].AsHash();
-								const int64	PkgDiskSize = PackageObj["disksize"].AsUInt64();
 								FPackageStoreEntryResource Entry = FPackageStoreEntryResource::FromCbObject(OplogObj["packagestoreentry"].AsObject());
 								const FName PackageName = Entry.PackageName;
 
 								const int32 Index = PackageStoreEntries.Num();
 
 								PackageStoreEntries.Add(MoveTemp(Entry));
-								FOplogCookInfo& CookInfo = CookedPackagesInfo.Add_GetRef(
-									FOplogCookInfo{
-										FCookedPackageInfo {PackageName, IoHashToMD5(PkgHash), PkgGuid, PkgDiskSize }
-									});
+								FOplogCookInfo& CookInfo = CookedPackagesInfo.Add_GetRef({ PackageName });
 								PackageNameToIndex.Add(PackageName, Index);
 
 								for (FCbFieldView Field : OplogObj)
@@ -690,17 +682,6 @@ void FZenStoreWriter::CommitPackageInternal(FZenCommitInfo&& ZenCommitInfo)
 		
 		// Commit attachments
 		FOplogCookInfo& CookInfo = CookedPackagesInfo[CommitEventArgs.EntryIndex];
-		CookInfo = FOplogCookInfo
-		{
-			FCookedPackageInfo
-			{ 
-				CommitInfo.PackageName,
-				IoHashToMD5(PkgDataAttachment.GetHash()),
-				CommitInfo.PackageGuid,
-				int64(PkgDataAttachment.AsCompressedBinary().GetRawSize())
-			}
-		};
-
 		CookInfo.bUpToDate = true;
 
 		const int32 NumAttachments = CommitInfo.Attachments.Num();
@@ -744,12 +725,9 @@ void FZenStoreWriter::CommitPackageInternal(FZenCommitInfo&& ZenCommitInfo)
 		PackageNameKey.ToLowerInline();
 		OplogEntryDesc << "key" << PackageNameKey;
 
-		// NOTE: The package GUID and disk size are used for legacy iterative cooks when comparing asset registry package data
 		OplogEntryDesc.BeginObject("package");
 		OplogEntryDesc << "id" << PkgData.ChunkId;
-		OplogEntryDesc << "guid" << CommitInfo.PackageGuid;
 		OplogEntryDesc << "data" << PkgDataAttachment;
-		OplogEntryDesc << "disksize" << PkgDataAttachment.AsCompressedBinary().GetRawSize();
 		OplogEntryDesc.EndObject();
 
 		OplogEntryDesc << "packagestoreentry" << PkgData.PackageStoreEntry;
@@ -1001,7 +979,7 @@ void FZenStoreWriter::RemoveCookedPackages(TArrayView<const FName> PackageNamesT
 		int32 EntryIndex = 0;
 		for (int32 Idx : PackageIndicesToKeep)
 		{
-			const FName PackageName = PreviousCookedPackageInfo[Idx].CookInfo.PackageName;
+			const FName PackageName = PreviousCookedPackageInfo[Idx].PackageName;
 
 			PackageStoreEntries.Add(MoveTemp(PreviousPackageStoreEntries[Idx]));
 			CookedPackagesInfo.Add(MoveTemp(PreviousCookedPackageInfo[Idx]));
diff --git a/Engine/Source/Developer/IoStoreUtilities/Public/PackageStoreWriter.h b/Engine/Source/Developer/IoStoreUtilities/Public/PackageStoreWriter.h
--- a/Engine/Source/Developer/IoStoreUtilities/Public/PackageStoreWriter.h
+++ b/Engine/Source/Developer/IoStoreUtilities/Public/PackageStoreWriter.h
@@ -22,7 +22,7 @@ public:
 			FIoHash Hash;
 		};
 
-		ICookedPackageWriter::FCookedPackageInfo CookInfo;
+		FName PackageName;
 		TArray<FAttachment> Attachments;
 		bool bUpToDate = false;
 	};
diff --git a/Engine/Source/Runtime/Core/Public/Serialization/PackageWriter.h b/Engine/Source/Runtime/Core/Public/Serialization/PackageWriter.h
--- a/Engine/Source/Runtime/Core/Public/Serialization/PackageWriter.h
+++ b/Engine/Source/Runtime/Core/Public/Serialization/PackageWriter.h
@@ -269,14 +269,6 @@ public:
 	  */
 	virtual void EndCook(const FCookInfo& Info) = 0;
 
-	struct FCookedPackageInfo
-	{
-		FName PackageName;
-		FMD5Hash Hash;
-		FGuid PackageGuid;
-		int64 DiskSize = -1;
-	};
-
 	/**
 	 * Returns an AssetRegistry describing the previous cook results. This doesn't mean a cook saved off
 	 * to another directory - it means the AssetRegistry that's living in the directory we are about
