IoStore:
Remove the requirement to resolve full export names when staging from Zen

diff --git a/Engine/Source/Developer/IoStoreUtilities/Internal/PackageStoreOptimizer.h b/Engine/Source/Developer/IoStoreUtilities/Internal/PackageStoreOptimizer.h
--- a/Engine/Source/Developer/IoStoreUtilities/Internal/PackageStoreOptimizer.h
+++ b/Engine/Source/Developer/IoStoreUtilities/Internal/PackageStoreOptimizer.h
@@ -180,7 +180,6 @@ private:
 
 	struct FExport
 	{
-		FString FullName;
 		FName ObjectName;
 		uint64 PublicExportHash = 0;
 		FPackageObjectIndex OuterIndex;
@@ -203,6 +202,11 @@ private:
 		TArray<FExportBundleEntry> Entries;
 	};
 
+	struct FUnresolvedExport
+	{
+		FString FullName;
+	};
+
 	struct FUnresolvedImport
 	{
 		FString FullName;
@@ -343,7 +347,7 @@ public:
 
 	FPackageStorePackage* CreateMissingPackage(const FName& Name) const;
 	FPackageStorePackage* CreatePackageFromCookedHeader(const FName& Name, const FIoBuffer& CookedHeaderBuffer) const;
-	FPackageStorePackage* CreatePackageFromPackageStoreHeader(const FName& Name, const FIoBuffer& Buffer, const FPackageStoreEntryResource& PackageStoreEntry) const;
+	FPackageStorePackage* CreatePackageFromZenPackageHeader(const FName& Name, const FIoBuffer& Buffer, int32 ExportBundleCount, const TArrayView<const FPackageId>& ImportedPackageIds) const;
 	void FinalizePackage(FPackageStorePackage* Package);
 	FIoBuffer CreatePackageBuffer(const FPackageStorePackage* Package, const FIoBuffer& CookedExportsBuffer, TArray<FFileRegion>* InOutFileRegions) const;
 	FPackageStoreEntryResource CreatePackageStoreEntry(const FPackageStorePackage* Package, const FPackageStorePackage* OptionalSegmentPackage) const;
@@ -374,7 +378,7 @@ private:
 		TArray<FObjectDataResource> DataResources;
 	};
 
-	struct FPackageStoreHeaderData
+	struct FZenPackageHeaderData
 	{
 		FZenPackageSummary Summary;
 		TOptional<FZenPackageVersioningInfo> VersioningInfo;
@@ -395,16 +399,15 @@ private:
 
 	static uint64 GetPublicExportHash(FStringView PackageRelativeExportPath);
 	FCookedHeaderData LoadCookedHeader(const FIoBuffer& CookedHeaderBuffer) const;
-	FPackageStoreHeaderData LoadPackageStoreHeader(const FIoBuffer& PackageStoreHeaderBuffer, const FPackageStoreEntryResource& PackageStoreEntry) const;
+	FZenPackageHeaderData LoadZenPackageHeader(const FIoBuffer& HeaderBuffer, int32 ExportBundleCount, const TArrayView<const FPackageId>& ImportedPackageIds) const;
 	void ResolveImport(FPackageStorePackage::FUnresolvedImport* Imports, const FObjectImport* ObjectImports, int32 LocalImportIndex) const;
 	void ProcessImports(const FCookedHeaderData& CookedHeaderData, FPackageStorePackage* Package, TArray<FPackageStorePackage::FUnresolvedImport>& UnresolvedImports) const;
-	void ProcessImports(const FPackageStoreHeaderData& PackageStoreHeaderData, FPackageStorePackage* Package) const;
-	void ResolveExport(FPackageStorePackage::FExport* Exports, const FObjectExport* ObjectExports, const int32 LocalExportIndex, const FName& PackageName, FPackageStorePackage::FUnresolvedImport* Imports, const FObjectImport* ObjectImports) const;
-	void ResolveExport(FPackageStorePackage::FExport* Exports, const int32 LocalExportIndex, const FName& PackageName) const;
+	void ProcessImports(const FZenPackageHeaderData& ZenHeaderData, FPackageStorePackage* Package) const;
+	void ResolveExport(FPackageStorePackage::FUnresolvedExport* Exports, const FObjectExport* ObjectExports, const int32 LocalExportIndex, const FName& PackageName, FPackageStorePackage::FUnresolvedImport* Imports, const FObjectImport* ObjectImports) const;
 	void ProcessExports(const FCookedHeaderData& CookedHeaderData, FPackageStorePackage* Package, FPackageStorePackage::FUnresolvedImport* Imports) const;
-	void ProcessExports(const FPackageStoreHeaderData& PackageStoreHeaderData, FPackageStorePackage* Package) const;
+	void ProcessExports(const FZenPackageHeaderData& ZenHeaderData, FPackageStorePackage* Package) const;
 	void ProcessPreloadDependencies(const FCookedHeaderData& CookedHeaderData, FPackageStorePackage* Package) const;
-	void ProcessPreloadDependencies(const FPackageStoreHeaderData& PackageStoreHeaderData, FPackageStorePackage* Package) const;
+	void ProcessPreloadDependencies(const FZenPackageHeaderData& ZenHeaderData, FPackageStorePackage* Package) const;
 	void ProcessDataResources(const FCookedHeaderData& CookedHeaderData, FPackageStorePackage* Package) const;
 	TArray<FPackageStorePackage*> SortPackagesInLoadOrder(const TMap<FPackageId, FPackageStorePackage*>& PackagesMap) const;
 	void SerializeGraphData(const TArray<FPackageId>& ImportedPackageIds, FPackageStorePackage::FGraphData& GraphData, FBufferWriter& GraphArchive) const;
diff --git a/Engine/Source/Developer/IoStoreUtilities/Private/IoStoreUtilities.cpp b/Engine/Source/Developer/IoStoreUtilities/Private/IoStoreUtilities.cpp
--- a/Engine/Source/Developer/IoStoreUtilities/Private/IoStoreUtilities.cpp
+++ b/Engine/Source/Developer/IoStoreUtilities/Private/IoStoreUtilities.cpp
@@ -1602,7 +1602,8 @@ static void ParsePackageAssetsFromZen(FCookedPackageStore& PackageStore, TArray<
 		TIoStatusOr<FIoBuffer> HeaderBuffer = PackageStore.ReadPackageHeaderFromZen(Package->GlobalPackageId);
 		Package->UExpSize = Package->UAssetSize - HeaderBuffer.ValueOrDie().DataSize();
 		Package->UAssetSize = HeaderBuffer.ValueOrDie().DataSize();
-		Package->OptimizedPackage = PackageStoreOptimizer.CreatePackageFromPackageStoreHeader(Package->PackageName, HeaderBuffer.ValueOrDie(), *Package->PackageStoreEntry);
+		const FPackageStoreEntryResource* PackageStoreEntry = Package->PackageStoreEntry;
+		Package->OptimizedPackage = PackageStoreOptimizer.CreatePackageFromZenPackageHeader(Package->PackageName, HeaderBuffer.ValueOrDie(), PackageStoreEntry->ExportInfo.ExportBundleCount, PackageStoreEntry->ImportedPackageIds);
 		check(Package->OptimizedPackage->GetId() == Package->GlobalPackageId);
 	}, EParallelForFlags::Unbalanced);
 }
diff --git a/Engine/Source/Developer/IoStoreUtilities/Private/PackageStoreOptimizer.cpp b/Engine/Source/Developer/IoStoreUtilities/Private/PackageStoreOptimizer.cpp
--- a/Engine/Source/Developer/IoStoreUtilities/Private/PackageStoreOptimizer.cpp
+++ b/Engine/Source/Developer/IoStoreUtilities/Private/PackageStoreOptimizer.cpp
@@ -151,7 +151,7 @@ FPackageStorePackage* FPackageStoreOptimizer::CreatePackageFromCookedHeader(cons
 	return Package;
 }
 
-FPackageStorePackage* FPackageStoreOptimizer::CreatePackageFromPackageStoreHeader(const FName& Name, const FIoBuffer& Buffer, const FPackageStoreEntryResource& PackageStoreEntry) const
+FPackageStorePackage* FPackageStoreOptimizer::CreatePackageFromZenPackageHeader(const FName& Name, const FIoBuffer& Buffer, int32 ExportBundleCount, const TArrayView<const FPackageId>& ImportedPackageIds) const
 {
 	FPackageStorePackage* Package = new FPackageStorePackage();
 	Package->Id = FPackageId::FromName(Name);
@@ -168,21 +168,21 @@ FPackageStorePackage* FPackageStoreOptimizer::CreatePackageFromPackageStoreHeade
 	Package->Name = Name;
 	Package->SourceName = *RemapLocalizationPathIfNeeded(NameStr, &Package->Region);
 
-	FPackageStoreHeaderData PackageStoreHeaderData = LoadPackageStoreHeader(Buffer, PackageStoreEntry);
-	if (PackageStoreHeaderData.VersioningInfo.IsSet())
+	FZenPackageHeaderData ZenHeaderData = LoadZenPackageHeader(Buffer, ExportBundleCount, ImportedPackageIds);
+	if (ZenHeaderData.VersioningInfo.IsSet())
 	{
-		Package->VersioningInfo.Emplace(PackageStoreHeaderData.VersioningInfo.GetValue());
+		Package->VersioningInfo.Emplace(ZenHeaderData.VersioningInfo.GetValue());
 	}
-	Package->PackageFlags = PackageStoreHeaderData.Summary.PackageFlags;
-	Package->CookedHeaderSize = PackageStoreHeaderData.Summary.CookedHeaderSize;
-	for (FDisplayNameEntryId DisplayId : PackageStoreHeaderData.NameMap)
+	Package->PackageFlags = ZenHeaderData.Summary.PackageFlags;
+	Package->CookedHeaderSize = ZenHeaderData.Summary.CookedHeaderSize;
+	for (FDisplayNameEntryId DisplayId : ZenHeaderData.NameMap)
 	{
 		Package->NameMapBuilder.AddName(DisplayId);
 	}
-	Package->BulkDataEntries = MoveTemp(PackageStoreHeaderData.BulkDataEntries);
-	ProcessImports(PackageStoreHeaderData, Package);
-	ProcessExports(PackageStoreHeaderData, Package);
-	ProcessPreloadDependencies(PackageStoreHeaderData, Package);
+	Package->BulkDataEntries = MoveTemp(ZenHeaderData.BulkDataEntries);
+	ProcessImports(ZenHeaderData, Package);
+	ProcessExports(ZenHeaderData, Package);
+	ProcessPreloadDependencies(ZenHeaderData, Package);
 	CreateExportBundles(Package);
 
 	return Package;
@@ -295,68 +295,66 @@ FPackageStoreOptimizer::FCookedHeaderData FPackageStoreOptimizer::LoadCookedHead
 }
 PRAGMA_ENABLE_DEPRECATION_WARNINGS
 
-FPackageStoreOptimizer::FPackageStoreHeaderData FPackageStoreOptimizer::LoadPackageStoreHeader(const FIoBuffer& PackageStoreHeaderBuffer, const FPackageStoreEntryResource& PackageStoreEntry) const
+FPackageStoreOptimizer::FZenPackageHeaderData FPackageStoreOptimizer::LoadZenPackageHeader(const FIoBuffer& HeaderBuffer, int32 ExportBundleCount, const TArrayView<const FPackageId>& ImportedPackageIds) const
 {
-	FPackageStoreHeaderData PackageStoreHeaderData;
+	FZenPackageHeaderData ZenPackageHeaderData;
 
-	const uint8* HeaderData = PackageStoreHeaderBuffer.Data();
+	const uint8* HeaderData = HeaderBuffer.Data();
 
-	FZenPackageSummary& Summary = PackageStoreHeaderData.Summary;
+	FZenPackageSummary& Summary = ZenPackageHeaderData.Summary;
 	Summary = *reinterpret_cast<const FZenPackageSummary*>(HeaderData);
-	check(PackageStoreHeaderBuffer.DataSize() == Summary.HeaderSize);
+	check(HeaderBuffer.DataSize() == Summary.HeaderSize);
 
 	TArrayView<const uint8> HeaderDataView(HeaderData + sizeof(FZenPackageSummary), Summary.HeaderSize - sizeof(FZenPackageSummary));
 	FMemoryReaderView HeaderDataReader(HeaderDataView);
 	
 	if (Summary.bHasVersioningInfo)
 	{
-		FZenPackageVersioningInfo& VersioningInfo = PackageStoreHeaderData.VersioningInfo.Emplace();
+		FZenPackageVersioningInfo& VersioningInfo = ZenPackageHeaderData.VersioningInfo.Emplace();
 		HeaderDataReader << VersioningInfo;
 	}
 
-	TArray<FDisplayNameEntryId>& NameMap = PackageStoreHeaderData.NameMap;
+	TArray<FDisplayNameEntryId>& NameMap = ZenPackageHeaderData.NameMap;
 	NameMap = LoadNameBatch(HeaderDataReader);
 
-	const FZenPackageVersioningInfo* VersioningInfo = PackageStoreHeaderData.VersioningInfo.GetPtrOrNull();
+	const FZenPackageVersioningInfo* VersioningInfo = ZenPackageHeaderData.VersioningInfo.GetPtrOrNull();
 	if (VersioningInfo == nullptr || VersioningInfo->PackageVersion >= EUnrealEngineObjectUE5Version::DATA_RESOURCES)
 	{
 		int64 BulkDataMapSize = 0;
 		HeaderDataReader << BulkDataMapSize;
 		const uint8* BulkDataMapData = HeaderData + sizeof(FZenPackageSummary) + HeaderDataReader.Tell();
-		PackageStoreHeaderData.BulkDataEntries = MakeArrayView(reinterpret_cast<const FBulkDataMapEntry*>(BulkDataMapData), BulkDataMapSize / sizeof(FBulkDataMapEntry));
+		ZenPackageHeaderData.BulkDataEntries = MakeArrayView(reinterpret_cast<const FBulkDataMapEntry*>(BulkDataMapData), BulkDataMapSize / sizeof(FBulkDataMapEntry));
 	}
 
-	for (FPackageId PackageId : PackageStoreEntry.ImportedPackageIds)
-	{
-		PackageStoreHeaderData.ImportedPackageIds.Add(PackageId);
-	}
+	ZenPackageHeaderData.ImportedPackageIds = ImportedPackageIds;
+
 
-	PackageStoreHeaderData.ImportedPublicExportHashes =
+	ZenPackageHeaderData.ImportedPublicExportHashes =
 		MakeArrayView<const uint64>(
 			reinterpret_cast<const uint64*>(HeaderData + Summary.ImportedPublicExportHashesOffset),
 			(Summary.ImportMapOffset - Summary.ImportedPublicExportHashesOffset) / sizeof(uint64));
 
-	PackageStoreHeaderData.Imports =
+	ZenPackageHeaderData.Imports =
 		MakeArrayView<const FPackageObjectIndex>(
 			reinterpret_cast<const FPackageObjectIndex*>(HeaderData + Summary.ImportMapOffset),
 			(Summary.ExportMapOffset - Summary.ImportMapOffset) / sizeof(FPackageObjectIndex));
 
-	PackageStoreHeaderData.Exports =
+	ZenPackageHeaderData.Exports =
 		MakeArrayView<const FExportMapEntry>(
 			reinterpret_cast<const FExportMapEntry*>(HeaderData + Summary.ExportMapOffset),
 			(Summary.ExportBundleEntriesOffset - Summary.ExportMapOffset) / sizeof(FExportMapEntry));
 
-	PackageStoreHeaderData.ExportBundleHeaders =
+	ZenPackageHeaderData.ExportBundleHeaders =
 		MakeArrayView<const FExportBundleHeader>(
 			reinterpret_cast<const FExportBundleHeader*>(HeaderData + Summary.GraphDataOffset),
-			PackageStoreEntry.ExportInfo.ExportBundleCount);
+			ExportBundleCount);
 
-	PackageStoreHeaderData.ExportBundleEntries = 
+	ZenPackageHeaderData.ExportBundleEntries =
 		MakeArrayView<const FExportBundleEntry>(
 			reinterpret_cast<const FExportBundleEntry*>(HeaderData + Summary.ExportBundleEntriesOffset),
-			PackageStoreEntry.ExportInfo.ExportCount * FExportBundleEntry::ExportCommandType_Count);
+			(Summary.GraphDataOffset - Summary.ExportBundleEntriesOffset) / sizeof(FExportBundleEntry));
 
-	const uint64 ExportBundleHeadersSize = sizeof(FExportBundleHeader) * PackageStoreEntry.ExportInfo.ExportBundleCount;
+	const uint64 ExportBundleHeadersSize = sizeof(FExportBundleHeader) * ExportBundleCount;
 	const uint64 ArcsDataOffset = Summary.GraphDataOffset + ExportBundleHeadersSize;
 	const uint64 ArcsDataSize = Summary.HeaderSize - ArcsDataOffset;
 
@@ -367,12 +365,12 @@ FPackageStoreOptimizer::FPackageStoreHeaderData FPackageStoreOptimizer::LoadPack
 
 	for (int32 Idx = 0; Idx < InternalArcsCount; ++Idx)
 	{
-		FPackageStorePackage::FInternalArc& InternalArc = PackageStoreHeaderData.InternalArcs.AddDefaulted_GetRef();
+		FPackageStorePackage::FInternalArc& InternalArc = ZenPackageHeaderData.InternalArcs.AddDefaulted_GetRef();
 		ArcsAr << InternalArc.FromExportBundleIndex;
 		ArcsAr << InternalArc.ToExportBundleIndex;
 	}
 
-	for (FPackageId ImportedPackageId : PackageStoreHeaderData.ImportedPackageIds)
+	for (FPackageId ImportedPackageId : ZenPackageHeaderData.ImportedPackageIds)
 	{
 		int32 ExternalArcsCount = 0;
 		ArcsAr << ExternalArcsCount;
@@ -386,11 +384,11 @@ FPackageStoreOptimizer::FPackageStoreHeaderData FPackageStoreOptimizer::LoadPack
 			ExternalArc.FromCommandType = static_cast<FExportBundleEntry::EExportCommandType>(FromCommandType);
 			ArcsAr << ExternalArc.ToExportBundleIndex;
 
-			PackageStoreHeaderData.ExternalArcs.Add(ExternalArc);
+			ZenPackageHeaderData.ExternalArcs.Add(ExternalArc);
 		}
 	}
 
-	return PackageStoreHeaderData;
+	return ZenPackageHeaderData;
 }
 
 void FPackageStoreOptimizer::ResolveImport(FPackageStorePackage::FUnresolvedImport* Imports, const FObjectImport* ObjectImports, int32 LocalImportIndex) const
@@ -495,22 +493,22 @@ void FPackageStoreOptimizer::ProcessImports(const FCookedHeaderData& CookedHeade
 	}
 }
 
-void FPackageStoreOptimizer::ProcessImports(const FPackageStoreHeaderData& PackageStoreHeaderData, FPackageStorePackage* Package) const
+void FPackageStoreOptimizer::ProcessImports(const FZenPackageHeaderData& ZenHeaderData, FPackageStorePackage* Package) const
 {
-	Package->ImportedPackageIds = PackageStoreHeaderData.ImportedPackageIds;
-	Package->ImportedPublicExportHashes = PackageStoreHeaderData.ImportedPublicExportHashes;
-	Package->Imports = PackageStoreHeaderData.Imports;
+	Package->ImportedPackageIds = ZenHeaderData.ImportedPackageIds;
+	Package->ImportedPublicExportHashes = ZenHeaderData.ImportedPublicExportHashes;
+	Package->Imports = ZenHeaderData.Imports;
 }
 
 void FPackageStoreOptimizer::ResolveExport(
-	FPackageStorePackage::FExport* Exports,
+	FPackageStorePackage::FUnresolvedExport* Exports,
 	const FObjectExport* ObjectExports,
 	const int32 LocalExportIndex,
 	const FName& PackageName,
 	FPackageStorePackage::FUnresolvedImport* Imports,
 	const FObjectImport* ObjectImports) const
 {
-	FPackageStorePackage::FExport* Export = Exports + LocalExportIndex;
+	FPackageStorePackage::FUnresolvedExport* Export = Exports + LocalExportIndex;
 	if (Export->FullName.Len() == 0)
 	{
 		Export->FullName.Reserve(256);
@@ -549,38 +547,12 @@ void FPackageStoreOptimizer::ResolveExport(
 	}
 }
 
-void FPackageStoreOptimizer::ResolveExport(FPackageStorePackage::FExport* Exports, const int32 LocalExportIndex, const FName& PackageName) const
-{
-	FPackageStorePackage::FExport* Export = Exports + LocalExportIndex;
-	if (Export->FullName.Len() == 0)
-	{
-		Export->FullName.Reserve(256);
-		if (Export->OuterIndex.IsNull())
-		{
-			PackageName.AppendString(Export->FullName);
-			Export->FullName.AppendChar(TEXT('/'));
-			Export->ObjectName.AppendString(Export->FullName);
-			Export->FullName.ToLowerInline();
-			check(Export->FullName.Len() > 0);
-		}
-		else
-		{
-			check(Export->OuterIndex.IsExport());
-			int32 OuterExportIndex = Export->OuterIndex.ToExport();
-			ResolveExport(Exports, OuterExportIndex, PackageName);
-			FString& OuterName = Exports[OuterExportIndex].FullName;
-			check(OuterName.Len() > 0);
-			Export->FullName.Append(OuterName);
-			Export->FullName.AppendChar(TEXT('/'));
-			Export->ObjectName.AppendString(Export->FullName);
-			Export->FullName.ToLowerInline();
-		}
-	}
-}
-
 void FPackageStoreOptimizer::ProcessExports(const FCookedHeaderData& CookedHeaderData, FPackageStorePackage* Package, FPackageStorePackage::FUnresolvedImport* Imports) const
 {
 	int32 ExportCount = CookedHeaderData.ObjectExports.Num();
+
+	TArray<FPackageStorePackage::FUnresolvedExport> UnresolvedExports;
+	UnresolvedExports.SetNum(ExportCount);
 	Package->Exports.SetNum(ExportCount);
 	Package->ExportGraphNodes.Reserve(ExportCount * 2);
 
@@ -599,13 +571,14 @@ void FPackageStoreOptimizer::ProcessExports(const FCookedHeaderData& CookedHeade
 	};
 
 	FString PackageNameStr = Package->Name.ToString();
-	TMap<uint64, const FPackageStorePackage::FExport*> SeenPublicExportHashes;
+	TMap<uint64, const FPackageStorePackage::FUnresolvedExport*> SeenPublicExportHashes;
 	for (int32 ExportIndex = 0; ExportIndex < ExportCount; ++ExportIndex)
 	{
 		const FObjectExport& ObjectExport = CookedHeaderData.ObjectExports[ExportIndex];
 		Package->ExportsSerialSize += ObjectExport.SerialSize;
 
 		FPackageStorePackage::FExport& Export = Package->Exports[ExportIndex];
+		FPackageStorePackage::FUnresolvedExport& UnresolvedExport = UnresolvedExports[ExportIndex];
 		Export.ObjectName = ObjectExport.ObjectName;
 		Export.ObjectFlags = ObjectExport.ObjectFlags;
 		Export.CookedSerialOffset = ObjectExport.SerialOffset;
@@ -614,19 +587,19 @@ void FPackageStoreOptimizer::ProcessExports(const FCookedHeaderData& CookedHeade
 		Export.bNotForClient = ObjectExport.bNotForClient;
 		Export.bNotForServer = ObjectExport.bNotForServer;
 		Export.bIsPublic = (Export.ObjectFlags & RF_Public) > 0 || ObjectExport.bGeneratePublicHash;
-		ResolveExport(Package->Exports.GetData(), CookedHeaderData.ObjectExports.GetData(), ExportIndex, Package->Name, Imports, CookedHeaderData.ObjectImports.GetData());
+		ResolveExport(UnresolvedExports.GetData(), CookedHeaderData.ObjectExports.GetData(), ExportIndex, Package->Name, Imports, CookedHeaderData.ObjectImports.GetData());
 		if (Export.bIsPublic)
 		{
-			check(Export.FullName.Len() > 0);
-			FStringView PackageRelativeName = FStringView(Export.FullName).RightChop(PackageNameStr.Len());
+			check(UnresolvedExport.FullName.Len() > 0);
+			FStringView PackageRelativeName = FStringView(UnresolvedExport.FullName).RightChop(PackageNameStr.Len());
 			check(PackageRelativeName.Len());
 			Export.PublicExportHash = GetPublicExportHash(PackageRelativeName);
-			const FPackageStorePackage::FExport* FindCollidingExport = SeenPublicExportHashes.FindRef(Export.PublicExportHash);
+			const FPackageStorePackage::FUnresolvedExport* FindCollidingExport = SeenPublicExportHashes.FindRef(Export.PublicExportHash);
 			if (FindCollidingExport)
 			{
 				UE_LOG(LogPackageStoreOptimizer, Fatal, TEXT("Export hash collision in package \"%s\": \"%s\" and \"%s"), *PackageNameStr, PackageRelativeName.GetData(), *FindCollidingExport->FullName.RightChop(PackageNameStr.Len()));
 			}
-			SeenPublicExportHashes.Add(Export.PublicExportHash, &Export);
+			SeenPublicExportHashes.Add(Export.PublicExportHash, &UnresolvedExport);
 		}
 
 		Export.OuterIndex = PackageObjectIdFromPackageIndex(Package->Imports, ObjectExport.OuterIndex);
@@ -645,18 +618,18 @@ void FPackageStoreOptimizer::ProcessExports(const FCookedHeaderData& CookedHeade
 	}
 }
 
-void FPackageStoreOptimizer::ProcessExports(const FPackageStoreHeaderData& PackageStoreHeaderData, FPackageStorePackage* Package) const
+void FPackageStoreOptimizer::ProcessExports(const FZenPackageHeaderData& ZenHeaderData, FPackageStorePackage* Package) const
 {
-	const int32 ExportCount = PackageStoreHeaderData.Exports.Num();
+	const int32 ExportCount = ZenHeaderData.Exports.Num();
 	Package->Exports.SetNum(ExportCount);
 	Package->ExportGraphNodes.Reserve(ExportCount * 2);
 
-	const TArray<FDisplayNameEntryId>& NameMap = PackageStoreHeaderData.NameMap;
+	const TArray<FDisplayNameEntryId>& NameMap = ZenHeaderData.NameMap;
 
-	Package->ImportedPublicExportHashes = PackageStoreHeaderData.ImportedPublicExportHashes;
+	Package->ImportedPublicExportHashes = ZenHeaderData.ImportedPublicExportHashes;
 	for (int32 ExportIndex = 0; ExportIndex < ExportCount; ++ExportIndex)
 	{
-		const FExportMapEntry& ExportEntry =  PackageStoreHeaderData.Exports[ExportIndex];
+		const FExportMapEntry& ExportEntry = ZenHeaderData.Exports[ExportIndex];
 		Package->ExportsSerialSize += ExportEntry.CookedSerialSize;
 
 		FPackageStorePackage::FExport& Export = Package->Exports[ExportIndex];
@@ -685,18 +658,13 @@ void FPackageStoreOptimizer::ProcessExports(const FPackageStoreHeaderData& Packa
 		Package->NameMapBuilder.MarkNameAsReferenced(Export.ObjectName);
 	}
 
-	for (int32 ExportIndex = 0; ExportIndex < ExportCount; ++ExportIndex)
-	{
-		ResolveExport(Package->Exports.GetData(), ExportIndex, Package->Name);
-	}
-
 	uint64 ExportSerialOffset = 0;
-	for (const FExportBundleHeader& ExportBundleHeader : PackageStoreHeaderData.ExportBundleHeaders)
+	for (const FExportBundleHeader& ExportBundleHeader : ZenHeaderData.ExportBundleHeaders)
 	{
 		int32 ExportEntryIndex = ExportBundleHeader.FirstEntryIndex;
 		for (uint32 Idx = 0; Idx < ExportBundleHeader.EntryCount; ++Idx)
 		{
-			const FExportBundleEntry& BundleEntry = PackageStoreHeaderData.ExportBundleEntries[ExportEntryIndex++];
+			const FExportBundleEntry& BundleEntry = ZenHeaderData.ExportBundleEntries[ExportEntryIndex++];
 			if (BundleEntry.CommandType == FExportBundleEntry::ExportCommandType_Serialize)
 			{
 				FPackageStorePackage::FExport& Export = Package->Exports[BundleEntry.LocalExportIndex];
@@ -1036,20 +1004,20 @@ void FPackageStoreOptimizer::ProcessDataResources(const FCookedHeaderData& Cooke
 	}
 }
 
-void FPackageStoreOptimizer::ProcessPreloadDependencies(const FPackageStoreHeaderData& PackageStoreHeaderData, FPackageStorePackage* Package) const
+void FPackageStoreOptimizer::ProcessPreloadDependencies(const FZenPackageHeaderData& ZenHeaderData, FPackageStorePackage* Package) const
 {
 	TRACE_CPUPROFILER_EVENT_SCOPE(ProcessPreloadDependencies);
 
-	for (const FPackageStorePackage::FInternalArc& InternalArc : PackageStoreHeaderData.InternalArcs)
+	for (const FPackageStorePackage::FInternalArc& InternalArc : ZenHeaderData.InternalArcs)
 	{
-		const FExportBundleHeader& FromExportBundle = PackageStoreHeaderData.ExportBundleHeaders[InternalArc.FromExportBundleIndex];
-		const FExportBundleHeader& ToExportBundle = PackageStoreHeaderData.ExportBundleHeaders[InternalArc.ToExportBundleIndex];
+		const FExportBundleHeader& FromExportBundle = ZenHeaderData.ExportBundleHeaders[InternalArc.FromExportBundleIndex];
+		const FExportBundleHeader& ToExportBundle = ZenHeaderData.ExportBundleHeaders[InternalArc.ToExportBundleIndex];
 
 		uint32 FromBundleEntryIndex = FromExportBundle.FirstEntryIndex;
 		const uint32 LastFromBundleEntryIndex = FromBundleEntryIndex + FromExportBundle.EntryCount;
 		while (FromBundleEntryIndex < LastFromBundleEntryIndex)
 		{
-			const FExportBundleEntry& FromBundleEntry = PackageStoreHeaderData.ExportBundleEntries[FromBundleEntryIndex++];
+			const FExportBundleEntry& FromBundleEntry = ZenHeaderData.ExportBundleEntries[FromBundleEntryIndex++];
 			FPackageStorePackage::FExport& FromExport = Package->Exports[FromBundleEntry.LocalExportIndex];
 			FPackageStorePackage::FExportGraphNode* FromNode = FromExport.Nodes[FromBundleEntry.CommandType];
 
@@ -1057,7 +1025,7 @@ void FPackageStoreOptimizer::ProcessPreloadDependencies(const FPackageStoreHeade
 			const uint32 LastToBundleEntryIndex = ToBundleEntryIndex + ToExportBundle.EntryCount;
 			while (ToBundleEntryIndex < LastToBundleEntryIndex)
 			{
-				const FExportBundleEntry& ToBundleEntry = PackageStoreHeaderData.ExportBundleEntries[ToBundleEntryIndex++];
+				const FExportBundleEntry& ToBundleEntry = ZenHeaderData.ExportBundleEntries[ToBundleEntryIndex++];
 				FPackageStorePackage::FExport& ToExport = Package->Exports[ToBundleEntry.LocalExportIndex];
 				FPackageStorePackage::FExportGraphNode* ToNode = ToExport.Nodes[ToBundleEntry.CommandType];
 				ToNode->InternalDependencies.Add(FromNode);
@@ -1065,14 +1033,14 @@ void FPackageStoreOptimizer::ProcessPreloadDependencies(const FPackageStoreHeade
 		}
 	}
 
-	for (const FPackageStorePackage::FExternalArc& ExternalArc : PackageStoreHeaderData.ExternalArcs)
+	for (const FPackageStorePackage::FExternalArc& ExternalArc : ZenHeaderData.ExternalArcs)
 	{
-		const FExportBundleHeader& ToExportBundle = PackageStoreHeaderData.ExportBundleHeaders[ExternalArc.ToExportBundleIndex];
+		const FExportBundleHeader& ToExportBundle = ZenHeaderData.ExportBundleHeaders[ExternalArc.ToExportBundleIndex];
 		uint32 ToBundleEntryIndex = ToExportBundle.FirstEntryIndex;
 		const uint32 LastToBundleEntryIndex = ToBundleEntryIndex + ToExportBundle.EntryCount;
 		while (ToBundleEntryIndex < LastToBundleEntryIndex)
 		{
-			const FExportBundleEntry& ToBundleEntry = PackageStoreHeaderData.ExportBundleEntries[ToBundleEntryIndex++];
+			const FExportBundleEntry& ToBundleEntry = ZenHeaderData.ExportBundleEntries[ToBundleEntryIndex++];
 			FPackageStorePackage::FExport& ToExport = Package->Exports[ToBundleEntry.LocalExportIndex];
 			FPackageStorePackage::FExportGraphNode* ToNode = ToExport.Nodes[ToBundleEntry.CommandType];
 			FPackageStorePackage::FExternalDependency& ExternalDependency = ToNode->ExternalDependencies.AddDefaulted_GetRef();
@@ -1165,38 +1133,6 @@ void FPackageStoreOptimizer::CreateExportBundles(FPackageStorePackage* Package)
 	}
 }
 
-static const TCHAR* GetExportNameSafe(const FString& ExportFullName, const FName& PackageName, int32 PackageNameLen)
-{
-	const bool bValidNameLen = ExportFullName.Len() > PackageNameLen + 1;
-	if (bValidNameLen)
-	{
-		const TCHAR* ExportNameStr = *ExportFullName + PackageNameLen;
-		const bool bValidNameFormat = *ExportNameStr == '/';
-		if (bValidNameFormat)
-		{
-			return ExportNameStr + 1; // skip verified '/'
-		}
-		else
-		{
-			UE_CLOG(!bValidNameFormat, LogPackageStoreOptimizer, Warning,
-				TEXT("Export name '%s' should start with '/' at position %d, i.e. right after package prefix '%s'"),
-				*ExportFullName,
-				PackageNameLen,
-				*PackageName.ToString());
-		}
-	}
-	else
-	{
-		UE_CLOG(!bValidNameLen, LogPackageStoreOptimizer, Warning,
-			TEXT("Export name '%s' with length %d should be longer than package name '%s' with length %d"),
-			*ExportFullName,
-			PackageNameLen,
-			*PackageName.ToString());
-	}
-
-	return nullptr;
-};
-
 bool FPackageStoreOptimizer::VerifyRedirect(
 	const FPackageStorePackage* SourcePackage,
 	FPackageStorePackage& TargetPackage,
@@ -1240,91 +1176,54 @@ bool FPackageStoreOptimizer::VerifyRedirect(
 		FailReason.Appendf(TEXT("Public export '%s' has %s %s vs. %s"),
 			*ExportName.ToString(),
 			Text,
-			*TargetPackage.Exports[TargetIndex.ToExport()].FullName,
-			*SourcePackage->Exports[SourceIndex.ToExport()].FullName);
+			*TargetPackage.Exports[TargetIndex.ToExport()].ObjectName.ToString(),
+			*SourcePackage->Exports[SourceIndex.ToExport()].ObjectName.ToString());
 	};
 
-	const int32 TargetPackageNameLen = TargetPackage.Name.GetStringLength();
-	const int32 SourcePackageNameLen = SourcePackage->Name.GetStringLength();
-
-	bool bSuccess = true;
-	int32 TargetIndex = 0;
-	int32 SourceIndex = 0;
-	while (TargetIndex < ExportCount && SourceIndex < ExportCount)
+	TMap<uint64, const FPackageStorePackage::FExport*> SourceExports;
+	for (const FPackageStorePackage::FExport& SourceExport : SourcePackage->Exports)
 	{
-		FString FailReason;
-		const FPackageStorePackage::FExport& TargetExport = TargetPackage.Exports[TargetIndex];
-		const FPackageStorePackage::FExport& SourceExport = SourcePackage->Exports[SourceIndex];
-
-		const TCHAR* TargetExportStr = GetExportNameSafe(
-			TargetExport.FullName, TargetPackage.Name, TargetPackageNameLen);
-		const TCHAR* SourceExportStr = GetExportNameSafe(
-			SourceExport.FullName, SourcePackage->Name, SourcePackageNameLen);
-
-		if (!TargetExportStr || !SourceExportStr)
-		{
-			UE_LOG(LogPackageStoreOptimizer, Error,
-				TEXT("Redirection target package '%s' (0x%llX) for source package '%s' (0x%llX) - Has some bad data from an earlier phase."),
-				*TargetPackage.Name.ToString(),
-				TargetPackage.Id.ValueForDebugging(),
-				*SourcePackage->Name.ToString(),
-				SourcePackage->Id.ValueForDebugging())
-				return false;
-		}
-
-		int32 CompareResult = FCString::Stricmp(TargetExportStr, SourceExportStr);
-		if (CompareResult < 0)
-		{
-			++TargetIndex;
-		}
-		else if (CompareResult > 0)
+		if (SourceExport.bIsPublic)
 		{
-			++SourceIndex;
-
-			if (SourceExport.bIsPublic)
-			{
-				FailReason.Appendf(TEXT("Public source export '%s' is missing in the localized package"),
-					*SourceExport.ObjectName.ToString());
-			}
+			SourceExports.Add(SourceExport.PublicExportHash, &SourceExport);
 		}
-		else
+	}
+	TArray<FString> FailReasons;
+	for (const FPackageStorePackage::FExport& TargetExport : TargetPackage.Exports)
+	{
+		if (TargetExport.bIsPublic)
 		{
-			++TargetIndex;
-			++SourceIndex;
-
-			if (SourceExport.bIsPublic)
+			const FPackageStorePackage::FExport* SourceExport;
+			if (SourceExports.RemoveAndCopyValue(TargetExport.PublicExportHash, SourceExport))
 			{
-				if (!TargetExport.bIsPublic)
-				{
-					FailReason.Appendf(TEXT("Public source export '%s' exists in the localized package")
-						TEXT(", but is not a public localized export."),
-						*SourceExport.ObjectName.ToString());
-				}
-				else if (TargetExport.ClassIndex != SourceExport.ClassIndex)
+				if (TargetExport.ClassIndex != SourceExport->ClassIndex)
 				{
 					AppendMismatchMessage(TEXT("class"), TargetExport.ObjectName,
-						TargetExport.ClassIndex, SourceExport.ClassIndex, FailReason);
+						TargetExport.ClassIndex, SourceExport->ClassIndex, FailReasons.AddDefaulted_GetRef());
 				}
-				else if (TargetExport.TemplateIndex != SourceExport.TemplateIndex)
+				if (TargetExport.TemplateIndex != SourceExport->TemplateIndex)
 				{
 					AppendMismatchMessage(TEXT("template"), TargetExport.ObjectName,
-						TargetExport.TemplateIndex, SourceExport.TemplateIndex, FailReason);
+						TargetExport.TemplateIndex, SourceExport->TemplateIndex, FailReasons.AddDefaulted_GetRef());
 				}
-				else if (TargetExport.SuperIndex != SourceExport.SuperIndex)
+				if (TargetExport.SuperIndex != SourceExport->SuperIndex)
 				{
 					AppendMismatchMessage(TEXT("super"), TargetExport.ObjectName,
-						TargetExport.SuperIndex, SourceExport.SuperIndex, FailReason);
+						TargetExport.SuperIndex, SourceExport->SuperIndex, FailReasons.AddDefaulted_GetRef());
 				}
 			}
-			else if (TargetExport.bIsPublic)
-			{
-				FailReason.Appendf(TEXT("Export '%s' exists in the source package")
-					TEXT(", but is not a public source export."),
-					*TargetExport.ObjectName.ToString());
-			}
 		}
+	}
+	for (const auto& KV : SourceExports)
+	{
+		const FPackageStorePackage::FExport* SourceExport = KV.Value;
+		FailReasons.AddDefaulted_GetRef().Appendf(TEXT("Public source export '%s' is missing or not public in the target package"),
+			*SourceExport->ObjectName.ToString());
+	}
 
-		if (FailReason.Len() > 0)
+	if (!FailReasons.IsEmpty())
+	{
+		for (const FString& FailReason : FailReasons)
 		{
 			UE_LOG(LogPackageStoreOptimizer, Warning,
 				TEXT("Redirection target package '%s' (0x%llX) for '%s' (0x%llX) - %s"),
@@ -1333,11 +1232,10 @@ bool FPackageStoreOptimizer::VerifyRedirect(
 				*SourcePackage->Name.ToString(),
 				SourcePackage->Id.ValueForDebugging(),
 				*FailReason);
-			bSuccess = false;
 		}
+		return false;
 	}
-
-	return bSuccess;
+	return true;
 }
 
 void FPackageStoreOptimizer::ProcessRedirects(const TMap<FPackageId, FPackageStorePackage*>& PackagesMap, bool bIsBuildingDLC) const
