ZenCook: Add async write of additional files like .ufont to the sandbox directory before compressing for Zen.
Fix problem of missing additional files in .pak files when packaging based on a ZenCook.

#rnx

diff --git a/Engine/Source/Developer/IoStoreUtilities/Private/ZenStoreWriter.cpp b/Engine/Source/Developer/IoStoreUtilities/Private/ZenStoreWriter.cpp
--- a/Engine/Source/Developer/IoStoreUtilities/Private/ZenStoreWriter.cpp
+++ b/Engine/Source/Developer/IoStoreUtilities/Private/ZenStoreWriter.cpp
@@ -335,8 +335,37 @@ void FZenStoreWriter::WriteAdditionalFile(const FAdditionalFileInfo& Info, const
 	
 	FileData.EnsureOwned();
 
-	FileEntry.CompressedPayload = Async(EAsyncExecution::TaskGraph, [this, FileData]()
-	{ 
+	auto WriteToFile = [](const FString& Filename, const FIoBuffer& FileData)
+	{
+		IFileManager& FileManager = IFileManager::Get();
+		int64 DataSize = IntCastChecked<int64>(FileData.DataSize());
+
+		for (int Tries = 0; Tries < 3; ++Tries)
+		{
+			if (FArchive* Ar = FileManager.CreateFileWriter(*Filename))
+			{
+				Ar->Serialize(const_cast<uint8*>(FileData.GetData()), DataSize);
+				bool bArchiveError = Ar->IsError();
+				delete Ar;
+
+				int64 ActualSize = FileManager.FileSize(*Filename);
+				if (ActualSize != DataSize)
+				{
+					FileManager.Delete(*Filename);
+
+					UE_LOG(LogZenStoreWriter, Fatal, TEXT("Could not save to %s! Tried to write %" INT64_FMT " bytes but resultant size was %" INT64_FMT ".%s"),
+						*Filename, DataSize, ActualSize, bArchiveError ? TEXT(" Ar->Serialize failed.") : TEXT(""));
+				}
+				return;
+			}
+		}
+
+		UE_LOG(LogZenStoreWriter, Fatal, TEXT("Could not write to %s!"), *Filename);
+	};
+
+	FileEntry.CompressedPayload = Async(EAsyncExecution::TaskGraph, [this, Info, FileData, WriteToFile]()
+	{
+		WriteToFile(Info.Filename, FileData);
 		return FCompressedBuffer::Compress(FSharedBuffer::MakeView(FileData.GetView()), Compressor, CompressionLevel);
 	});
 
