iostore -list -csv commandlet can take a dir for csv argument and write one csv file per container

diff --git a/Engine/Source/Developer/IoStoreUtilities/Private/IoStoreUtilities.cpp b/Engine/Source/Developer/IoStoreUtilities/Private/IoStoreUtilities.cpp
--- a/Engine/Source/Developer/IoStoreUtilities/Private/IoStoreUtilities.cpp
+++ b/Engine/Source/Developer/IoStoreUtilities/Private/IoStoreUtilities.cpp
@@ -3985,11 +3985,20 @@ int32 ListContainer(
 		UE_LOG(LogIoStore, Error, TEXT("Container '%s' doesn't exist and no container matches wildcard."), *ContainerPathOrWildcard);
 		return -1;
 	}
+	
+	// if CsvPath is a dir, not a file name, then write one csv per container to the dir
+	// otherwise, write all contents to one big csv
+	bool bCsvPathIsDir = IFileManager::Get().DirectoryExists(*CsvPath);
 
-	TUniquePtr<FOutputDeviceFile> Out = MakeUnique<FOutputDeviceFile>(*CsvPath, true);
-	Out->SetSuppressEventTag(true);
+	TUniquePtr<FOutputDeviceFile> AllContainersOutCsvFile;
+	if ( ! bCsvPathIsDir )
+	{
+		AllContainersOutCsvFile = MakeUnique<FOutputDeviceFile>(*CsvPath, true);
+		
+		AllContainersOutCsvFile->SetSuppressEventTag(true);
+	}
 
-	Out->Log(TEXT("OrderInContainer, ChunkId, PackageId, PackageName, Filename, ContainerName, Offset, OffsetOnDisk, Size, CompressedSize, Hash, ChunkType"));
+	bool bNeedsHeader = true;
 
 	for (const FString& ContainerFilePath : ContainerFilePaths)
 	{
@@ -3999,6 +4008,35 @@ int32 ListContainer(
 			UE_LOG(LogIoStore, Warning, TEXT("Failed to read container '%s'"), *ContainerFilePath);
 			continue;
 		}
+		
+		TUniquePtr<FOutputDeviceFile> PerContainerOutCsvFile;
+		FOutputDeviceFile * Out;
+		
+		if ( AllContainersOutCsvFile != nullptr )
+		{
+			Out = AllContainersOutCsvFile.Get();
+		}
+		else
+		{
+			// ContainerFilePath is a .utoc
+			//	change to .csv and put in CsvPath
+			FString PerContainerCsvPath = CsvPath / FPaths::GetCleanFilename(ContainerFilePath);
+			PerContainerCsvPath = FPaths::ChangeExtension(PerContainerCsvPath, FString(TEXT(".csv")) );
+			
+			UE_LOG(LogIoStore, Display, TEXT("PerContainerCsvPath: '%s'"), *PerContainerCsvPath);
+
+			PerContainerOutCsvFile = MakeUnique<FOutputDeviceFile>(*PerContainerCsvPath, true);
+
+			Out = PerContainerOutCsvFile.Get();
+			Out->SetSuppressEventTag(true);
+			bNeedsHeader = true;
+		}
+
+		if ( bNeedsHeader )
+		{
+			Out->Log(TEXT("OrderInContainer, ChunkId, PackageId, PackageName, Filename, ContainerName, Offset, OffsetOnDisk, Size, CompressedSize, Hash, ChunkType"));
+			bNeedsHeader = false;
+		}
 
 		if (!EnumHasAnyFlags(Reader->GetContainerFlags(), EIoContainerFlags::Indexed))
 		{
