MPCook: ZenStoreWriter: Replicate ZenStoreWriter artifacts data.
#rnx

diff --git a/Engine/Source/Developer/IoStoreUtilities/Private/ZenFileSystemManifest.cpp b/Engine/Source/Developer/IoStoreUtilities/Private/ZenFileSystemManifest.cpp
--- a/Engine/Source/Developer/IoStoreUtilities/Private/ZenFileSystemManifest.cpp
+++ b/Engine/Source/Developer/IoStoreUtilities/Private/ZenFileSystemManifest.cpp
@@ -1,8 +1,10 @@
 // Copyright Epic Games, Inc. All Rights Reserved.
 
 #include "ZenFileSystemManifest.h"
-#include "Interfaces/ITargetPlatform.h"
+
+#include "Algo/Sort.h"
 #include "HAL/PlatformFileManager.h"
+#include "Interfaces/ITargetPlatform.h"
 #include "Misc/App.h"
 #include "Misc/FileHelper.h"
 #include "Misc/Paths.h"
@@ -400,10 +402,18 @@ bool FZenFileSystemManifest::Save(const TCHAR* Filename)
 	CsvLines.Add(TEXT("FileId, ServerPath, ClientPath"));
 
 	TStringBuilder<2048> Sb;
+	TArray<const FManifestEntry*> SortedEntries;
+	SortedEntries.Reserve(Entries.Num());
 	for (const FManifestEntry& Entry : Entries)
+	{
+		SortedEntries.Add(&Entry);
+	}
+	Algo::Sort(SortedEntries, [](const FManifestEntry* A, const FManifestEntry* B) { return A->ClientPath < B->ClientPath; });
+
+	for (const FManifestEntry* Entry : SortedEntries)
 	{
 		Sb.Reset();
-		Sb << Entry.FileChunkId << TEXT(", ") << Entry.ServerPath << TEXT(", ") << *Entry.ClientPath;
+		Sb << Entry->FileChunkId << TEXT(", ") << Entry->ServerPath << TEXT(", ") << *Entry->ClientPath;
 		CsvLines.Add(Sb.ToString());
 	}
 
diff --git a/Engine/Source/Developer/IoStoreUtilities/Private/ZenStoreWriter.cpp b/Engine/Source/Developer/IoStoreUtilities/Private/ZenStoreWriter.cpp
--- a/Engine/Source/Developer/IoStoreUtilities/Private/ZenStoreWriter.cpp
+++ b/Engine/Source/Developer/IoStoreUtilities/Private/ZenStoreWriter.cpp
@@ -11,8 +11,10 @@
 #include "Async/Async.h"
 #include "Containers/Queue.h"
 #include "Serialization/ArrayReader.h"
+#include "Serialization/CompactBinaryContainerSerialization.h"
 #include "Serialization/CompactBinaryPackage.h"
 #include "Serialization/CompactBinaryWriter.h"
+#include "Serialization/CompactBinarySerialization.h"
 #include "Serialization/LargeMemoryWriter.h" 
 #include "HAL/FileManager.h"
 #include "HAL/PlatformFileManager.h"
@@ -154,6 +156,7 @@ FZenStoreWriter::FZenStoreWriter(
 	, PackageStoreOptimizer(new FPackageStoreOptimizer())
 	, CookMode(ICookedPackageWriter::FCookInfo::CookByTheBookMode)
 	, bInitialized(false)
+	, bProvidePerPackageResults(false)
 {
 	StaticInit();
 
@@ -327,8 +330,6 @@ void FZenStoreWriter::WriteBulkData(const FBulkDataInfo& Info, const FIoBuffer&
 
 void FZenStoreWriter::WriteAdditionalFile(const FAdditionalFileInfo& Info, const FIoBuffer& FileData)
 {
-	const FZenFileSystemManifest::FManifestEntry& ManifestEntry = ZenFileSystemManifest->CreateManifestEntry(Info.Filename);
-	
 	FPendingPackageState& ExistingState = GetPendingPackage(Info.PackageName);
 
 	FFileDataEntry& FileEntry = ExistingState.FileData.AddDefaulted_GetRef();
@@ -369,10 +370,16 @@ void FZenStoreWriter::WriteAdditionalFile(const FAdditionalFileInfo& Info, const
 		return FCompressedBuffer::Compress(FSharedBuffer::MakeView(FileData.GetView()), Compressor, CompressionLevel);
 	});
 
+	const FZenFileSystemManifest::FManifestEntry& ManifestEntry = ZenFileSystemManifest->CreateManifestEntry(Info.Filename);
 	FileEntry.Info					= Info;
 	FileEntry.Info.ChunkId			= ManifestEntry.FileChunkId;
 	FileEntry.ZenManifestServerPath = ManifestEntry.ServerPath;
 	FileEntry.ZenManifestClientPath = ManifestEntry.ClientPath;
+
+	if (bProvidePerPackageResults)
+	{
+		PackageAdditionalFiles.FindOrAdd(Info.PackageName).Add(Info.Filename);
+	}
 }
 
 void FZenStoreWriter::WriteLinkerAdditionalData(const FLinkerAdditionalDataInfo& Info, const FIoBuffer& Data, const TArray<FFileRegion>& FileRegions)
@@ -525,6 +532,7 @@ void FZenStoreWriter::BeginCook(const FCookInfo& Info)
 	else
 	{
 		PackageStoreManifest.SetTrackPackageData(true);
+		bProvidePerPackageResults = true;
 	}
 	AllPackageHashes.Empty();
 
@@ -573,24 +581,24 @@ void FZenStoreWriter::EndCook(const FCookInfo& Info)
 	CommitQueue->CompleteAdding();
 	CommitThread.Wait();
 
-	FCbPackage Pkg;
-	FCbWriter PackageObj;
+	if (!Info.bWorkerOnSharedSandbox)
+	{
+		FCbPackage Pkg;
+		FCbWriter PackageObj;
 	
-	PackageObj.BeginObject();
-	PackageObj << "key" << "EndCook";
+		PackageObj.BeginObject();
+		PackageObj << "key" << "EndCook";
 
-	CreateProjectMetaData(Pkg, PackageObj);
+		CreateProjectMetaData(Pkg, PackageObj);
 
-	PackageObj.EndObject();
-	FCbObject Obj = PackageObj.Save().AsObject();
+		PackageObj.EndObject();
+		FCbObject Obj = PackageObj.Save().AsObject();
 
-	Pkg.SetObject(Obj);
+		Pkg.SetObject(Obj);
 
-	TIoStatusOr<uint64> Status = HttpClient->EndBuildPass(Pkg);
-	UE_CLOG(!Status.IsOk(), LogZenStoreWriter, Fatal, TEXT("Failed to append OpLog and end the build pass"));
+		TIoStatusOr<uint64> Status = HttpClient->EndBuildPass(Pkg);
+		UE_CLOG(!Status.IsOk(), LogZenStoreWriter, Fatal, TEXT("Failed to append OpLog and end the build pass"));
 
-	if (!Info.bWorkerOnSharedSandbox)
-	{
 		PackageStoreManifest.Save(*(MetadataDirectoryPath / TEXT("packagestore.manifest")));
 	}
 
@@ -637,6 +645,11 @@ void FZenStoreWriter::CommitPackage(FCommitPackageInfo&& Info)
 	{
 		FPendingPackageState& ExistingState = GetPendingPackage(Info.PackageName);
 		ExistingState.PackageHashes = new FPackageHashes();
+		if (bProvidePerPackageResults)
+		{
+			ExistingState.PackageHashesCompletionPromise.Reset(new TPromise<int>());
+			ExistingState.PackageHashes->CompletionFuture = ExistingState.PackageHashesCompletionPromise->GetFuture();
+		}
 
 		if (Info.Status == ECommitStatus::Success)
 		{
@@ -874,6 +887,12 @@ void FZenStoreWriter::CommitPackageInternal(FZenCommitInfo&& ZenCommitInfo)
 		BroadcastCommit(CommitEventArgs);
 	}
 
+	if (PackageState->PackageHashesCompletionPromise)
+	{
+		// Setting the CompletionFuture value may call arbitrary continuation code, so it
+		// must be done outside of any lock.
+		PackageState->PackageHashesCompletionPromise->EmplaceValue(0);
+	}
 	UE::SavePackageUtilities::DecrementOutstandingAsyncWrites();
 }
 
@@ -1059,23 +1078,95 @@ void FZenStoreWriter::MarkPackagesUpToDate(TArrayView<const FName> UpToDatePacka
 
 TFuture<FCbObject> FZenStoreWriter::WriteMPCookMessageForPackage(FName PackageName)
 {
-	// MPCOOKTODO: Replicate the other accumulated data in FZenStoreWriter
-	FCbWriter Writer;
-	Writer.BeginObject();
-	Writer.SetName("Manifest");
-	PackageStoreManifest.WritePackage(Writer, PackageName);
-	Writer.EndObject();
-	FCbObject Object = Writer.Save().AsObject();
-
-	TPromise<FCbObject> Promise;
-	Promise.SetValue(MoveTemp(Object));
-	return Promise.GetFuture();
+	FCbWriter ManifestWriter;
+	PackageStoreManifest.WritePackage(ManifestWriter, PackageName);
+	FCbFieldIterator ManifestField = ManifestWriter.Save();
+
+	TArray<FString> AdditionalFiles;
+	PackageAdditionalFiles.RemoveAndCopyValue(PackageName, AdditionalFiles);
+
+	TRefCountPtr<FPackageHashes> PackageHashes;
+	AllPackageHashes.RemoveAndCopyValue(PackageName, PackageHashes);
+
+	auto ComposeMessage =
+	[ManifestField = MoveTemp(ManifestField), AdditionalFiles=MoveTemp(AdditionalFiles)](FPackageHashes* PackageHashes)
+	{
+		FCbWriter Writer;
+		Writer.BeginObject();
+		Writer << "Manifest" << ManifestField;
+		if (!AdditionalFiles.IsEmpty())
+		{
+			Writer << "AdditionalFiles" << AdditionalFiles;
+		}
+		if (PackageHashes)
+		{
+			Writer << "PackageHash" << PackageHashes->PackageHash;
+			Writer << "ChunkHashes" << PackageHashes->ChunkHashes.Array();
+		}
+		Writer.EndObject();
+		return Writer.Save().AsObject();
+	};
 
+	if (PackageHashes && PackageHashes->CompletionFuture.IsValid())
+	{
+		TUniquePtr<TPromise<FCbObject>> Promise(new TPromise<FCbObject>());
+		TFuture<FCbObject> ResultFuture = Promise->GetFuture();
+		PackageHashes->CompletionFuture.Next(
+			[PackageHashes, Promise = MoveTemp(Promise), ComposeMessage = MoveTemp(ComposeMessage)](int)
+			{
+				Promise->SetValue(ComposeMessage(PackageHashes.GetReference()));
+			});
+		return ResultFuture;
+	}
+	else
+	{
+		TPromise<FCbObject> Promise;
+		Promise.SetValue(ComposeMessage(PackageHashes.GetReference()));
+		return Promise.GetFuture();
+	}
 }
 
 bool FZenStoreWriter::TryReadMPCookMessageForPackage(FName PackageName, FCbObjectView Message)
 {
-	return PackageStoreManifest.TryReadPackage(Message["Manifest"], PackageName);
+	bool bOk = PackageStoreManifest.TryReadPackage(Message["Manifest"], PackageName);
+	if (!bOk)
+	{
+		return false;
+	}
+
+	TArray<FString> AdditionalFiles;
+	if (LoadFromCompactBinary(Message["AdditionalFiles"], AdditionalFiles))
+	{
+		for (const FString& Filename : AdditionalFiles)
+		{
+			ZenFileSystemManifest->CreateManifestEntry(Filename);
+		}
+	}
+
+	TRefCountPtr<FPackageHashes> ThisPackageHashes(new FPackageHashes());
+	if (LoadFromCompactBinary(Message["PackageHash"], ThisPackageHashes->PackageHash))
+	{
+		TArray<TPair<FIoChunkId, FIoHash>> LocalChunkHashes;
+		bOk = LoadFromCompactBinary(Message["ChunkHashes"], LocalChunkHashes) & bOk;
+		if (bOk)
+		{
+			for (TPair<FIoChunkId, FIoHash>& Pair : LocalChunkHashes)
+			{
+				ThisPackageHashes->ChunkHashes.Add(Pair.Key, Pair.Value);
+			}
+			bool bAlreadyExisted = false;
+			TRefCountPtr<FPackageHashes>& ExistingPackageHashes = AllPackageHashes.FindOrAdd(PackageName);
+			bAlreadyExisted = ExistingPackageHashes.IsValid();
+			ExistingPackageHashes = ThisPackageHashes;
+			if (bAlreadyExisted)
+			{
+				UE_LOG(LogSavePackage, Error, TEXT("FLooseCookedPackageWriter encountered the same package twice in a cook! (%s)"),
+					*PackageName.ToString());
+			}
+		}
+	}
+
+	return bOk;
 }
 
 void FZenStoreWriter::CreateProjectMetaData(FCbPackage& Pkg, FCbWriter& PackageObj)
diff --git a/Engine/Source/Developer/IoStoreUtilities/Public/ZenStoreWriter.h b/Engine/Source/Developer/IoStoreUtilities/Public/ZenStoreWriter.h
--- a/Engine/Source/Developer/IoStoreUtilities/Public/ZenStoreWriter.h
+++ b/Engine/Source/Developer/IoStoreUtilities/Public/ZenStoreWriter.h
@@ -141,6 +141,7 @@ private:
 		TArray<FBulkDataEntry> BulkData;
 		TArray<FFileDataEntry> FileData;
 		TRefCountPtr<FPackageHashes> PackageHashes;
+		TUniquePtr<TPromise<int>> PackageHashesCompletionPromise;
 	};
 
 	FPendingPackageState& GetPendingPackage(const FName& PackageName)
@@ -197,7 +198,8 @@ private:
 	TMap<FName, int32>					PackageNameToIndex;
 
 	TUniquePtr<FZenFileSystemManifest>	ZenFileSystemManifest;
-	
+	TMap<FName, TArray<FString>>		PackageAdditionalFiles;
+
 	FEntryCreatedEvent					EntryCreatedEvent;
 	FCriticalSection					CommitEventCriticalSection;
 	FCommitEvent						CommitEvent;
@@ -214,6 +216,7 @@ private:
 	TFuture<void>						CommitThread;
 
 	bool								bInitialized;
+	bool								bProvidePerPackageResults;
 
 	static void StaticInit();
 	static bool IsReservedOplogKey(FUtf8StringView Key);
diff --git a/Engine/Source/Editor/UnrealEd/Private/CookOnTheFlyServer.cpp b/Engine/Source/Editor/UnrealEd/Private/CookOnTheFlyServer.cpp
--- a/Engine/Source/Editor/UnrealEd/Private/CookOnTheFlyServer.cpp
+++ b/Engine/Source/Editor/UnrealEd/Private/CookOnTheFlyServer.cpp
@@ -9725,6 +9725,7 @@ FBeginCookContext UCookOnTheFlyServer::CreateCookWorkerContext()
 {
 	FBeginCookContext BeginContext(*this);
 	*CookByTheBookOptions = CookWorkerClient->ConsumeCookByTheBookOptions();
+	bZenStore = CookWorkerClient->GetInitializationIsZenStore();
 	CookByTheBookOptions->CookTime = 0.0f;
 	CookByTheBookOptions->CookStartTime = FPlatformTime::Seconds();
 	*CookOnTheFlyOptions = CookWorkerClient->ConsumeCookOnTheFlyOptions();
diff --git a/Engine/Source/Editor/UnrealEd/Private/Cooker/CookWorkerClient.cpp b/Engine/Source/Editor/UnrealEd/Private/Cooker/CookWorkerClient.cpp
--- a/Engine/Source/Editor/UnrealEd/Private/Cooker/CookWorkerClient.cpp
+++ b/Engine/Source/Editor/UnrealEd/Private/Cooker/CookWorkerClient.cpp
@@ -125,6 +125,13 @@ ECookInitializationFlags FCookWorkerClient::GetCookInitializationFlags()
 	check(InitialConfigMessage); // Should only be called after TryConnect and before DoneWithInitialSettings
 	return InitialConfigMessage->GetCookInitializationFlags();
 }
+
+bool FCookWorkerClient::GetInitializationIsZenStore()
+{
+	check(InitialConfigMessage); // Should only be called after TryConnect and before DoneWithInitialSettings
+	return InitialConfigMessage->IsZenStore();
+}
+
 FInitializeConfigSettings&& FCookWorkerClient::ConsumeInitializeConfigSettings()
 {
 	check(InitialConfigMessage); // Should only be called after TryConnect and before DoneWithInitialSettings
diff --git a/Engine/Source/Editor/UnrealEd/Private/Cooker/CookWorkerClient.h b/Engine/Source/Editor/UnrealEd/Private/Cooker/CookWorkerClient.h
--- a/Engine/Source/Editor/UnrealEd/Private/Cooker/CookWorkerClient.h
+++ b/Engine/Source/Editor/UnrealEd/Private/Cooker/CookWorkerClient.h
@@ -47,6 +47,7 @@ public:
 	const TArray<ITargetPlatform*>& GetTargetPlatforms() const;
 	/** Consumes the initialization settings from the Director. Only available during initialization. */
 	ECookInitializationFlags GetCookInitializationFlags();
+	bool GetInitializationIsZenStore();
 	FInitializeConfigSettings&& ConsumeInitializeConfigSettings();
 	FBeginCookConfigSettings&& ConsumeBeginCookConfigSettings();
 	FCookByTheBookOptions&& ConsumeCookByTheBookOptions();
diff --git a/Engine/Source/Editor/UnrealEd/Private/Cooker/LooseCookedPackageWriter.h b/Engine/Source/Editor/UnrealEd/Private/Cooker/LooseCookedPackageWriter.h
--- a/Engine/Source/Editor/UnrealEd/Private/Cooker/LooseCookedPackageWriter.h
+++ b/Engine/Source/Editor/UnrealEd/Private/Cooker/LooseCookedPackageWriter.h
@@ -147,6 +147,7 @@ private:
 	TMap<FName, TRefCountPtr<FPackageHashes>> AllPackageHashes;
 
 	TMap<FName, FName> UncookedPathToCookedPath;
+	/** CommitPackage can be called in parallel if using recursive save, so we need a lock for shared containers used during CommitPackage */
 	FCriticalSection PackageHashesLock;
 	FString OutputPath;
 	FString MetadataDirectoryPath;
