UnrealPak: Add more fine grained trace cpu scopes for profiling on remote machines.

#rnx

diff --git a/Engine/Source/Developer/IoStoreUtilities/Private/IoStoreUtilities.cpp b/Engine/Source/Developer/IoStoreUtilities/Private/IoStoreUtilities.cpp
--- a/Engine/Source/Developer/IoStoreUtilities/Private/IoStoreUtilities.cpp
+++ b/Engine/Source/Developer/IoStoreUtilities/Private/IoStoreUtilities.cpp
@@ -652,6 +652,7 @@ public:
 
 	FIoStatus Load(const TCHAR* ManifestFilename)
 	{
+		IOSTORE_CPU_SCOPE(LoadCookedPackageStore);
 		FIoStatus Status = PackageStoreManifest.Load(ManifestFilename);
 		if (!Status.IsOk())
 		{
@@ -1704,6 +1705,7 @@ bool ConvertToIoStoreShaderLibrary(
 	TArray<TTuple<FSHAHash, TArray<FIoChunkId>>>& OutShaderMaps,
 	TArray<TTuple<FSHAHash, TSet<FName>>>& OutShaderMapAssetAssociations)
 {
+	IOSTORE_CPU_SCOPE(ConvertShaderLibrary);
 	// ShaderArchive-MyProject-PCD3D.ushaderbytecode
 	FStringView BaseFileNameView = FPathViews::GetBaseFilename(FileName);
 	int32 FormatStartIndex = -1;
@@ -1926,173 +1928,182 @@ void ProcessShaderLibraries(const FIoStoreArguments& Arguments, TArray<FContaine
 	TMap<FName, TSet<FSHAHash>> PackageNameToShaderMaps;
 	TMap<FContainerTargetSpec*, TSet<FShaderInfo*>> AllContainerShaderLibraryShadersMap;
 
-	for (FContainerTargetSpec* ContainerTarget : ContainerTargets)
 	{
-		TArray<FShaderInfo> Shaders;
-		for (FContainerTargetFile& TargetFile : ContainerTarget->TargetFiles)
+		IOSTORE_CPU_SCOPE(ConvertShaderLibraries);
+		for (FContainerTargetSpec* ContainerTarget : ContainerTargets)
 		{
-			if (TargetFile.ChunkType == EContainerChunkType::ShaderCodeLibrary)
+			TArray<FShaderInfo> Shaders;
+			for (FContainerTargetFile& TargetFile : ContainerTarget->TargetFiles)
 			{
-				TSet<FShaderInfo*>& ContainerShaderLibraryShaders = AllContainerShaderLibraryShadersMap.FindOrAdd(ContainerTarget);
-
-				TArray<TTuple<FSHAHash, TArray<FIoChunkId>>> ShaderMaps;
-				TTuple<FIoChunkId, FIoBuffer> LibraryChunk;
-				TArray<TTuple<FIoChunkId, FIoBuffer, uint32>> CodeChunks;
-				TArray<TTuple<FSHAHash, TSet<FName>>> ShaderMapAssetAssociations;
-				if (!ConvertToIoStoreShaderLibrary(*TargetFile.NormalizedSourcePath, LibraryChunk, CodeChunks, ShaderMaps, ShaderMapAssetAssociations))
+				if (TargetFile.ChunkType == EContainerChunkType::ShaderCodeLibrary)
 				{
-					UE_LOG(LogIoStore, Warning, TEXT("Failed converting shader library '%s'"), *TargetFile.NormalizedSourcePath);
-					continue;
-				}
-				TargetFile.ChunkId = LibraryChunk.Key;
-				TargetFile.SourceBuffer.Emplace(LibraryChunk.Value);
-				const bool bIsGlobalShaderLibrary = FPaths::GetCleanFilename(TargetFile.NormalizedSourcePath).StartsWith(TEXT("ShaderArchive-Global-"));
-				const FShaderInfo::EShaderType ShaderType = bIsGlobalShaderLibrary ? FShaderInfo::Global : FShaderInfo::Normal;
-				for (const TTuple<FIoChunkId, FIoBuffer, uint32>& CodeChunk : CodeChunks)
-				{
-					const FIoChunkId& ShaderChunkId = CodeChunk.Get<0>();
-					FShaderInfo* ShaderInfo = ChunkIdToShaderInfoMap.FindRef(ShaderChunkId);
-					if (!ShaderInfo)
+					TSet<FShaderInfo*>& ContainerShaderLibraryShaders = AllContainerShaderLibraryShadersMap.FindOrAdd(ContainerTarget);
+
+					TArray<TTuple<FSHAHash, TArray<FIoChunkId>>> ShaderMaps;
+					TTuple<FIoChunkId, FIoBuffer> LibraryChunk;
+					TArray<TTuple<FIoChunkId, FIoBuffer, uint32>> CodeChunks;
+					TArray<TTuple<FSHAHash, TSet<FName>>> ShaderMapAssetAssociations;
+					if (!ConvertToIoStoreShaderLibrary(*TargetFile.NormalizedSourcePath, LibraryChunk, CodeChunks, ShaderMaps, ShaderMapAssetAssociations))
 					{
-						ShaderInfo = new FShaderInfo();
-						ShaderInfo->ChunkId = ShaderChunkId;
-						ShaderInfo->CodeIoBuffer = CodeChunk.Get<1>();
-						ShaderInfo->LoadOrderFactor = CodeChunk.Get<2>();
-						OutShaders.Add(ShaderInfo);
-						ChunkIdToShaderInfoMap.Add(ShaderChunkId, ShaderInfo);
+						UE_LOG(LogIoStore, Warning, TEXT("Failed converting shader library '%s'"), *TargetFile.NormalizedSourcePath);
+						continue;
 					}
-					else
+					TargetFile.ChunkId = LibraryChunk.Key;
+					TargetFile.SourceBuffer.Emplace(LibraryChunk.Value);
+					const bool bIsGlobalShaderLibrary = FPaths::GetCleanFilename(TargetFile.NormalizedSourcePath).StartsWith(TEXT("ShaderArchive-Global-"));
+					const FShaderInfo::EShaderType ShaderType = bIsGlobalShaderLibrary ? FShaderInfo::Global : FShaderInfo::Normal;
+					for (const TTuple<FIoChunkId, FIoBuffer, uint32>& CodeChunk : CodeChunks)
 					{
-						// first, make sure that the code is exactly the same
-						if (ShaderInfo->CodeIoBuffer != CodeChunk.Get<1>())
+						const FIoChunkId& ShaderChunkId = CodeChunk.Get<0>();
+						FShaderInfo* ShaderInfo = ChunkIdToShaderInfoMap.FindRef(ShaderChunkId);
+						if (!ShaderInfo)
 						{
-							UE_LOG(LogIoStore, Error, TEXT("Collision of two shader code chunks, same Id (%s), different code. Packaged game will likely crash, not being able to decompress the shaders."), *LexToString(ShaderChunkId) );
+							ShaderInfo = new FShaderInfo();
+							ShaderInfo->ChunkId = ShaderChunkId;
+							ShaderInfo->CodeIoBuffer = CodeChunk.Get<1>();
+							ShaderInfo->LoadOrderFactor = CodeChunk.Get<2>();
+							OutShaders.Add(ShaderInfo);
+							ChunkIdToShaderInfoMap.Add(ShaderChunkId, ShaderInfo);
 						}
-
-						// If we already exist, then we have two separate LoadOrderFactors,
-						// which one we got first affects build determinism. Take the lower.
-						if (CodeChunk.Get<2>() < ShaderInfo->LoadOrderFactor)
+						else
 						{
-							ShaderInfo->LoadOrderFactor = CodeChunk.Get<2>();
+							// first, make sure that the code is exactly the same
+							if (ShaderInfo->CodeIoBuffer != CodeChunk.Get<1>())
+							{
+								UE_LOG(LogIoStore, Error, TEXT("Collision of two shader code chunks, same Id (%s), different code. Packaged game will likely crash, not being able to decompress the shaders."), *LexToString(ShaderChunkId) );
+							}
+
+							// If we already exist, then we have two separate LoadOrderFactors,
+							// which one we got first affects build determinism. Take the lower.
+							if (CodeChunk.Get<2>() < ShaderInfo->LoadOrderFactor)
+							{
+								ShaderInfo->LoadOrderFactor = CodeChunk.Get<2>();
+							}
 						}
-					}
 
 
-					const FShaderInfo::EShaderType* CurrentShaderTypeInContainer = ShaderInfo->TypeInContainer.Find(ContainerTarget);
-					if (!CurrentShaderTypeInContainer || *CurrentShaderTypeInContainer != FShaderInfo::Global)
-					{
-						// If a shader is both global and shared consider it to be global
-						ShaderInfo->TypeInContainer.Add(ContainerTarget, ShaderType);
+						const FShaderInfo::EShaderType* CurrentShaderTypeInContainer = ShaderInfo->TypeInContainer.Find(ContainerTarget);
+						if (!CurrentShaderTypeInContainer || *CurrentShaderTypeInContainer != FShaderInfo::Global)
+						{
+							// If a shader is both global and shared consider it to be global
+							ShaderInfo->TypeInContainer.Add(ContainerTarget, ShaderType);
+						}
+						
+						ContainerShaderLibraryShaders.Add(ShaderInfo);
 					}
-					
-					ContainerShaderLibraryShaders.Add(ShaderInfo);
-				}
-				for (const TTuple<FSHAHash, TSet<FName>>& ShaderMapAssetAssociation : ShaderMapAssetAssociations)
-				{
-					for (const FName& PackageName : ShaderMapAssetAssociation.Value)
+					for (const TTuple<FSHAHash, TSet<FName>>& ShaderMapAssetAssociation : ShaderMapAssetAssociations)
 					{
-						TSet<FSHAHash>& PackageShaderMaps = PackageNameToShaderMaps.FindOrAdd(PackageName);
-						PackageShaderMaps.Add(ShaderMapAssetAssociation.Key);
+						for (const FName& PackageName : ShaderMapAssetAssociation.Value)
+						{
+							TSet<FSHAHash>& PackageShaderMaps = PackageNameToShaderMaps.FindOrAdd(PackageName);
+							PackageShaderMaps.Add(ShaderMapAssetAssociation.Key);
+						}
 					}
-				}
 
-				for (TTuple<FSHAHash, TArray<FIoChunkId>>& ShaderMap : ShaderMaps)
-				{
-					ShaderChunkIdsByShaderMapHash.Add(ShaderMap.Key, MoveTemp(ShaderMap.Value));
-				}
-			} // end if containerchunktype shadercodelibrary
-		} // end foreach targetfile
-	} // end foreach container
+					for (TTuple<FSHAHash, TArray<FIoChunkId>>& ShaderMap : ShaderMaps)
+					{
+						ShaderChunkIdsByShaderMapHash.Add(ShaderMap.Key, MoveTemp(ShaderMap.Value));
+					}
+				} // end if containerchunktype shadercodelibrary
+			} // end foreach targetfile
+		} // end foreach container
+	}
 
 	// 1. Update ShaderInfos with which packages we reference.
 	// 2. Add to packages which shaders we use.
 	// 3. Add to PackageStore what shaders we use.
-	for (FContainerTargetSpec* ContainerTarget : ContainerTargets)
 	{
-		for (FLegacyCookedPackage* Package : ContainerTarget->Packages)
+		IOSTORE_CPU_SCOPE(UpdatePackageStoreShaders);
+		for (FContainerTargetSpec* ContainerTarget : ContainerTargets)
 		{
-			TSet<FSHAHash>* FindShaderMapHashes = PackageNameToShaderMaps.Find(Package->PackageName);
-			if (FindShaderMapHashes)
+			for (FLegacyCookedPackage* Package : ContainerTarget->Packages)
 			{
-				for (const FSHAHash& ShaderMapHash : *FindShaderMapHashes)
+				TSet<FSHAHash>* FindShaderMapHashes = PackageNameToShaderMaps.Find(Package->PackageName);
+				if (FindShaderMapHashes)
 				{
-					const TArray<FIoChunkId>* FindChunkIds = ShaderChunkIdsByShaderMapHash.Find(ShaderMapHash);
-					if (!FindChunkIds)
+					for (const FSHAHash& ShaderMapHash : *FindShaderMapHashes)
 					{
-						UE_LOG(LogIoStore, Warning, TEXT("Package '%s' in '%s' referencing missing shader map '%s'"), *Package->PackageName.ToString(), *ContainerTarget->Name.ToString(), *ShaderMapHash.ToString());
-						continue;
-					}
-					Package->OptimizedPackage->AddShaderMapHash(ShaderMapHash);
-					for (const FIoChunkId& ShaderChunkId : *FindChunkIds)
-					{
-						FShaderInfo* ShaderInfo = ChunkIdToShaderInfoMap.FindRef(ShaderChunkId);
-						if (!ShaderInfo)
+						const TArray<FIoChunkId>* FindChunkIds = ShaderChunkIdsByShaderMapHash.Find(ShaderMapHash);
+						if (!FindChunkIds)
 						{
-							UE_LOG(LogIoStore, Warning, TEXT("Package '%s' in '%s' referencing missing shader with chunk id '%s'"), *Package->PackageName.ToString(), *ContainerTarget->Name.ToString(), *BytesToHex(ShaderChunkId.GetData(), ShaderChunkId.GetSize()));
+							UE_LOG(LogIoStore, Warning, TEXT("Package '%s' in '%s' referencing missing shader map '%s'"), *Package->PackageName.ToString(), *ContainerTarget->Name.ToString(), *ShaderMapHash.ToString());
 							continue;
 						}
+						Package->OptimizedPackage->AddShaderMapHash(ShaderMapHash);
+						for (const FIoChunkId& ShaderChunkId : *FindChunkIds)
+						{
+							FShaderInfo* ShaderInfo = ChunkIdToShaderInfoMap.FindRef(ShaderChunkId);
+							if (!ShaderInfo)
+							{
+								UE_LOG(LogIoStore, Warning, TEXT("Package '%s' in '%s' referencing missing shader with chunk id '%s'"), *Package->PackageName.ToString(), *ContainerTarget->Name.ToString(), *BytesToHex(ShaderChunkId.GetData(), ShaderChunkId.GetSize()));
+								continue;
+							}
 
-						check(ShaderInfo);
-						ShaderInfo->ReferencedByPackages.Add(Package);
-						Package->Shaders.AddUnique(ShaderInfo);
+							check(ShaderInfo);
+							ShaderInfo->ReferencedByPackages.Add(Package);
+							Package->Shaders.AddUnique(ShaderInfo);
+						}
 					}
 				}
 			}
 		}
 	}
 
-	for (FContainerTargetSpec* ContainerTarget : ContainerTargets)
 	{
-		auto AddShaderTargetFile = [&ContainerTarget](FShaderInfo* ShaderInfo)
+		IOSTORE_CPU_SCOPE(AddShaderChunks);
+		for (FContainerTargetSpec* ContainerTarget : ContainerTargets)
 		{
-			FContainerTargetFile& ShaderTargetFile = ContainerTarget->TargetFiles.AddDefaulted_GetRef();
-			ShaderTargetFile.ContainerTarget = ContainerTarget;
-			ShaderTargetFile.ChunkId = ShaderInfo->ChunkId;
-			ShaderTargetFile.ChunkType = EContainerChunkType::ShaderCode;
-			ShaderTargetFile.bForceUncompressed = true;
-			ShaderTargetFile.SourceBuffer.Emplace(ShaderInfo->CodeIoBuffer);
-			ShaderTargetFile.SourceSize = ShaderInfo->CodeIoBuffer.DataSize();
-		};
+			auto AddShaderTargetFile = [&ContainerTarget](FShaderInfo* ShaderInfo)
+			{
+				FContainerTargetFile& ShaderTargetFile = ContainerTarget->TargetFiles.AddDefaulted_GetRef();
+				ShaderTargetFile.ContainerTarget = ContainerTarget;
+				ShaderTargetFile.ChunkId = ShaderInfo->ChunkId;
+				ShaderTargetFile.ChunkType = EContainerChunkType::ShaderCode;
+				ShaderTargetFile.bForceUncompressed = true;
+				ShaderTargetFile.SourceBuffer.Emplace(ShaderInfo->CodeIoBuffer);
+				ShaderTargetFile.SourceSize = ShaderInfo->CodeIoBuffer.DataSize();
+			};
 
-		const TSet<FShaderInfo*>* FindContainerShaderLibraryShaders = AllContainerShaderLibraryShadersMap.Find(ContainerTarget);
-		if (FindContainerShaderLibraryShaders)
-		{
-			for (FLegacyCookedPackage* Package : ContainerTarget->Packages)
+			const TSet<FShaderInfo*>* FindContainerShaderLibraryShaders = AllContainerShaderLibraryShadersMap.Find(ContainerTarget);
+			if (FindContainerShaderLibraryShaders)
 			{
-				for (FShaderInfo* ShaderInfo : Package->Shaders)
+				for (FLegacyCookedPackage* Package : ContainerTarget->Packages)
 				{
-					if (ShaderInfo->ReferencedByPackages.Num() == 1)
+					for (FShaderInfo* ShaderInfo : Package->Shaders)
 					{
-						FShaderInfo::EShaderType* ShaderType = ShaderInfo->TypeInContainer.Find(ContainerTarget);
-						if (ShaderType && *ShaderType != FShaderInfo::Global)
+						if (ShaderInfo->ReferencedByPackages.Num() == 1)
 						{
-							*ShaderType = FShaderInfo::Inline;
+							FShaderInfo::EShaderType* ShaderType = ShaderInfo->TypeInContainer.Find(ContainerTarget);
+							if (ShaderType && *ShaderType != FShaderInfo::Global)
+							{
+								*ShaderType = FShaderInfo::Inline;
+							}
 						}
 					}
 				}
-			}
-			for (FShaderInfo* ShaderInfo : *FindContainerShaderLibraryShaders)
-			{
-				FShaderInfo::EShaderType* ShaderType = ShaderInfo->TypeInContainer.Find(ContainerTarget);
-				check(ShaderType);
-				if (*ShaderType == FShaderInfo::Global)
-				{
-					ContainerTarget->GlobalShaders.Add(ShaderInfo);
-				}
-				else if (*ShaderType == FShaderInfo::Inline)
-				{
-					ContainerTarget->InlineShaders.Add(ShaderInfo);
-				}
-				else if (ShaderInfo->ReferencedByPackages.Num() > 1)
+				for (FShaderInfo* ShaderInfo : *FindContainerShaderLibraryShaders)
 				{
-					ContainerTarget->SharedShaders.Add(ShaderInfo);
-				}
-				else
-				{
-					// If there are unreferenced shaders they will go in here and be sorted last
-					ContainerTarget->UniqueShaders.Add(ShaderInfo);
+					FShaderInfo::EShaderType* ShaderType = ShaderInfo->TypeInContainer.Find(ContainerTarget);
+					check(ShaderType);
+					if (*ShaderType == FShaderInfo::Global)
+					{
+						ContainerTarget->GlobalShaders.Add(ShaderInfo);
+					}
+					else if (*ShaderType == FShaderInfo::Inline)
+					{
+						ContainerTarget->InlineShaders.Add(ShaderInfo);
+					}
+					else if (ShaderInfo->ReferencedByPackages.Num() > 1)
+					{
+						ContainerTarget->SharedShaders.Add(ShaderInfo);
+					}
+					else
+					{
+						// If there are unreferenced shaders they will go in here and be sorted last
+						ContainerTarget->UniqueShaders.Add(ShaderInfo);
+					}
+					AddShaderTargetFile(ShaderInfo);
 				}
-				AddShaderTargetFile(ShaderInfo);
 			}
 		}
 	}
@@ -3336,6 +3347,7 @@ public:
 
 int32 CreateTarget(const FIoStoreArguments& Arguments, const FIoStoreWriterSettings& GeneralIoWriterSettings)
 {
+	IOSTORE_CPU_SCOPE(CreateTarget);
 	TGuardValue<int32> GuardAllowUnversionedContentInEditor(GAllowUnversionedContentInEditor, 1);
 
 #if OUTPUT_CHUNKID_DIRECTORY
@@ -3378,15 +3390,20 @@ int32 CreateTarget(const FIoStoreArguments& Arguments, const FIoStoreWriterSetti
 		InitializeContainerTargetsAndPackages(Arguments, Packages, PackageNameMap, PackageIdMap, ContainerTargets, PackageStoreOptimizer);
 	}
 
-	TUniquePtr<FIoStoreWriterContext> IoStoreWriterContext(new FIoStoreWriterContext());
-	FIoStatus IoStatus = IoStoreWriterContext->Initialize(GeneralIoWriterSettings);
-	check(IoStatus.IsOk());
+	TUniquePtr<FIoStoreWriterContext> IoStoreWriterContext;
+	{
+		IOSTORE_CPU_SCOPE(InitializeIoStoreWriters);
+		IoStoreWriterContext.Reset(new FIoStoreWriterContext());
+		FIoStatus IoStatus = IoStoreWriterContext->Initialize(GeneralIoWriterSettings);
+		check(IoStatus.IsOk());
+	}
 	TArray<TSharedPtr<IIoStoreWriter>> IoStoreWriters;
 	TSharedPtr<IIoStoreWriter> GlobalIoStoreWriter;
 	{
-		IOSTORE_CPU_SCOPE(InitializeIoStoreWriters);
+		IOSTORE_CPU_SCOPE(InitializeWriters);
 		if (!Arguments.IsDLC())
 		{
+			IOSTORE_CPU_SCOPE(InitializeGlobalWriter);
 			FIoContainerSettings GlobalContainerSettings;
 			if (Arguments.bSign)
 			{
@@ -3398,6 +3415,7 @@ int32 CreateTarget(const FIoStoreArguments& Arguments, const FIoStoreWriterSetti
 		}
 		for (FContainerTargetSpec* ContainerTarget : ContainerTargets)
 		{
+			IOSTORE_CPU_SCOPE(InitializeWriter);
 			check(ContainerTarget->ContainerId.IsValid());
 			if (!ContainerTarget->OutputPath.IsEmpty())
 			{
@@ -3449,6 +3467,7 @@ int32 CreateTarget(const FIoStoreArguments& Arguments, const FIoStoreWriterSetti
 
 	if (Arguments.IsDLC() && Arguments.bRemapPluginContentToGame)
 	{
+		IOSTORE_CPU_SCOPE(RemapPluginContent);
 		for (FLegacyCookedPackage* Package : Packages)
 		{
 			const int32 DLCNameLen = Arguments.DLCName.Len() + 1;
@@ -3460,32 +3479,35 @@ int32 CreateTarget(const FIoStoreArguments& Arguments, const FIoStoreWriterSetti
 		}
 	}
 
-	for (FContainerTargetSpec* ContainerTarget : ContainerTargets)
 	{
-		if (ContainerTarget->IoStoreWriter)
+		IOSTORE_CPU_SCOPE(AppendNonPackageChunks);
+		for (FContainerTargetSpec* ContainerTarget : ContainerTargets)
 		{
-			for (FContainerTargetFile& TargetFile : ContainerTarget->TargetFiles)
+			if (ContainerTarget->IoStoreWriter)
 			{
-				if (TargetFile.ChunkType != EContainerChunkType::PackageData && TargetFile.ChunkType != EContainerChunkType::OptionalSegmentPackageData)
+				for (FContainerTargetFile& TargetFile : ContainerTarget->TargetFiles)
 				{
-					FIoWriteOptions WriteOptions;
-					WriteOptions.DebugName = *TargetFile.DestinationPath;
-					WriteOptions.bForceUncompressed = TargetFile.bForceUncompressed;
-					WriteOptions.bIsMemoryMapped = TargetFile.ChunkType == EContainerChunkType::MemoryMappedBulkData;
-					WriteOptions.FileName = TargetFile.DestinationPath;
-					if (TargetFile.ChunkType == EContainerChunkType::OptionalSegmentBulkData)
+					if (TargetFile.ChunkType != EContainerChunkType::PackageData && TargetFile.ChunkType != EContainerChunkType::OptionalSegmentPackageData)
 					{
-						FIoChunkId ChunkId = TargetFile.ChunkId;
-						if (TargetFile.Package->OptimizedOptionalSegmentPackage && TargetFile.Package->OptimizedOptionalSegmentPackage->HasEditorData())
+						FIoWriteOptions WriteOptions;
+						WriteOptions.DebugName = *TargetFile.DestinationPath;
+						WriteOptions.bForceUncompressed = TargetFile.bForceUncompressed;
+						WriteOptions.bIsMemoryMapped = TargetFile.ChunkType == EContainerChunkType::MemoryMappedBulkData;
+						WriteOptions.FileName = TargetFile.DestinationPath;
+						if (TargetFile.ChunkType == EContainerChunkType::OptionalSegmentBulkData)
+						{
+							FIoChunkId ChunkId = TargetFile.ChunkId;
+							if (TargetFile.Package->OptimizedOptionalSegmentPackage && TargetFile.Package->OptimizedOptionalSegmentPackage->HasEditorData())
+							{
+								// Auto optional packages replace the non-optional part when the container is mounted
+								ChunkId = CreateIoChunkId(TargetFile.Package->GlobalPackageId.Value(), 0, EIoChunkType::BulkData);
+							}
+							ContainerTarget->OptionalSegmentIoStoreWriter->Append(ChunkId, WriteRequestManager.Read(TargetFile), WriteOptions);
+						}
+						else
 						{
-							// Auto optional packages replace the non-optional part when the container is mounted
-							ChunkId = CreateIoChunkId(TargetFile.Package->GlobalPackageId.Value(), 0, EIoChunkType::BulkData);
+							ContainerTarget->IoStoreWriter->Append(TargetFile.ChunkId, WriteRequestManager.Read(TargetFile), WriteOptions);
 						}
-						ContainerTarget->OptionalSegmentIoStoreWriter->Append(ChunkId, WriteRequestManager.Read(TargetFile), WriteOptions);
-					}
-					else
-					{
-						ContainerTarget->IoStoreWriter->Append(TargetFile.ChunkId, WriteRequestManager.Read(TargetFile), WriteOptions);
 					}
 				}
 			}
@@ -3493,10 +3515,13 @@ int32 CreateTarget(const FIoStoreArguments& Arguments, const FIoStoreWriterSetti
 	}
 
 	TMap<FPackageId, FPackageStorePackage*> OptimizedPackagesMap;
-	for (FLegacyCookedPackage* Package : Packages)
 	{
-		check(Package->OptimizedPackage);
-		OptimizedPackagesMap.Add(Package->OptimizedPackage->GetId(), Package->OptimizedPackage);
+		IOSTORE_CPU_SCOPE(OptimizedPackagesMap);
+		for (FLegacyCookedPackage* Package : Packages)
+		{
+			check(Package->OptimizedPackage);
+			OptimizedPackagesMap.Add(Package->OptimizedPackage->GetId(), Package->OptimizedPackage);
+		}
 	}
 
 	UE_LOG(LogIoStore, Display, TEXT("Processing redirects..."));
@@ -3506,13 +3531,16 @@ int32 CreateTarget(const FIoStoreArguments& Arguments, const FIoStoreWriterSetti
 	PackageStoreOptimizer.OptimizeExportBundles(OptimizedPackagesMap);
 
 	UE_LOG(LogIoStore, Display, TEXT("Finalizing packages..."));
-	for (FLegacyCookedPackage* Package : Packages)
 	{
-		check(Package->OptimizedPackage);
-		PackageStoreOptimizer.FinalizePackage(Package->OptimizedPackage);
-		if (Package->OptimizedOptionalSegmentPackage)
+		IOSTORE_CPU_SCOPE(FinalizingPackages);
+		for (FLegacyCookedPackage* Package : Packages)
 		{
-			PackageStoreOptimizer.FinalizePackage(Package->OptimizedOptionalSegmentPackage);
+			check(Package->OptimizedPackage);
+			PackageStoreOptimizer.FinalizePackage(Package->OptimizedPackage);
+			if (Package->OptimizedOptionalSegmentPackage)
+			{
+				PackageStoreOptimizer.FinalizePackage(Package->OptimizedOptionalSegmentPackage);
+			}
 		}
 	}
 
@@ -3520,113 +3548,119 @@ int32 CreateTarget(const FIoStoreArguments& Arguments, const FIoStoreWriterSetti
 	FString ClusterCSVPath;
 	if (FParse::Value(FCommandLine::Get(), TEXT("ClusterCSV="), ClusterCSVPath))
 	{
+		IOSTORE_CPU_SCOPE(CreateClusterCSV);
 		ClusterStatsCsv.CreateOutputFile(ClusterCSVPath);
 	}
 	CreateDiskLayout(ContainerTargets, Packages, Arguments.OrderMaps, PackageIdMap, Arguments.bClusterByOrderFilePriority);
 
-	for (FContainerTargetSpec* ContainerTarget : ContainerTargets)
 	{
-		if (ContainerTarget->IoStoreWriter)
+		IOSTORE_CPU_SCOPE(AppendPackageChunks);
+		for (FContainerTargetSpec* ContainerTarget : ContainerTargets)
 		{
-			TArray<FPackageStoreEntryResource> PackageStoreEntries;
-			for (FContainerTargetFile& TargetFile : ContainerTarget->TargetFiles)
+			if (ContainerTarget->IoStoreWriter)
 			{
-				if (TargetFile.ChunkType == EContainerChunkType::PackageData || TargetFile.ChunkType == EContainerChunkType::OptionalSegmentPackageData)
+				TArray<FPackageStoreEntryResource> PackageStoreEntries;
+				for (FContainerTargetFile& TargetFile : ContainerTarget->TargetFiles)
 				{
-					check(TargetFile.Package);
-					FIoWriteOptions WriteOptions;
-					WriteOptions.DebugName = *TargetFile.DestinationPath;
-					WriteOptions.bForceUncompressed = TargetFile.bForceUncompressed;
-					WriteOptions.FileName = TargetFile.DestinationPath;
-					if (TargetFile.ChunkType == EContainerChunkType::OptionalSegmentPackageData)
+					if (TargetFile.ChunkType == EContainerChunkType::PackageData || TargetFile.ChunkType == EContainerChunkType::OptionalSegmentPackageData)
 					{
-						check(ContainerTarget->OptionalSegmentIoStoreWriter);
-						check(TargetFile.Package->OptimizedOptionalSegmentPackage);
-						FIoChunkId ChunkId = TargetFile.ChunkId;
-						if (TargetFile.Package->OptimizedOptionalSegmentPackage->HasEditorData())
+						check(TargetFile.Package);
+						FIoWriteOptions WriteOptions;
+						WriteOptions.DebugName = *TargetFile.DestinationPath;
+						WriteOptions.bForceUncompressed = TargetFile.bForceUncompressed;
+						WriteOptions.FileName = TargetFile.DestinationPath;
+						if (TargetFile.ChunkType == EContainerChunkType::OptionalSegmentPackageData)
+						{
+							check(ContainerTarget->OptionalSegmentIoStoreWriter);
+							check(TargetFile.Package->OptimizedOptionalSegmentPackage);
+							FIoChunkId ChunkId = TargetFile.ChunkId;
+							if (TargetFile.Package->OptimizedOptionalSegmentPackage->HasEditorData())
+							{
+								// Auto optional packages replace the non-optional part when the container is mounted
+								ChunkId = CreateIoChunkId(TargetFile.Package->GlobalPackageId.Value(), 0, EIoChunkType::ExportBundleData);
+							}
+							ContainerTarget->OptionalSegmentIoStoreWriter->Append(ChunkId, WriteRequestManager.Read(TargetFile), WriteOptions);
+						}
+						else
 						{
-							// Auto optional packages replace the non-optional part when the container is mounted
-							ChunkId = CreateIoChunkId(TargetFile.Package->GlobalPackageId.Value(), 0, EIoChunkType::ExportBundleData);
+							ContainerTarget->IoStoreWriter->Append(TargetFile.ChunkId, WriteRequestManager.Read(TargetFile), WriteOptions);
+							PackageStoreEntries.Add(PackageStoreOptimizer.CreatePackageStoreEntry(TargetFile.Package->OptimizedPackage, TargetFile.Package->OptimizedOptionalSegmentPackage));
 						}
-						ContainerTarget->OptionalSegmentIoStoreWriter->Append(ChunkId, WriteRequestManager.Read(TargetFile), WriteOptions);
-					}
-					else
-					{
-						ContainerTarget->IoStoreWriter->Append(TargetFile.ChunkId, WriteRequestManager.Read(TargetFile), WriteOptions);
-						PackageStoreEntries.Add(PackageStoreOptimizer.CreatePackageStoreEntry(TargetFile.Package->OptimizedPackage, TargetFile.Package->OptimizedOptionalSegmentPackage));
 					}
 				}
-			}
-
-			auto WriteContainerHeaderChunk = [](FIoContainerHeader& Header, IIoStoreWriter* IoStoreWriter)
-			{
-				FLargeMemoryWriter HeaderAr(0, true);
-				HeaderAr << Header;
-				int64 DataSize = HeaderAr.TotalSize();
-				FIoBuffer ContainerHeaderBuffer(FIoBuffer::AssumeOwnership, HeaderAr.ReleaseOwnership(), DataSize);
 
-				FIoWriteOptions WriteOptions;
-				WriteOptions.DebugName = TEXT("ContainerHeader");
-				WriteOptions.bForceUncompressed = true;
-				IoStoreWriter->Append(
-					CreateIoChunkId(Header.ContainerId.Value(), 0, EIoChunkType::ContainerHeader),
-					ContainerHeaderBuffer,
-					WriteOptions);
-			};
-
-			ContainerTarget->Header = PackageStoreOptimizer.CreateContainerHeader(ContainerTarget->ContainerId, PackageStoreEntries);
-			WriteContainerHeaderChunk(ContainerTarget->Header, ContainerTarget->IoStoreWriter.Get());
+				auto WriteContainerHeaderChunk = [](FIoContainerHeader& Header, IIoStoreWriter* IoStoreWriter)
+				{
+					FLargeMemoryWriter HeaderAr(0, true);
+					HeaderAr << Header;
+					int64 DataSize = HeaderAr.TotalSize();
+					FIoBuffer ContainerHeaderBuffer(FIoBuffer::AssumeOwnership, HeaderAr.ReleaseOwnership(), DataSize);
 
-			if (ContainerTarget->OptionalSegmentIoStoreWriter)
-			{
-				ContainerTarget->OptionalSegmentHeader = PackageStoreOptimizer.CreateOptionalContainerHeader(ContainerTarget->ContainerId, PackageStoreEntries);
-				WriteContainerHeaderChunk(ContainerTarget->OptionalSegmentHeader, ContainerTarget->OptionalSegmentIoStoreWriter.Get());
+					FIoWriteOptions WriteOptions;
+					WriteOptions.DebugName = TEXT("ContainerHeader");
+					WriteOptions.bForceUncompressed = true;
+					IoStoreWriter->Append(
+						CreateIoChunkId(Header.ContainerId.Value(), 0, EIoChunkType::ContainerHeader),
+						ContainerHeaderBuffer,
+						WriteOptions);
+				};
+
+				ContainerTarget->Header = PackageStoreOptimizer.CreateContainerHeader(ContainerTarget->ContainerId, PackageStoreEntries);
+				WriteContainerHeaderChunk(ContainerTarget->Header, ContainerTarget->IoStoreWriter.Get());
+
+				if (ContainerTarget->OptionalSegmentIoStoreWriter)
+				{
+					ContainerTarget->OptionalSegmentHeader = PackageStoreOptimizer.CreateOptionalContainerHeader(ContainerTarget->ContainerId, PackageStoreEntries);
+					WriteContainerHeaderChunk(ContainerTarget->OptionalSegmentHeader, ContainerTarget->OptionalSegmentIoStoreWriter.Get());
+				}
 			}
-		}
 
-		// Check if we need to dump the final order of the packages. Useful, to debug packing.
-		if (FParse::Param(FCommandLine::Get(), TEXT("writefinalorder")))
-		{
-			FString FinalContainerOrderFile = FPaths::GetPath(ContainerTarget->OutputPath) + FPaths::GetBaseFilename(ContainerTarget->OutputPath) + TEXT("-order.txt");
-			TUniquePtr<FArchive> IoOrderListArchive(IFileManager::Get().CreateFileWriter(*FinalContainerOrderFile));
-			if (IoOrderListArchive)
+			// Check if we need to dump the final order of the packages. Useful, to debug packing.
+			if (FParse::Param(FCommandLine::Get(), TEXT("writefinalorder")))
 			{
-				IoOrderListArchive->SetIsTextFormat(true);
-
-				// The TargetFiles list is not the order written to disk - FIoStoreWriter sorts
-				// on the IdealOrder prior to writing.
-				TArray<const FContainerTargetFile*> SortedTargetFiles;
-				SortedTargetFiles.Reserve(ContainerTarget->TargetFiles.Num());
-				for (const FContainerTargetFile& TargetFile : ContainerTarget->TargetFiles)
+				IOSTORE_CPU_SCOPE(WriteFinalOrder);
+				FString FinalContainerOrderFile = FPaths::GetPath(ContainerTarget->OutputPath) + FPaths::GetBaseFilename(ContainerTarget->OutputPath) + TEXT("-order.txt");
+				TUniquePtr<FArchive> IoOrderListArchive(IFileManager::Get().CreateFileWriter(*FinalContainerOrderFile));
+				if (IoOrderListArchive)
 				{
-					SortedTargetFiles.Add(&TargetFile);
-				}
+					IoOrderListArchive->SetIsTextFormat(true);
 
-				Algo::SortBy(SortedTargetFiles, [](const FContainerTargetFile* TargetFile) { return TargetFile->IdealOrder; });
+					// The TargetFiles list is not the order written to disk - FIoStoreWriter sorts
+					// on the IdealOrder prior to writing.
+					TArray<const FContainerTargetFile*> SortedTargetFiles;
+					SortedTargetFiles.Reserve(ContainerTarget->TargetFiles.Num());
+					for (const FContainerTargetFile& TargetFile : ContainerTarget->TargetFiles)
+					{
+						SortedTargetFiles.Add(&TargetFile);
+					}
 
-				for (const FContainerTargetFile* TargetFile : SortedTargetFiles)
-				{
-					TStringBuilder<256> Line;
-					Line.Append(LexToString(TargetFile->ChunkId));
+					Algo::SortBy(SortedTargetFiles, [](const FContainerTargetFile* TargetFile) { return TargetFile->IdealOrder; });
 
-					if (TargetFile->Package)
+					for (const FContainerTargetFile* TargetFile : SortedTargetFiles)
 					{
-						Line.Append(TEXT(" "));
-						Line.Append(TargetFile->Package->FileName);
+						TStringBuilder<256> Line;
+						Line.Append(LexToString(TargetFile->ChunkId));
+
+						if (TargetFile->Package)
+						{
+							Line.Append(TEXT(" "));
+							Line.Append(TargetFile->Package->FileName);
+						}
+
+						IoOrderListArchive->Logf(TEXT("%s"), Line.ToString());
 					}
 
-					IoOrderListArchive->Logf(TEXT("%s"), Line.ToString());
+					IoOrderListArchive->Close();
 				}
-
-				IoOrderListArchive->Close();
 			}
-		}
 
+		}
 	}
 
 	uint64 InitialLoadSize = 0;
 	if (GlobalIoStoreWriter)
 	{
+		IOSTORE_CPU_SCOPE(WriteScriptObjects);
 		FIoBuffer ScriptObjectsBuffer = PackageStoreOptimizer.CreateScriptObjectsBuffer();
 		InitialLoadSize = ScriptObjectsBuffer.DataSize();
 		FIoWriteOptions WriteOptions;
@@ -3635,37 +3669,40 @@ int32 CreateTarget(const FIoStoreArguments& Arguments, const FIoStoreWriterSetti
 	}
 
 	UE_LOG(LogIoStore, Display, TEXT("Serializing container(s)..."));
-
-	TFuture<void> FlushTask = Async(EAsyncExecution::Thread, [&IoStoreWriterContext]()
 	{
-		IoStoreWriterContext->Flush();
-	});
+		IOSTORE_CPU_SCOPE(Serializing);
 
-	while (!FlushTask.IsReady())
-	{
-		FlushTask.WaitFor(FTimespan::FromSeconds(2.0));
-		FIoStoreWriterContext::FProgress Progress = IoStoreWriterContext->GetProgress();
-		TStringBuilder<1024> ProgressStringBuilder;
-		if (Progress.SerializedChunksCount >= Progress.TotalChunksCount)
+		TFuture<void> FlushTask = Async(EAsyncExecution::Thread, [&IoStoreWriterContext]()
 		{
-			ProgressStringBuilder.Appendf(TEXT("Writing tocs..."));
-		}
-		else if (Progress.SerializedChunksCount)
+			IoStoreWriterContext->Flush();
+		});
+
+		while (!FlushTask.IsReady())
 		{
-			ProgressStringBuilder.Appendf(TEXT("Writing chunks (%llu/%llu)..."), Progress.SerializedChunksCount, Progress.TotalChunksCount);
-			if (Progress.CompressedChunksCount)
+			FlushTask.WaitFor(FTimespan::FromSeconds(2.0));
+			FIoStoreWriterContext::FProgress Progress = IoStoreWriterContext->GetProgress();
+			TStringBuilder<1024> ProgressStringBuilder;
+			if (Progress.SerializedChunksCount >= Progress.TotalChunksCount)
+			{
+				ProgressStringBuilder.Appendf(TEXT("Writing tocs..."));
+			}
+			else if (Progress.SerializedChunksCount)
 			{
-				ProgressStringBuilder.Appendf(TEXT(" [%llu compressed]"), Progress.CompressedChunksCount);
+				ProgressStringBuilder.Appendf(TEXT("Writing chunks (%llu/%llu)..."), Progress.SerializedChunksCount, Progress.TotalChunksCount);
+				if (Progress.CompressedChunksCount)
+				{
+					ProgressStringBuilder.Appendf(TEXT(" [%llu compressed]"), Progress.CompressedChunksCount);
+				}
+				if (Progress.ScheduledCompressionTasksCount)
+				{
+					ProgressStringBuilder.Appendf(TEXT(" [%llu compression tasks scheduled]"), Progress.ScheduledCompressionTasksCount);
+				}
+				UE_LOG(LogIoStore, Display, TEXT("%s"), *ProgressStringBuilder);
 			}
-			if (Progress.ScheduledCompressionTasksCount)
+			else
 			{
-				ProgressStringBuilder.Appendf(TEXT(" [%llu compression tasks scheduled]"), Progress.ScheduledCompressionTasksCount);
+				UE_LOG(LogIoStore, Display, TEXT("Hashing chunks (%llu/%llu) %llu from hash database..."), Progress.HashedChunksCount, Progress.TotalChunksCount, Progress.HashDbChunksCount);
 			}
-			UE_LOG(LogIoStore, Display, TEXT("%s"), *ProgressStringBuilder);
-		}
-		else
-		{
-			UE_LOG(LogIoStore, Display, TEXT("Hashing chunks (%llu/%llu) %llu from hash database..."), Progress.HashedChunksCount, Progress.TotalChunksCount, Progress.HashDbChunksCount);
 		}
 	}
 	if (GeneralIoWriterSettings.bCompressionEnableDDC)
@@ -6317,6 +6354,8 @@ static bool ParseSizeArgument(const TCHAR* CmdLine, const TCHAR* Argument, uint6
 
 static bool ParseOrderFileArguments(FIoStoreArguments& Arguments)
 {
+	IOSTORE_CPU_SCOPE(ParseOrderFileArguments);
+
 	uint64 OrderMapStartIndex = 0;
 	FString OrderFileStr;
 	if (FParse::Value(FCommandLine::Get(), TEXT("Order="), OrderFileStr, false))
@@ -6384,6 +6423,7 @@ static bool ParseOrderFileArguments(FIoStoreArguments& Arguments)
 
 bool ParseContainerGenerationArguments(FIoStoreArguments& Arguments, FIoStoreWriterSettings& WriterSettings)
 {
+	IOSTORE_CPU_SCOPE(ParseContainerGenerationArguments);
 	if (FParse::Param(FCommandLine::Get(), TEXT("sign")))
 	{
 		Arguments.bSign = true;
diff --git a/Engine/Source/Developer/IoStoreUtilities/Private/PackageStoreManifest.cpp b/Engine/Source/Developer/IoStoreUtilities/Private/PackageStoreManifest.cpp
--- a/Engine/Source/Developer/IoStoreUtilities/Private/PackageStoreManifest.cpp
+++ b/Engine/Source/Developer/IoStoreUtilities/Private/PackageStoreManifest.cpp
@@ -1,6 +1,7 @@
 // Copyright Epic Games, Inc. All Rights Reserved.
 
 #include "PackageStoreManifest.h"
+#include "ProfilingDebugging/CpuProfilerTrace.h"
 #include "Serialization/CompactBinaryContainerSerialization.h"
 #include "Serialization/CompactBinarySerialization.h"
 #include "Serialization/CompactBinaryWriter.h"
@@ -233,6 +234,7 @@ FIoStatus FPackageStoreManifest::Save(const TCHAR* Filename) const
 
 FIoStatus FPackageStoreManifest::Load(const TCHAR* Filename)
 {
+	TRACE_CPUPROFILER_EVENT_SCOPE(FPackageStoreManifest::Load);
 	LLM_SCOPE_BYTAG(Cooker_PackageStoreManifest);
 	FScopeLock Lock(&CriticalSection);
 	PackageInfoByNameMap.Empty();
diff --git a/Engine/Source/Developer/IoStoreUtilities/Private/PackageStoreOptimizer.cpp b/Engine/Source/Developer/IoStoreUtilities/Private/PackageStoreOptimizer.cpp
--- a/Engine/Source/Developer/IoStoreUtilities/Private/PackageStoreOptimizer.cpp
+++ b/Engine/Source/Developer/IoStoreUtilities/Private/PackageStoreOptimizer.cpp
@@ -807,6 +807,7 @@ TArray<FPackageStorePackage::FExportBundleGraphNode*> FPackageStoreOptimizer::So
 
 void FPackageStoreOptimizer::OptimizeExportBundles(const TMap<FPackageId, FPackageStorePackage*>& PackagesMap)
 {
+	TRACE_CPUPROFILER_EVENT_SCOPE(OptimizeExportBundles);
 	TArray<FPackageStorePackage*> Packages = SortPackagesInLoadOrder(PackagesMap);
 	for (FPackageStorePackage* Package : Packages)
 	{
@@ -1659,6 +1660,7 @@ void FPackageStoreOptimizer::FindScriptObjectsRecursive(FPackageObjectIndex Oute
 
 void FPackageStoreOptimizer::FindScriptObjects()
 {
+	TRACE_CPUPROFILER_EVENT_SCOPE(FindScriptObjects);
 	TArray<UPackage*> ScriptPackages;
 	FindAllRuntimeScriptPackages(ScriptPackages);
 
@@ -1729,6 +1731,7 @@ FIoBuffer FPackageStoreOptimizer::CreateScriptObjectsBuffer() const
 
 void FPackageStoreOptimizer::LoadScriptObjectsBuffer(const FIoBuffer& ScriptObjectsBuffer)
 {
+	TRACE_CPUPROFILER_EVENT_SCOPE(LoadScriptObjectsBuffer);
 	FLargeMemoryReader ScriptObjectsArchive(ScriptObjectsBuffer.Data(), ScriptObjectsBuffer.DataSize());
 	TArray<FDisplayNameEntryId> NameMap = LoadNameBatch(ScriptObjectsArchive);
 	int32 NumScriptObjects;
@@ -1778,6 +1781,7 @@ FPackageStoreEntryResource FPackageStoreOptimizer::CreatePackageStoreEntry(const
 
 FIoContainerHeader FPackageStoreOptimizer::CreateContainerHeaderInternal(const FIoContainerId& ContainerId, TArrayView<const FPackageStoreEntryResource> PackageStoreEntries, bool bIsOptional) const
 {
+	TRACE_CPUPROFILER_EVENT_SCOPE(CreateContainerHeaderInternal);
 	FIoContainerHeader Header;
 	Header.ContainerId = ContainerId;
 	
diff --git a/Engine/Source/Runtime/Core/Private/Async/TaskGraph.cpp b/Engine/Source/Runtime/Core/Private/Async/TaskGraph.cpp
--- a/Engine/Source/Runtime/Core/Private/Async/TaskGraph.cpp
+++ b/Engine/Source/Runtime/Core/Private/Async/TaskGraph.cpp
@@ -2310,6 +2310,7 @@ private:
 
 	void StartReserveWorkers()
 	{
+		TRACE_CPUPROFILER_EVENT_SCOPE(StartReserveWorkers);
 		if (bReserveWorkersEnabled)
 		{
 			return; // once enabled, reserve workers can't be disabled
diff --git a/Engine/Source/Runtime/Core/Private/IO/IoStore.cpp b/Engine/Source/Runtime/Core/Private/IO/IoStore.cpp
--- a/Engine/Source/Runtime/Core/Private/IO/IoStore.cpp
+++ b/Engine/Source/Runtime/Core/Private/IO/IoStore.cpp
@@ -317,6 +317,7 @@ public:
 
 	UE_NODISCARD FIoStatus Initialize(const FIoStoreWriterSettings& InWriterSettings)
 	{
+		TRACE_CPUPROFILER_EVENT_SCOPE(FIoStoreWriterContext::Initialize);
 		WriterSettings = InWriterSettings;
 		CompressionBufferAvailableEvent = FPlatformProcess::GetSynchEventFromPool(false);
 
@@ -1850,6 +1851,7 @@ TSharedPtr<IIoStoreWriter> FIoStoreWriterContextImpl::CreateContainer(const TCHA
 
 void FIoStoreWriterContextImpl::Flush()
 {
+	TRACE_CPUPROFILER_EVENT_SCOPE(FIoStoreWriterContext::Flush);
 	TArray<FIoStoreWriteQueueEntry*> AllEntries;
 	for (TSharedPtr<FIoStoreWriter> IoStoreWriter : IoStoreWriters)
 	{
@@ -2076,6 +2078,7 @@ public:
 
 	UE_NODISCARD FIoStatus Initialize(const TCHAR* InContainerPath, const TMap<FGuid, FAES::FAESKey>& InDecryptionKeys)
 	{
+		TRACE_CPUPROFILER_EVENT_SCOPE(FIoStoreReader::Initialize);
 		ContainerPath = InContainerPath;
 
 		TStringBuilder<256> TocFilePath;
diff --git a/Engine/Source/Runtime/Core/Private/Misc/ConfigCacheIni.cpp b/Engine/Source/Runtime/Core/Private/Misc/ConfigCacheIni.cpp
--- a/Engine/Source/Runtime/Core/Private/Misc/ConfigCacheIni.cpp
+++ b/Engine/Source/Runtime/Core/Private/Misc/ConfigCacheIni.cpp
@@ -21,6 +21,7 @@
 #include "Misc/DataDrivenPlatformInfoRegistry.h"
 #include "Misc/StringBuilder.h"
 #include "Misc/Paths.h"
+#include "ProfilingDebugging/CpuProfilerTrace.h"
 #include "Serialization/MemoryReader.h"
 #include "Serialization/MemoryWriter.h"
 #include "Serialization/LargeMemoryReader.h"
@@ -3694,10 +3695,13 @@ bool FConfigCacheIni::CreateGConfigFromSaved(const TCHAR* Filename)
 	// now let the delegates pull their data out, after GConfig is set up
 //	FCoreDelegates::TSAccessExtraBinaryConfigData().Broadcast(ExtraData);
 
-	FCoreDelegates::TSConfigReadyForUse().Broadcast();
+	{
+		TRACE_CPUPROFILER_EVENT_SCOPE(ConfigReadyForUseBroadcast);
+		FCoreDelegates::TSConfigReadyForUse().Broadcast();
 PRAGMA_DISABLE_DEPRECATION_WARNINGS
-	FCoreDelegates::ConfigReadyForUse.Broadcast();
+		FCoreDelegates::ConfigReadyForUse.Broadcast();
 PRAGMA_ENABLE_DEPRECATION_WARNINGS
+	}
 
 	FMemory::Free(PreloadedData);
 	return true;
@@ -3803,6 +3807,7 @@ void FConfigCacheIni::InitializeConfigSystem()
 	FString IniBootstrapFilename;
 	if (FParse::Value( FCommandLine::Get(), TEXT("IniBootstrap="), IniBootstrapFilename))
 	{
+		TRACE_CPUPROFILER_EVENT_SCOPE(IniBootstrap);
 		TArray<uint8> FileContent;
 		if (FFileHelper::LoadFileToArray(FileContent, *IniBootstrapFilename, FILEREAD_Silent))
 		{
@@ -3810,6 +3815,7 @@ void FConfigCacheIni::InitializeConfigSystem()
 			GConfig = new FConfigCacheIni(EConfigCacheType::Temporary);
 			GConfig->SerializeStateForBootstrap_Impl(MemoryReader);
 			GConfig->bIsReadyForUse = true;
+			TRACE_CPUPROFILER_EVENT_SCOPE(ConfigReadyForUseBroadcast);
 			FCoreDelegates::TSConfigReadyForUse().Broadcast();
 PRAGMA_DISABLE_DEPRECATION_WARNINGS
 			FCoreDelegates::ConfigReadyForUse.Broadcast();
@@ -3863,6 +3869,7 @@ PRAGMA_ENABLE_DEPRECATION_WARNINGS
 	// now we can make use of GConfig
 	GConfig->bIsReadyForUse = true;
 
+	TRACE_CPUPROFILER_EVENT_SCOPE(ConfigReadyForUseBroadcast);
 	FCoreDelegates::TSConfigReadyForUse().Broadcast();
 PRAGMA_DISABLE_DEPRECATION_WARNINGS
 	FCoreDelegates::ConfigReadyForUse.Broadcast();
diff --git a/Engine/Source/Runtime/Core/Private/Misc/ConfigContext.cpp b/Engine/Source/Runtime/Core/Private/Misc/ConfigContext.cpp
--- a/Engine/Source/Runtime/Core/Private/Misc/ConfigContext.cpp
+++ b/Engine/Source/Runtime/Core/Private/Misc/ConfigContext.cpp
@@ -10,6 +10,7 @@
 #include "Misc/RemoteConfigIni.h"
 #include "Misc/Paths.h"
 #include "Misc/DataDrivenPlatformInfoRegistry.h"
+#include "ProfilingDebugging/CpuProfilerTrace.h"
 
 namespace
 {
@@ -122,6 +123,7 @@ const FConfigContext::FPerPlatformDirs& FConfigContext::GetPerPlatformDirs(const
 
 bool FConfigContext::Load(const TCHAR* InBaseIniName, FString& OutFinalFilename)
 {
+	TRACE_CPUPROFILER_EVENT_SCOPE(FConfigContext::Load);
 	// for single file loads, just return early of the file doesn't exist
 	const bool bBaseIniNameIsFullInIFilePath = FString(InBaseIniName).EndsWith(TEXT(".ini"));
 	if (!bIsHierarchicalConfig && bBaseIniNameIsFullInIFilePath && !DoesConfigFileExistWrapper(InBaseIniName, IniCacheSet))
@@ -597,6 +599,7 @@ void FConfigContext::AddStaticLayersToHierarchy()
  **/
 static bool LoadIniFileHierarchy(const FConfigFileHierarchy& HierarchyToLoad, FConfigFile& ConfigFile, bool bUseCache, const TSet<FString>* IniCacheSet)
 {
+	TRACE_CPUPROFILER_EVENT_SCOPE(LoadIniFileHierarchy);
 	// Traverse ini list back to front, merging along the way.
 	for (const TPair<int32, FString>& HierarchyIt : HierarchyToLoad)
 	{
