IoStore:
* Fixed inderministic sort order for shaders of the same size, leading to different shader group chunks being generated
* Fixed memory leak when generating shader groups
* Fixed indeterministic classification of shaders appearing in two shader maps with the same hash

diff --git a/Engine/Source/Developer/IoStoreUtilities/Private/IoStoreUtilities.cpp b/Engine/Source/Developer/IoStoreUtilities/Private/IoStoreUtilities.cpp
--- a/Engine/Source/Developer/IoStoreUtilities/Private/IoStoreUtilities.cpp
+++ b/Engine/Source/Developer/IoStoreUtilities/Private/IoStoreUtilities.cpp
@@ -569,8 +569,7 @@ struct FShaderInfo
 	FIoChunkId ChunkId;
 	FIoBuffer CodeIoBuffer;
 	// the smaller the order, the more likely this will be loaded
-	uint32 LoadOrderFactor;
-	uint64 DiskLayoutOrder = MAX_uint64;
+	uint32 LoadOrderFactor = MAX_uint32;
 	TSet<struct FLegacyCookedPackage*> ReferencedByPackages;
 	TMap<FContainerTargetSpec*, EShaderType> TypeInContainer;
 
@@ -578,18 +577,12 @@ struct FShaderInfo
 	// across builds.
 	static bool Sort(const FShaderInfo* A, const FShaderInfo* B)
 	{
-		if (A->DiskLayoutOrder == B->DiskLayoutOrder)
+		if (A->LoadOrderFactor == B->LoadOrderFactor)
 		{
-			if (A->LoadOrderFactor == B->LoadOrderFactor)
-			{
-				// Shader chunk IDs are the hash of the shader so this is consistent across builds.
-				uint64 AHash = *(uint64*)A->ChunkId.GetData();
-				uint64 BHash = *(uint64*)B->ChunkId.GetData();
-				return AHash < BHash;
-			}
-			return A->LoadOrderFactor < B->LoadOrderFactor;
+			// Shader chunk IDs are the hash of the shader so this is consistent across builds.
+			return FMemory::Memcmp(A->ChunkId.GetData(), B->ChunkId.GetData(), A->ChunkId.GetSize()) < 0;
 		}
-		return A->DiskLayoutOrder < B->DiskLayoutOrder;
+		return A->LoadOrderFactor < B->LoadOrderFactor;
 	}
 };
 
@@ -1846,6 +1839,7 @@ bool ConvertToIoStoreShaderLibrary(
 				Group.CompressedSize = Group.UncompressedSize;
 			}
 			FMemory::Free(UncompressedGroupMemory);
+			FMemory::Free(CompressedShaderGroupMemory);
 		},
 		EParallelForFlags::Unbalanced
 	);
@@ -1932,7 +1926,6 @@ void ProcessShaderLibraries(const FIoStoreArguments& Arguments, TArray<FContaine
 		IOSTORE_CPU_SCOPE(ConvertShaderLibraries);
 		for (FContainerTargetSpec* ContainerTarget : ContainerTargets)
 		{
-			TArray<FShaderInfo> Shaders;
 			for (FContainerTargetFile& TargetFile : ContainerTarget->TargetFiles)
 			{
 				if (TargetFile.ChunkType == EContainerChunkType::ShaderCodeLibrary)
@@ -2002,7 +1995,8 @@ void ProcessShaderLibraries(const FIoStoreArguments& Arguments, TArray<FContaine
 
 					for (TTuple<FSHAHash, TArray<FIoChunkId>>& ShaderMap : ShaderMaps)
 					{
-						ShaderChunkIdsByShaderMapHash.Add(ShaderMap.Key, MoveTemp(ShaderMap.Value));
+						TArray<FIoChunkId>& ShaderMapChunkIds = ShaderChunkIdsByShaderMapHash.FindOrAdd(ShaderMap.Key);
+						ShaderMapChunkIds.Append(MoveTemp(ShaderMap.Value));
 					}
 				} // end if containerchunktype shadercodelibrary
 			} // end foreach targetfile
diff --git a/Engine/Source/Runtime/RenderCore/Private/ShaderCodeArchive.cpp b/Engine/Source/Runtime/RenderCore/Private/ShaderCodeArchive.cpp
--- a/Engine/Source/Runtime/RenderCore/Private/ShaderCodeArchive.cpp
+++ b/Engine/Source/Runtime/RenderCore/Private/ShaderCodeArchive.cpp
@@ -1419,7 +1419,21 @@ void FIoStoreShaderCodeArchive::CreateIoStoreShaderCodeArchiveHeader(const FName
 		ShaderIndicesInGroup.Sort(
 			[&SerializedShaders](const int32 ShaderIndexA, const int32 ShaderIndexB)
 			{
-				return SerializedShaders.ShaderEntries[ShaderIndexA].UncompressedSize < SerializedShaders.ShaderEntries[ShaderIndexB].UncompressedSize;
+				const FShaderCodeEntry& ShaderEntryA = SerializedShaders.ShaderEntries[ShaderIndexA];
+				const FShaderCodeEntry& ShaderEntryB = SerializedShaders.ShaderEntries[ShaderIndexB];
+				if (ShaderEntryA.UncompressedSize != ShaderEntryB.UncompressedSize)
+				{
+					return ShaderEntryA.UncompressedSize < ShaderEntryB.UncompressedSize;
+				}
+				if (ShaderEntryA.Size != ShaderEntryB.Size)
+				{
+					return ShaderEntryA.Size < ShaderEntryB.Size;
+				}
+				if (ShaderEntryA.Frequency != ShaderEntryB.Frequency)
+				{
+					return ShaderEntryA.Frequency < ShaderEntryB.Frequency;
+				}
+				return ShaderEntryA.Offset < ShaderEntryB.Offset;
 			}
 		);
 
@@ -1490,7 +1504,21 @@ void FIoStoreShaderCodeArchive::CreateIoStoreShaderCodeArchiveHeader(const FName
 			CurrentShaderGroup.Sort(
 				[&SerializedShaders](const int32 ShaderIndexA, const int32 ShaderIndexB)
 				{
-					return SerializedShaders.ShaderEntries[ShaderIndexA].UncompressedSize > SerializedShaders.ShaderEntries[ShaderIndexB].UncompressedSize;
+					const FShaderCodeEntry& ShaderEntryA = SerializedShaders.ShaderEntries[ShaderIndexA];
+					const FShaderCodeEntry& ShaderEntryB = SerializedShaders.ShaderEntries[ShaderIndexB];
+					if (ShaderEntryA.UncompressedSize != ShaderEntryB.UncompressedSize)
+					{
+						return ShaderEntryA.UncompressedSize > ShaderEntryB.UncompressedSize;
+					}
+					if (ShaderEntryA.Size != ShaderEntryB.Size)
+					{
+						return ShaderEntryA.Size > ShaderEntryB.Size;
+					}
+					if (ShaderEntryA.Frequency != ShaderEntryB.Frequency)
+					{
+						return ShaderEntryA.Frequency > ShaderEntryB.Frequency;
+					}
+					return ShaderEntryA.Offset > ShaderEntryB.Offset;
 				}
 			);
 
