Removal of the r.Shaders.CompressionFormat CVar, and shader group compression during UnrealPak now inherits the packaging compresison level.

Due to the fact that UnrealPak doesn't load INIs, the cvar was always its default value. This, combined with the fact that shaders don't store what codec they are compressed with means that setting this cvar to anything other than default (oodle) causes a crash in UnrealPak as DecompressShader will try to decompress with Oodle independent of the actual codec. Rather than plumb this information down through ProjectParams, I felt it prudent to just always use Oodle.

Additionally, in order to facilitate faster iteration times for development build farm jobs, the shader compression level used during UnrealPak will inherit from the overall package compression level set via Project Settings (PackageCompressionLevel_*).

This tacitly avoids patch issues as the default shader compression level happens to be the default package compression level, except now it inherits any faster levels for iteration time.

Note that we still use Mermaid as the codec as it's faster for decompressing than Kraken, which we care about for shaders.

diff --git a/Engine/Source/Developer/IoStoreUtilities/Private/IoStoreUtilities.cpp b/Engine/Source/Developer/IoStoreUtilities/Private/IoStoreUtilities.cpp
--- a/Engine/Source/Developer/IoStoreUtilities/Private/IoStoreUtilities.cpp
+++ b/Engine/Source/Developer/IoStoreUtilities/Private/IoStoreUtilities.cpp
@@ -869,6 +869,9 @@ struct FIoStoreArguments
 	bool bFileRegions = false;
 	EAssetRegistryWritebackMethod WriteBackMetadataToAssetRegistry = EAssetRegistryWritebackMethod::Disabled;
 
+	FOodleDataCompression::ECompressor ShaderOodleCompressor = FOodleDataCompression::ECompressor::Mermaid;
+	FOodleDataCompression::ECompressionLevel ShaderOodleLevel = FOodleDataCompression::ECompressionLevel::Normal;
+
 	bool IsDLC() const
 	{
 		return DLCPluginPath.Len() > 0;
@@ -1691,8 +1694,10 @@ bool LoadShaderAssetInfo(const FString& Filename, TMap<FSHAHash, TSet<FName>>& O
 	return true;
 }
 
-bool ConvertToIoStoreShaderLibrary(
+static bool ConvertToIoStoreShaderLibrary(
 	const TCHAR* FileName,
+	FOodleDataCompression::ECompressor InShaderOodleCompressor,
+	FOodleDataCompression::ECompressionLevel InShaderOodleLevel,
 	TTuple<FIoChunkId, FIoBuffer>& OutLibraryIoChunk,
 	TArray<TTuple<FIoChunkId, FIoBuffer, uint32>>& OutCodeIoChunks,
 	TArray<TTuple<FSHAHash, TArray<FIoChunkId>>>& OutShaderMaps,
@@ -1751,7 +1756,7 @@ bool ConvertToIoStoreShaderLibrary(
 	OutCodeIoChunks.SetNum(GroupCount);
 	FCriticalSection DiskArchiveAccess;
 	ParallelFor(OutCodeIoChunks.Num(),
-		[&OutCodeIoChunks, &IoStoreLibraryHeader, &DiskArchiveAccess, &LibraryAr, &SerializedShaders, &OffsetToShaderCode](int32 GroupIndex)
+		[&OutCodeIoChunks, &IoStoreLibraryHeader, &DiskArchiveAccess, &LibraryAr, &SerializedShaders, &OffsetToShaderCode, InShaderOodleCompressor, InShaderOodleLevel](int32 GroupIndex)
 		{
 			FIoStoreShaderGroupEntry& Group = IoStoreLibraryHeader.ShaderGroupEntries[GroupIndex];
 			uint8* UncompressedGroupMemory = reinterpret_cast<uint8*>(FMemory::Malloc(Group.UncompressedSize));
@@ -1794,7 +1799,7 @@ bool ConvertToIoStoreShaderLibrary(
 						}
 
 						// This function will crash if decompression fails.
-						ShaderCodeArchive::DecompressShader(ShaderStart, IndividuallyCompressedShader.UncompressedSize, CompressedShaderMemory, IndividuallyCompressedShader.Size);
+						ShaderCodeArchive::DecompressShaderWithOodle(ShaderStart, IndividuallyCompressedShader.UncompressedSize, CompressedShaderMemory, IndividuallyCompressedShader.Size);
 
 						FMemory::Free(CompressedShaderMemory);
 					}
@@ -1809,11 +1814,12 @@ bool ConvertToIoStoreShaderLibrary(
 
 			// now compress the whole group
 			int64 CompressedGroupSize = 0;
-			ShaderCodeArchive::CompressShaderUsingCurrentSettings(nullptr, CompressedGroupSize, UncompressedGroupMemory, Group.UncompressedSize);
+			uint8* CompressedShaderGroupMemory = nullptr;
+			ShaderCodeArchive::CompressShaderWithOodle(nullptr, CompressedGroupSize, UncompressedGroupMemory, Group.UncompressedSize, InShaderOodleCompressor, InShaderOodleLevel);
 			checkf(CompressedGroupSize > 0, TEXT("CompressedGroupSize estimate seems wrong (%lld)"), CompressedGroupSize);
 
-			uint8* CompressedShaderGroupMemory = reinterpret_cast<uint8*>(FMemory::Malloc(CompressedGroupSize));
-			bool bCompressed = ShaderCodeArchive::CompressShaderUsingCurrentSettings(CompressedShaderGroupMemory, CompressedGroupSize, UncompressedGroupMemory, Group.UncompressedSize);
+			CompressedShaderGroupMemory = reinterpret_cast<uint8*>(FMemory::Malloc(CompressedGroupSize));
+			bool bCompressed = ShaderCodeArchive::CompressShaderWithOodle(CompressedShaderGroupMemory, CompressedGroupSize, UncompressedGroupMemory, Group.UncompressedSize, InShaderOodleCompressor, InShaderOodleLevel);
 			checkf(bCompressed, TEXT("We could not compress the shader group after providing an estimated memory."));
 
 			TTuple<FIoChunkId, FIoBuffer, uint32>& OutCodeIoChunk = OutCodeIoChunks[GroupIndex];
@@ -1913,7 +1919,7 @@ bool ConvertToIoStoreShaderLibrary(
 	return true;
 }
 
-void ProcessShaderLibraries(const FIoStoreArguments& Arguments, TArray<FContainerTargetSpec*>& ContainerTargets, TArray<FShaderInfo*> OutShaders)
+static void ProcessShaderLibraries(const FIoStoreArguments& Arguments, TArray<FContainerTargetSpec*>& ContainerTargets, TArray<FShaderInfo*> OutShaders)
 {
 	IOSTORE_CPU_SCOPE(ProcessShaderLibraries);
 
@@ -1936,7 +1942,7 @@ void ProcessShaderLibraries(const FIoStoreArguments& Arguments, TArray<FContaine
 					TTuple<FIoChunkId, FIoBuffer> LibraryChunk;
 					TArray<TTuple<FIoChunkId, FIoBuffer, uint32>> CodeChunks;
 					TArray<TTuple<FSHAHash, TSet<FName>>> ShaderMapAssetAssociations;
-					if (!ConvertToIoStoreShaderLibrary(*TargetFile.NormalizedSourcePath, LibraryChunk, CodeChunks, ShaderMaps, ShaderMapAssetAssociations))
+				    if (!ConvertToIoStoreShaderLibrary(*TargetFile.NormalizedSourcePath, Arguments.ShaderOodleCompressor, Arguments.ShaderOodleLevel, LibraryChunk, CodeChunks, ShaderMaps, ShaderMapAssetAssociations))
 					{
 						UE_LOG(LogIoStore, Warning, TEXT("Failed converting shader library '%s'"), *TargetFile.NormalizedSourcePath);
 						continue;
@@ -3455,7 +3461,7 @@ int32 CreateTarget(const FIoStoreArguments& Arguments, const FIoStoreWriterSetti
 		ParsePackageAssetsFromFiles(Packages, PackageStoreOptimizer);
 	}
 
-	UE_LOG(LogIoStore, Display, TEXT("Processing shader libraries..."));
+	UE_LOG(LogIoStore, Display, TEXT("Processing shader libraries, compressing with Oodle %s, level %d (%s)"), FOodleDataCompression::ECompressorToString(Arguments.ShaderOodleCompressor), (int32)Arguments.ShaderOodleLevel, FOodleDataCompression::ECompressionLevelToString(Arguments.ShaderOodleLevel));
 	TArray<FShaderInfo*> Shaders;
 	ProcessShaderLibraries(Arguments, ContainerTargets, Shaders);
 
@@ -6889,6 +6895,22 @@ int32 CreateIoStoreContainerFiles(const TCHAR* CmdLine)
 		return -1;
 	}
 
+	//
+	// -compresslevel is technically consumed by OodleDataCompressionFormat, however the setting that it represents (PackageCompressionLevel_*)
+	// is the intention of package compression, and so should also determine the compression we use for shaders. For Shaders we want fast decompression,
+	// so we always used Mermaid. Note that this relies on compresslevel being passed even when the containers aren't compressed.
+	//
+	FString ShaderOodleLevel;
+	FParse::Value(FCommandLine::Get(), TEXT("compresslevel="), ShaderOodleLevel);
+	if (ShaderOodleLevel.Len())
+	{
+		if (FOodleDataCompression::ECompressionLevelFromString(*ShaderOodleLevel, Arguments.ShaderOodleLevel))
+		{
+			UE_LOG(LogIoStore, Display, TEXT("Selected Oodle level %d (%s) from command line for shaders"), (int)Arguments.ShaderOodleLevel, FOodleDataCompression::ECompressionLevelToString(Arguments.ShaderOodleLevel));
+		}
+		
+	}
+
 	// Whether or not to write compressed asset sizes back to the asset registry.
 
 	FString WriteBackMetadataToAssetRegistry;
diff --git a/Engine/Source/Programs/AutomationTool/Scripts/CopyBuildToStagingDirectory.Automation.cs b/Engine/Source/Programs/AutomationTool/Scripts/CopyBuildToStagingDirectory.Automation.cs
--- a/Engine/Source/Programs/AutomationTool/Scripts/CopyBuildToStagingDirectory.Automation.cs
+++ b/Engine/Source/Programs/AutomationTool/Scripts/CopyBuildToStagingDirectory.Automation.cs
@@ -2680,6 +2680,28 @@ namespace AutomationScripts
 				CompressionFormats = "";
 			}
 
+			// Always pass -compresslevel because we use it to inform the compression level of
+			// shaders, which are always compressed even if the containers aren't compressed.
+			int CompressionLevel = 0;
+			// GetInt32 fills out with = 0 if not found
+			if (Params.Distribution)
+			{
+				PlatformGameConfig.GetInt32("/Script/UnrealEd.ProjectPackagingSettings", "PackageCompressionLevel_Distribution", out CompressionLevel);
+			}
+			else if (SC.StageTargetConfigurations.Any(Config => Config == UnrealTargetConfiguration.Test || Config == UnrealTargetConfiguration.Shipping))
+			{
+				PlatformGameConfig.GetInt32("/Script/UnrealEd.ProjectPackagingSettings", "PackageCompressionLevel_TestShipping", out CompressionLevel);
+			}
+			else
+			{
+				PlatformGameConfig.GetInt32("/Script/UnrealEd.ProjectPackagingSettings", "PackageCompressionLevel_DebugDevelopment", out CompressionLevel);
+			}
+
+			if (CompressionLevel != 0)
+			{
+				CompressionFormats += " -compresslevel=" + CompressionLevel;
+			}
+
 			string AdditionalCompressionOptionsOnCommandLine = "";
 			if (Params.Compressed)
 			{
@@ -2688,38 +2710,13 @@ namespace AutomationScripts
 				// there's not a great way to communicate random strings down into the plugins during plugin init time)
 				PlatformGameConfig.GetString("/Script/UnrealEd.ProjectPackagingSettings", "PackageAdditionalCompressionOptions", out AdditionalCompressionOptionsOnCommandLine);
 
-				//LogInformation("PackageAdditionalCompressionOptions = {0}", AdditionalCompressionOptionsOnCommandLine);
-
 				string CompressionMethod;
 				PlatformGameConfig.GetString("/Script/UnrealEd.ProjectPackagingSettings", "PackageCompressionMethod", out CompressionMethod);
 				if (!string.IsNullOrWhiteSpace(CompressionMethod))
 				{
-					//LogInformation("CompressionMethod = {0}", CompressionMethod);
 					CompressionFormats += " -compressmethod=" + CompressionMethod;
 				}
 
-				// get the encoder compression effort level based on the cook type
-				int CompressionLevel = 0;
-				// GetInt32 fills out with = 0 if not found
-				if (Params.Distribution)
-				{
-					PlatformGameConfig.GetInt32("/Script/UnrealEd.ProjectPackagingSettings", "PackageCompressionLevel_Distribution", out CompressionLevel);
-				}
-				else if (SC.StageTargetConfigurations.Any(Config => Config == UnrealTargetConfiguration.Test || Config == UnrealTargetConfiguration.Shipping))
-				{
-					PlatformGameConfig.GetInt32("/Script/UnrealEd.ProjectPackagingSettings", "PackageCompressionLevel_TestShipping", out CompressionLevel);
-				}
-				else
-				{
-					PlatformGameConfig.GetInt32("/Script/UnrealEd.ProjectPackagingSettings", "PackageCompressionLevel_DebugDevelopment", out CompressionLevel);
-				}
-
-				if (CompressionLevel != 0)
-				{
-					//LogInformation(" CompressionLevel = {0}", CompressionLevel);
-					CompressionFormats += " -compresslevel=" + CompressionLevel;
-				}
-
 
 				// unlike the other config, the DLL should be fetched from the build platform not the target platform
 				// it's the DLL to use during packaging, not during runtime, so it needs to be for the platform the packager runs on, not the platform you're packaging for.
diff --git a/Engine/Source/Runtime/Core/Private/Compression/OodleDataCompression.cpp b/Engine/Source/Runtime/Core/Private/Compression/OodleDataCompression.cpp
--- a/Engine/Source/Runtime/Core/Private/Compression/OodleDataCompression.cpp
+++ b/Engine/Source/Runtime/Core/Private/Compression/OodleDataCompression.cpp
@@ -83,6 +83,31 @@ static struct { ECompressionLevel Level; const TCHAR* Name; } CompressionLevelNa
 	{ECompressionLevel::Optimal4, TEXT("Optimal4")}
 };
 
+CORE_API bool ECompressionLevelFromString(const TCHAR* InName, ECompressionLevel& OutLevel)
+{
+	for (SIZE_T i = 0; i < sizeof(CompressionLevelNameMap) / sizeof(CompressionLevelNameMap[0]); i++)
+	{
+		if (FCString::Stricmp(CompressionLevelNameMap[i].Name, InName) == 0)
+		{
+			OutLevel = CompressionLevelNameMap[i].Level;
+			return true;
+		}
+	}
+	// Since 0 is a valid compression level, we can't just atoi it or we'd always succeed.
+	if ((InName[0] == '-' && FChar::IsDigit(InName[1]))
+		|| FChar::IsDigit(InName[0]) )
+	{
+		int32 PossibleCompressionLevel = FCString::Atoi(InName);
+		if (PossibleCompressionLevel >= OodleLZ_CompressionLevel_Min &&
+			PossibleCompressionLevel <= OodleLZ_CompressionLevel_Max)
+		{
+			OutLevel = (ECompressionLevel)PossibleCompressionLevel;
+			return true;
+		}
+	}
+	return false;
+}
+
 CORE_API bool ECompressionLevelToString(ECompressionLevel InLevel, const TCHAR** OutName)
 {
 	for (SIZE_T i = 0; i < sizeof(CompressionLevelNameMap) / sizeof(CompressionLevelNameMap[0]); i++)
diff --git a/Engine/Source/Runtime/Core/Public/Compression/OodleDataCompression.h b/Engine/Source/Runtime/Core/Public/Compression/OodleDataCompression.h
--- a/Engine/Source/Runtime/Core/Public/Compression/OodleDataCompression.h
+++ b/Engine/Source/Runtime/Core/Public/Compression/OodleDataCompression.h
@@ -109,6 +109,7 @@ enum class ECompressionLevel : int8
 bool CORE_API ECompressionLevelFromValue(int8 InValue, ECompressionLevel& OutLevel);
 bool CORE_API ECompressionLevelToString(ECompressionLevel InLevel, const TCHAR** OutName);
 CORE_API const TCHAR * ECompressionLevelToString(ECompressionLevel InLevel);
+CORE_API bool ECompressionLevelFromString(const TCHAR* InName, ECompressionLevel& OutLevel);
 
 enum class ECompressionCommonUsage : uint8
 {
diff --git a/Engine/Source/Runtime/Engine/Classes/Engine/RendererSettings.h b/Engine/Source/Runtime/Engine/Classes/Engine/RendererSettings.h
--- a/Engine/Source/Runtime/Engine/Classes/Engine/RendererSettings.h
+++ b/Engine/Source/Runtime/Engine/Classes/Engine/RendererSettings.h
@@ -329,11 +329,8 @@ class ENGINE_API URendererSettings : public UDeveloperSettings
 		ToolTip = "When running in game mode, whether to keep shaders for all quality levels in memory or only those needed for the current quality level.\nUnchecked: Keep all quality levels in memory allowing a runtime quality level change. (default)\nChecked: Discard unused quality levels when loading content for the game, saving some memory."))
 	uint32 bDiscardUnusedQualityLevels : 1;
 
-	UPROPERTY(config, EditAnywhere, Category = Materials, meta = (
-		ConsoleVariable = "r.Shaders.CompressionFormat", DisplayName = "Shader Compression Format",
-		ToolTip = "Select how the shaders are compressed for storage",
-		ConfigRestartRequired = true))
-	TEnumAsByte<EShaderCompressionFormat::Type> ShaderCompressionFormat;	// requires restart as we may have jobs in flight that already used the previous one. Also it's also used for decompression the currently compressed shaders
+	UPROPERTY()
+	TEnumAsByte<EShaderCompressionFormat::Type> ShaderCompressionFormat_DEPRECATED;	// we now force oodle
 
 	UPROPERTY(config, EditAnywhere, Category=Culling, meta=(
 		ConsoleVariable="r.AllowOcclusionQueries",DisplayName="Occlusion Culling",
diff --git a/Engine/Source/Runtime/RenderCore/Private/ShaderCodeArchive.cpp b/Engine/Source/Runtime/RenderCore/Private/ShaderCodeArchive.cpp
--- a/Engine/Source/Runtime/RenderCore/Private/ShaderCodeArchive.cpp
+++ b/Engine/Source/Runtime/RenderCore/Private/ShaderCodeArchive.cpp
@@ -348,47 +348,30 @@ void FShaderUsageVisualizer::SaveShaderUsageBitmap(const FString& Name, EShaderP
 }
 #endif
 
-void ShaderCodeArchive::DecompressShader(uint8* OutDecompressedShader, int64 UncompressedSize, const uint8* CompressedShaderCode, int64 CompressedSize)
+void ShaderCodeArchive::DecompressShaderWithOodle(uint8* OutDecompressedShader, int64 UncompressedSize, const uint8* CompressedShaderCode, int64 CompressedSize)
 {
-	bool bSucceed = FCompression::UncompressMemory(GetShaderCompressionFormat(), OutDecompressedShader, UncompressedSize, CompressedShaderCode, CompressedSize);
+	// Iostore always compresses with Oodle.
+	bool bSucceed = FCompression::UncompressMemory(NAME_Oodle, OutDecompressedShader, UncompressedSize, CompressedShaderCode, CompressedSize);
 	if (!bSucceed)
 	{
-		UE_LOG(LogShaderLibrary, Fatal, TEXT("ShaderCodeArchive::DecompressShader(): Could not decompress shader (GetShaderCompressionFormat=%s)"), *GetShaderCompressionFormat().ToString());
+		UE_LOG(LogShaderLibrary, Fatal, TEXT("ShaderCodeArchive::DecompressShader(): Could not decompress shader with Oodle"));
 	}
 }
 
-bool ShaderCodeArchive::CompressShaderUsingCurrentSettings(uint8* OutCompressedShader, int64& OutCompressedSize, const uint8* UncompressedShaderCode, int64 UncompressedSize)
+bool ShaderCodeArchive::CompressShaderWithOodle(uint8* OutCompressedShader, int64& OutCompressedSize, const uint8* InUncompressedShaderCode, int64 InUncompressedSize, FOodleDataCompression::ECompressor InOodleCompressor, FOodleDataCompression::ECompressionLevel InOodleLevel)
 {
-	// see FShaderCode::Compress - while this doesn't have to match exactly, it should match at least the parameters used there
-	const FName ShaderCompressionFormat = GetShaderCompressionFormat();
-
-	bool bCompressed = false;
-	if (ShaderCompressionFormat != NAME_Oodle)
+	if (OutCompressedShader)
 	{
-		int32 CompressedSize32 = static_cast<int32>(OutCompressedSize);
-		checkf(static_cast<int64>(CompressedSize32) == OutCompressedSize, TEXT("CompressedSize is too large (%lld) for an old API that takes int32"), OutCompressedSize);
-		bCompressed = FCompression::CompressMemory(ShaderCompressionFormat, OutCompressedShader, CompressedSize32, UncompressedShaderCode, UncompressedSize, COMPRESS_BiasSize);
-		OutCompressedSize = static_cast<int64>(CompressedSize32);
+		OutCompressedSize = FOodleDataCompression::Compress(OutCompressedShader, OutCompressedSize, InUncompressedShaderCode, InUncompressedSize, InOodleCompressor, InOodleLevel);
+		check(OutCompressedSize != 0);
+		return OutCompressedSize != 0;
 	}
 	else
 	{
-		FOodleDataCompression::ECompressor OodleCompressor;
-		FOodleDataCompression::ECompressionLevel OodleLevel;
-		GetShaderCompressionOodleSettings(OodleCompressor, OodleLevel);
-
-		// don't pass a nullptr to Oodle if we're only requesting an estimate
-		OutCompressedSize = OutCompressedShader ? FOodleDataCompression::Compress(OutCompressedShader, OutCompressedSize, UncompressedShaderCode, UncompressedSize, OodleCompressor, OodleLevel) : 0;
-		bCompressed = OutCompressedSize != 0;
-
-		// Oodle needs to return an estimate
-		if (!bCompressed)
-		{
-			// for Oodle, there is a separate estimation functon
-			OutCompressedSize = FOodleDataCompression::CompressedBufferSizeNeeded(UncompressedSize);
-		}	
-	}
-
-	return bCompressed;
+		// Just requesting an estimate.
+		OutCompressedSize = FOodleDataCompression::CompressedBufferSizeNeeded(InUncompressedSize);
+		return false;
+	}	
 }
 
 void FSerializedShaderArchive::DecompressShader(int32 Index, const TArray<TArray<uint8>>& ShaderCode, TArray<uint8>& OutDecompressedShader) const
@@ -401,7 +384,7 @@ void FSerializedShaderArchive::DecompressShader(int32 Index, const TArray<TArray
 	}
 	else
 	{
-		ShaderCodeArchive::DecompressShader(OutDecompressedShader.GetData(), Entry.UncompressedSize, ShaderCode[Index].GetData(), Entry.Size);
+		ShaderCodeArchive::DecompressShaderWithOodle(OutDecompressedShader.GetData(), Entry.UncompressedSize, ShaderCode[Index].GetData(), Entry.Size);
 	}
 }
 
@@ -1233,7 +1216,7 @@ TRefCountPtr<FRHIShader> FShaderCodeArchive::CreateShader(int32 Index)
 	if (ShaderEntry.UncompressedSize != ShaderEntry.Size)
 	{
 		uint8* UncompressedCode = reinterpret_cast<uint8*>(MemStack.Alloc(ShaderEntry.UncompressedSize, 16));
-		ShaderCodeArchive::DecompressShader(UncompressedCode, ShaderEntry.UncompressedSize, ShaderCode, ShaderEntry.Size);
+		ShaderCodeArchive::DecompressShaderWithOodle(UncompressedCode, ShaderEntry.UncompressedSize, ShaderCode, ShaderEntry.Size);
 		ShaderCode = (uint8*)UncompressedCode;
 	}
 
@@ -2069,7 +2052,7 @@ TRefCountPtr<FRHIShader> FIoStoreShaderCodeArchive::CreateShader(int32 ShaderInd
 	if (GroupEntry.IsGroupCompressed())
 	{
 		uint8* UncompressedCode = reinterpret_cast<uint8*>(MemStack.Alloc(GroupEntry.UncompressedSize, 16));
-		ShaderCodeArchive::DecompressShader(UncompressedCode, GroupEntry.UncompressedSize, ShaderCode, GroupEntry.CompressedSize);
+		ShaderCodeArchive::DecompressShaderWithOodle(UncompressedCode, GroupEntry.UncompressedSize, ShaderCode, GroupEntry.CompressedSize);
 		ShaderCode = reinterpret_cast<uint8*>(UncompressedCode) + ShaderEntry.UncompressedOffsetInGroup;
 
 #if UE_SCA_VISUALIZE_SHADER_USAGE
diff --git a/Engine/Source/Runtime/RenderCore/Private/ShaderCore.cpp b/Engine/Source/Runtime/RenderCore/Private/ShaderCore.cpp
--- a/Engine/Source/Runtime/RenderCore/Private/ShaderCore.cpp
+++ b/Engine/Source/Runtime/RenderCore/Private/ShaderCore.cpp
@@ -2393,6 +2393,8 @@ void FShaderCode::Compress(FName ShaderCompressionFormat, FOodleDataCompression:
 {
 	checkf(OptionalDataSize == -1, TEXT("FShaderCode::Compress() was called before calling FShaderCode::FinalizeShaderCode()"));
 
+	check(ShaderCompressionFormat == NAME_Oodle); // We now force shaders to compress with oodle (even if they are uncompressed)
+
 	TArray<uint8> Compressed;
 	// conventional formats will fail if the compressed size isn't enough, Oodle needs a more precise estimate
 	int32 CompressedSize = (ShaderCompressionFormat != NAME_Oodle) ? ShaderCodeWithOptionalData.Num() : FOodleDataCompression::CompressedBufferSizeNeeded(ShaderCodeWithOptionalData.Num());
diff --git a/Engine/Source/Runtime/RenderCore/Private/ShaderResource.cpp b/Engine/Source/Runtime/RenderCore/Private/ShaderResource.cpp
--- a/Engine/Source/Runtime/RenderCore/Private/ShaderResource.cpp
+++ b/Engine/Source/Runtime/RenderCore/Private/ShaderResource.cpp
@@ -28,48 +28,6 @@
 
 DECLARE_LOG_CATEGORY_CLASS(LogShaderWarnings, Log, Log);
 
-int32 GShaderCompressionFormatChoice = 2;
-static FAutoConsoleVariableRef CVarShaderCompressionFormatChoice(
-	TEXT("r.Shaders.CompressionFormat"),
-	GShaderCompressionFormatChoice,
-	TEXT("Select the compression methods for the shader code.\n")
-	TEXT(" 0: None (uncompressed)\n")
-	TEXT(" 1: LZ4\n")
-	TEXT(" 2: Oodle (default)\n")
-	TEXT(" 3: ZLib\n"),
-	ECVF_ReadOnly);
-
-int32 GShaderCompressionOodleAlgo = 2;
-static FAutoConsoleVariableRef CVarShaderCompressionOodleAlgo(
-	TEXT("r.Shaders.CompressionFormat.Oodle.Algo"),
-	GShaderCompressionOodleAlgo,
-	TEXT("Oodle compression method for the shader code, from fastest to slowest to decode.\n")
-	TEXT(" 0: None (invalid setting)\n")
-	TEXT(" 1: Selkie (fastest to decode)\n")
-	TEXT(" 2: Mermaid\n")
-	TEXT(" 3: Kraken\n")
-	TEXT(" 4: Leviathan (slowest to decode)\n"),
-	ECVF_ReadOnly);
-
-int32 GShaderCompressionOodleLevel = 6;
-static FAutoConsoleVariableRef CVarShaderCompressionOodleAlgoChoice(
-	TEXT("r.Shaders.CompressionFormat.Oodle.Level"),
-	GShaderCompressionOodleLevel,
-	TEXT("Oodle compression level. This mostly trades encode speed vs compression ratio, decode speed is determined by r.Shaders.CompressionFormat.Oodle.Algo\n")
-	TEXT(" -4 : HyperFast4\n")
-	TEXT(" -3 : HyperFast3\n")
-	TEXT(" -2 : HyperFast2\n")
-	TEXT(" -1 : HyperFast1\n")
-	TEXT("  0 : None\n")
-	TEXT("  1 : SuperFast\n")
-	TEXT("  2 : VeryFast\n")
-	TEXT("  3 : Fast\n")
-	TEXT("  4 : Normal\n")
-	TEXT("  5 : Optimal1\n")
-	TEXT("  6 : Optimal2\n")
-	TEXT("  7 : Optimal3\n")
-	TEXT("  8 : Optimal4\n"),
-	ECVF_ReadOnly);
 
 static int32 GShaderCompilerEmitWarningsOnLoad = 0;
 static FAutoConsoleVariableRef CVarShaderCompilerEmitWarningsOnLoad(
@@ -79,35 +37,19 @@ static FAutoConsoleVariableRef CVarShaderCompilerEmitWarningsOnLoad(
 	ECVF_Default
 );
 
-FName GetShaderCompressionFormat(const FName& ShaderFormat)
+FName GetShaderCompressionFormat()
 {
-	// support an older developer-only CVar for compatibility and make it preempt
-#if !(UE_BUILD_SHIPPING || UE_BUILD_TEST)
-	static const IConsoleVariable* CVarSkipCompression = IConsoleManager::Get().FindConsoleVariable(TEXT("r.Shaders.SkipCompression"));
-	static bool bSkipCompression = (CVarSkipCompression && CVarSkipCompression->GetInt() != 0);
-	if (UNLIKELY(bSkipCompression))
-	{
-		return NAME_None;
-	}
-#endif
-
-	static FName Formats[]
-	{
-		NAME_None,
-		NAME_LZ4,
-		NAME_Oodle,
-		NAME_Zlib
-	};
-	
-	//GShaderCompressionFormatChoice = (GShaderCompressionFormatChoice < 0) ? 0 : GShaderCompressionFormatChoice;
-	GShaderCompressionFormatChoice = FMath::Clamp<int32>(GShaderCompressionFormatChoice, 0, UE_ARRAY_COUNT(Formats) - 1);
-	return Formats[GShaderCompressionFormatChoice];
+	// We always use oodle now. This was instituted because UnrealPak recompresses the shaders and doens't have
+	// access to the INIs that drive the CVars and would always use default, resulting in mismatches for non
+	// default encoder selection.
+	return NAME_Oodle;
 }
 
 void GetShaderCompressionOodleSettings(FOodleDataCompression::ECompressor& OutCompressor, FOodleDataCompression::ECompressionLevel& OutLevel, const FName& ShaderFormat)
 {
 	// support an older developer-only CVar for compatibility and make it preempt
 #if !(UE_BUILD_SHIPPING || UE_BUILD_TEST)
+	// Since we always use Oodle, we make SkipCompression tell Oodle to not compress.
 	static const IConsoleVariable* CVarSkipCompression = IConsoleManager::Get().FindConsoleVariable(TEXT("r.Shaders.SkipCompression"));
 	static bool bSkipCompression = (CVarSkipCompression && CVarSkipCompression->GetInt() != 0);
 	if (UNLIKELY(bSkipCompression))
@@ -118,11 +60,9 @@ void GetShaderCompressionOodleSettings(FOodleDataCompression::ECompressor& OutCo
 	}
 #endif
 
-	GShaderCompressionOodleAlgo = FMath::Clamp(GShaderCompressionOodleAlgo, static_cast<int32>(FOodleDataCompression::ECompressor::NotSet), static_cast<int32>(FOodleDataCompression::ECompressor::Leviathan));
-	OutCompressor = static_cast<FOodleDataCompression::ECompressor>(GShaderCompressionOodleAlgo);
-
-	GShaderCompressionOodleLevel = FMath::Clamp(GShaderCompressionOodleLevel, static_cast<int32>(FOodleDataCompression::ECompressionLevel::HyperFast4), static_cast<int32>(FOodleDataCompression::ECompressionLevel::Optimal4));
-	OutLevel = static_cast<FOodleDataCompression::ECompressionLevel>(GShaderCompressionOodleLevel);
+	// We just use mermaid/normal here since these settings get overwritten in unrealpak, so this is just for non pak'd builds.
+	OutCompressor = FOodleDataCompression::ECompressor::Mermaid;
+	OutLevel = FOodleDataCompression::ECompressionLevel::Normal;
 }
 
 bool FShaderMapResource::ArePlatformsCompatible(EShaderPlatform CurrentPlatform, EShaderPlatform TargetPlatform)
diff --git a/Engine/Source/Runtime/RenderCore/Public/ShaderCodeArchive.h b/Engine/Source/Runtime/RenderCore/Public/ShaderCodeArchive.h
--- a/Engine/Source/Runtime/RenderCore/Public/ShaderCodeArchive.h
+++ b/Engine/Source/Runtime/RenderCore/Public/ShaderCodeArchive.h
@@ -385,13 +385,11 @@ namespace ShaderCodeArchive
 	/** Decompresses the shader into caller-provided memory. Caller is assumed to allocate at least ShaderEntry uncompressed size value.
 	 * The engine will crash (LogFatal) if this function fails.
 	 */
-	RENDERCORE_API void DecompressShader(uint8* OutDecompressedShader, int64 UncompressedSize, const uint8* CompressedShaderCode, int64 CompressedSize);
-
-	/** Compresses the shader into caller-provided memory using the current CVar settings (e.g. on Windows), which is as opposed to potentially different settings on 
-	 * the platform where the shaders are going to be decompressed (e.g. on console). It is the caller's responsibility to make sure those two match.
-	 * This function can also be used for the estimation as it will return false if the provided CompressedBufferLength isn't sufficient (it is safe to pass nullptr in that case).
-	 */
-	RENDERCORE_API bool CompressShaderUsingCurrentSettings(uint8* OutCompressedShader, int64& OutCompressedSize, const uint8* UncompressedShaderCode, int64 UncompressedSize);
+	RENDERCORE_API void DecompressShaderWithOodle(uint8* OutDecompressedShader, int64 UncompressedSize, const uint8* CompressedShaderCode, int64 CompressedSize);
+	
+	// We decompression, group, and recompress shaders when they are added to iostore containers in UnrealPak, where we don't have access to cvars - so there's no way 
+	// to configure - so we force Oodle and allow specification of the parameters here.
+	RENDERCORE_API bool CompressShaderWithOodle(uint8* OutCompressedShader, int64& OutCompressedSize, const uint8* InUncompressedShaderCode, int64 InUncompressedSize, FOodleDataCompression::ECompressor InOodleCompressor, FOodleDataCompression::ECompressionLevel InOodleLevel);
 }
 
 /** Descriptor of a shader map. This concept exists in run time, so this class describes the information stored in the library for a particular FShaderMap */
diff --git a/Engine/Source/Runtime/RenderCore/Public/ShaderCore.h b/Engine/Source/Runtime/RenderCore/Public/ShaderCore.h
--- a/Engine/Source/Runtime/RenderCore/Public/ShaderCore.h
+++ b/Engine/Source/Runtime/RenderCore/Public/ShaderCore.h
@@ -131,8 +131,8 @@ extern RENDERCORE_API bool AllowDebugViewmodes();
 /** Returns true if debug viewmodes are allowed for the given platform. */
 extern RENDERCORE_API bool AllowDebugViewmodes(EShaderPlatform Platform);
 
-/** Returns the shader compression format (passing ShaderFormat for future proofing, but as of now the setting is global for all formats). */
-extern RENDERCORE_API FName GetShaderCompressionFormat(const FName& ShaderFormat = NAME_None);
+/** Returns the shader compression format. Oodle is used exclusively now. r.Shaders.SkipCompression configures Oodle to be uncompressed instead of returning NAME_None.*/
+extern RENDERCORE_API FName GetShaderCompressionFormat();
 
 namespace FOodleDataCompression
 {
