Fix for parsing validate flag.
Set the reference chunk db on the optional writers.
Disable reference chunk use for uncompressed chunks.
Failing to load the reference chunk db is now a warning rather than a failure, since pak can succeed without it.

diff --git a/Engine/Source/Developer/IoStoreUtilities/Private/IoStoreUtilities.cpp b/Engine/Source/Developer/IoStoreUtilities/Private/IoStoreUtilities.cpp
--- a/Engine/Source/Developer/IoStoreUtilities/Private/IoStoreUtilities.cpp
+++ b/Engine/Source/Developer/IoStoreUtilities/Private/IoStoreUtilities.cpp
@@ -3364,8 +3364,7 @@ int32 CreateTarget(const FIoStoreArguments& Arguments, const FIoStoreWriterSetti
 		ChunkDatabase = MakeShared<FIoStoreChunkDatabase>();
 		if (((FIoStoreChunkDatabase&)*ChunkDatabase).Init(Arguments.ReferenceChunkGlobalContainerFileName, Arguments.ReferenceChunkKeys) == false)
 		{
-			UE_LOG(LogIoStore, Error, TEXT("Failed to initialize reference chunk store."));
-			return 1;
+			UE_LOG(LogIoStore, Warning, TEXT("Failed to initialize reference chunk store. Pak will continue."));
 		}
 	}
 
@@ -3461,6 +3460,7 @@ int32 CreateTarget(const FIoStoreArguments& Arguments, const FIoStoreWriterSetti
 					if (!ContainerTarget->OptionalSegmentOutputPath.IsEmpty())
 					{
 						ContainerTarget->OptionalSegmentIoStoreWriter = IoStoreWriterContext->CreateContainer(*ContainerTarget->OptionalSegmentOutputPath, ContainerSettings);
+						ContainerTarget->OptionalSegmentIoStoreWriter->SetReferenceChunkDatabase(ChunkDatabase);
 						IoStoreWriters.Add(ContainerTarget->OptionalSegmentIoStoreWriter);
 					}
 				}
@@ -6518,7 +6518,7 @@ bool ParseContainerGenerationArguments(FIoStoreArguments& Arguments, FIoStoreWri
 	// By default, we use any hashes in the asset registry that exist in order to avoid reading and hashing
 	// chunk unnecessarily. This flag causes us to read and hash anyway, and then ensure they match what is
 	// in the asset registry. It is very bad if this fails!
-	Arguments.bVerifyHashDatabase = FParse::Param(FCommandLine::Get(), TEXT("-verifyhashdatabase"));
+	Arguments.bVerifyHashDatabase = FParse::Param(FCommandLine::Get(), TEXT("verifyhashdatabase"));
 
 
 	uint64 PatchPaddingAlignment = 0;
diff --git a/Engine/Source/Runtime/Core/Private/IO/IoStore.cpp b/Engine/Source/Runtime/Core/Private/IO/IoStore.cpp
--- a/Engine/Source/Runtime/Core/Private/IO/IoStore.cpp
+++ b/Engine/Source/Runtime/Core/Private/IO/IoStore.cpp
@@ -1497,7 +1497,7 @@ private:
 		//
 		// Check if this chunk exists in the reference cache.
 		//
-		if (ReferenceChunkDatabase.IsValid())
+		if (ReferenceChunkDatabase.IsValid() && CompressionMethod != NAME_None)
 		{
 			TPair<FIoContainerId, FIoChunkHash> ChunkKey(ContainerSettings.ContainerId, Entry->ChunkHash);
 
