UnrealPak: Add -csv parameter to create pak and create ucas commands to list all the chunks.
Output is the same as when executing the -list command separately.
Running with -csv=<dir> will create one <pakchunkfilename>.pak.csv file per generated pak and one <pakchunkfilename>.utoc.csv file per generated container.
Running with -csv=<file> will create one single <file>.pak.csv for all paks and one <file>.utoc.csv for all containers.

Details:
The container output (the big set of data) is generated based on in-memory writer TOC data once the IoStoreWriters have been finalized.
The pak output (smaller set of data) is generated using the existing list command and loads the newly generated files from disk.
Cleanup the remains of three obsolete csv output features for containers (-csvoutput, -writefinalorder and OUTPUT_CHUNKID_DIRECTORY).
Optimize container output generation by extending FIoStoreTocChunkInfo with OffsetOnDisk and NumCompressedBlocks filled directly by EnumerateChunks.
Add a ChunkIdToFileNameMap to FCookedPackageStore.
Use the FCookedPackageStore to output the actual package ids and package names (this is currently broken when running -list as a separate command).

#rnx

diff --git a/Engine/Source/Developer/IoStoreUtilities/Private/IoStoreUtilities.cpp b/Engine/Source/Developer/IoStoreUtilities/Private/IoStoreUtilities.cpp
--- a/Engine/Source/Developer/IoStoreUtilities/Private/IoStoreUtilities.cpp
+++ b/Engine/Source/Developer/IoStoreUtilities/Private/IoStoreUtilities.cpp
@@ -15,7 +15,6 @@
 #include "Misc/ConfigCacheIni.h"
 #include "Misc/PackageName.h"
 #include "Misc/Paths.h"
-#include "Misc/PackageName.h"
 #include "Misc/Base64.h"
 #include "Misc/AES.h"
 #include "Misc/CoreDelegates.h"
@@ -78,8 +77,6 @@ IMPLEMENT_MODULE(FDefaultModuleImpl, IoStoreUtilities);
 
 TRACE_DECLARE_MEMORY_COUNTER(IoStoreUsedFileBufferMemory, TEXT("IoStore/UsedFileBufferMemory"));
 
-#define OUTPUT_CHUNKID_DIRECTORY 0
-
 static const FName DefaultCompressionMethod = NAME_Zlib;
 static const uint64 DefaultCompressionBlockSize = 64 << 10;
 static const uint64 DefaultCompressionBlockAlignment = 64 << 10;
@@ -640,8 +637,9 @@ struct FContainerTargetFile
 class FCookedPackageStore
 {
 public:
-	FCookedPackageStore(const FString& CookedDir)
-		: PackageStoreManifest(CookedDir)
+	FCookedPackageStore(const FString& InCookedDir)
+		: PackageStoreManifest(InCookedDir)
+		, CookedDir(InCookedDir)
 	{
 	}
 
@@ -703,6 +701,7 @@ public:
 		for (const FPackageStoreManifest::FFileInfo& FileInfo : PackageStoreManifest.GetFiles())
 		{
 			FilenameToChunkIdMap.Add(FileInfo.FileName, FileInfo.ChunkId);
+			ChunkIdToFileNameMap.Add(FileInfo.ChunkId, FileInfo.FileName.RightChop(CookedDir.Len() + 1));
 			FName* FindPackageName = ChunkIdToPackageNameMap.Find(FileInfo.ChunkId);
 			if (FindPackageName)
 			{
@@ -716,7 +715,12 @@ public:
 
 	FIoChunkId GetChunkIdFromFileName(const FString& Filename) const
 	{
-		return FilenameToChunkIdMap.FindRef(Filename);
+		return FilenameToChunkIdMap.FindRef(*Filename);
+	}
+
+	FString GetFilenameFromChunkId(const FIoChunkId& ChunkId) const
+	{
+		return ChunkIdToFileNameMap.FindRef(ChunkId);
 	}
 
 	FName GetPackageNameFromChunkId(const FIoChunkId& ChunkId) const
@@ -825,9 +829,11 @@ public:
 private:
 	TUniquePtr<UE::FZenStoreHttpClient> ZenStoreClient;
 	FPackageStoreManifest PackageStoreManifest;
+	FString CookedDir;
 	TArray<FPackageStoreManifest::FPackageInfo> ManifestPackageInfos;
 	TMap<FPackageId, FPackageStoreEntryResource> PackageIdToEntry;
 	TMap<FString, FIoChunkId> FilenameToChunkIdMap;
+	TMap<FIoChunkId, FString> ChunkIdToFileNameMap;
 	TMap<FIoChunkId, FName> ChunkIdToPackageNameMap;
 	TMap<FString, FName> FilenameToPackageNameMap;
 	TMap<FName, const FPackageStoreManifest::FPackageInfo*> PackageNameToPackageInfoMap;
@@ -859,6 +865,7 @@ struct FIoStoreArguments
 	FString DLCPluginPath;
 	FString DLCName;
 	FString ReferenceChunkGlobalContainerFileName;
+	FString CsvPath;
 	FKeyChain ReferenceChunkKeys;
 	FReleasedPackages ReleasedPackages;
 	TUniquePtr<FCookedPackageStore> PackageStore;
@@ -906,43 +913,51 @@ struct FContainerTargetSpec
 using FPackageNameMap = TMap<FName, FLegacyCookedPackage*>;
 using FPackageIdMap = TMap<FPackageId, FLegacyCookedPackage*>;
 
-#if OUTPUT_CHUNKID_DIRECTORY
-class FChunkIdCsv
+class FChunkEntryCsv
 {
 public:
 
-	~FChunkIdCsv()
+	~FChunkEntryCsv()
 	{
-		if (OutputArchive)
+		if (OutputArchive.IsValid())
 		{
 			OutputArchive->Flush();
 		}
 	}
 
-	void CreateOutputFile(const FString& RootPath)
+	void CreateOutputFile(const TCHAR* OutputFilename)
 	{
-		const FString OutputFilename = RootPath / TEXT("chunkid_directory.csv");
-		OutputArchive.Reset(IFileManager::Get().CreateFileWriter(*OutputFilename));
-		if (OutputArchive)
+		OutputArchive.Reset(IFileManager::Get().CreateFileWriter(OutputFilename));
+		if (OutputArchive.IsValid())
 		{
-			const ANSICHAR* Output = "NameIndex,NameNumber,ChunkIndex,ChunkType,ChunkIdHash,DebugString\n";
-			OutputArchive->Serialize((void*)Output, FPlatformString::Strlen(Output));
+			OutputArchive->Logf(TEXT("OrderInContainer, ChunkId, PackageId, PackageName, Filename, ContainerName, Offset, OffsetOnDisk, Size, CompressedSize, Hash, ChunkType"));
 		}
 	}
 
-	void AddChunk(uint32 NameIndex, uint32 NameNumber, uint16 ChunkIndex, uint8 ChunkType, uint32 ChunkIdHash, const TCHAR* DebugString)
+	void AddChunk(const FString& ContainerName, int32 Index, const FIoStoreTocChunkInfo& Info, FPackageId PackageId, const FString& PackageName)
 	{
-		ANSICHAR Buffer[MAX_SPRINTF + 1] = { 0 };
-		int32 NumOfCharacters = FCStringAnsi::Sprintf(Buffer, "%u,%u,%u,%u,%u,%s\n", NameIndex, NameNumber, ChunkIndex, ChunkType, ChunkIdHash, TCHAR_TO_ANSI(DebugString));
-		OutputArchive->Serialize(Buffer, NumOfCharacters);
+		if (OutputArchive.IsValid())
+		{
+			OutputArchive->Logf(TEXT("%d, %s, 0x%llX, %s, %s, %s, %lld, %lld, %lld, %lld, 0x%s, %s"),
+				Index,
+				*BytesToHex(Info.Id.GetData(), Info.Id.GetSize()),
+				PackageId.ValueForDebugging(),
+				*PackageName,
+				*Info.FileName,
+				*ContainerName,
+				Info.Offset,
+				Info.OffsetOnDisk,
+				Info.Size,
+				Info.CompressedSize,
+				*Info.Hash.ToString(),
+				*LexToString(Info.ChunkType)
+			);
+		}
 	}	
 
 private:
 	TUniquePtr<FArchive> OutputArchive;
 };
-FChunkIdCsv ChunkIdCsv;
-
-#endif
 
 
 
@@ -3354,10 +3369,6 @@ int32 CreateTarget(const FIoStoreArguments& Arguments, const FIoStoreWriterSetti
 	IOSTORE_CPU_SCOPE(CreateTarget);
 	TGuardValue<int32> GuardAllowUnversionedContentInEditor(GAllowUnversionedContentInEditor, 1);
 
-#if OUTPUT_CHUNKID_DIRECTORY
-	ChunkIdCsv.CreateOutputFile(CookedDir);
-#endif
-
 	TSharedPtr<IIoStoreWriterReferenceChunkDatabase> ChunkDatabase;
 	if (Arguments.ReferenceChunkGlobalContainerFileName.Len())
 	{
@@ -3631,45 +3642,6 @@ int32 CreateTarget(const FIoStoreArguments& Arguments, const FIoStoreWriterSetti
 				}
 			}
 
-			// Check if we need to dump the final order of the packages. Useful, to debug packing.
-			if (FParse::Param(FCommandLine::Get(), TEXT("writefinalorder")))
-			{
-				IOSTORE_CPU_SCOPE(WriteFinalOrder);
-				FString FinalContainerOrderFile = FPaths::GetPath(ContainerTarget->OutputPath) + FPaths::GetBaseFilename(ContainerTarget->OutputPath) + TEXT("-order.txt");
-				TUniquePtr<FArchive> IoOrderListArchive(IFileManager::Get().CreateFileWriter(*FinalContainerOrderFile));
-				if (IoOrderListArchive)
-				{
-					IoOrderListArchive->SetIsTextFormat(true);
-
-					// The TargetFiles list is not the order written to disk - FIoStoreWriter sorts
-					// on the IdealOrder prior to writing.
-					TArray<const FContainerTargetFile*> SortedTargetFiles;
-					SortedTargetFiles.Reserve(ContainerTarget->TargetFiles.Num());
-					for (const FContainerTargetFile& TargetFile : ContainerTarget->TargetFiles)
-					{
-						SortedTargetFiles.Add(&TargetFile);
-					}
-
-					Algo::SortBy(SortedTargetFiles, [](const FContainerTargetFile* TargetFile) { return TargetFile->IdealOrder; });
-
-					for (const FContainerTargetFile* TargetFile : SortedTargetFiles)
-					{
-						TStringBuilder<256> Line;
-						Line.Append(LexToString(TargetFile->ChunkId));
-
-						if (TargetFile->Package)
-						{
-							Line.Append(TEXT(" "));
-							Line.Append(TargetFile->Package->FileName);
-						}
-
-						IoOrderListArchive->Logf(TEXT("%s"), Line.ToString());
-					}
-
-					IoOrderListArchive->Close();
-				}
-			}
-
 		}
 	}
 
@@ -3735,14 +3707,95 @@ int32 CreateTarget(const FIoStoreArguments& Arguments, const FIoStoreWriterSetti
 	}
 
 	TArray<FIoStoreWriterResult> IoStoreWriterResults;
-	IoStoreWriterResults.Reserve(IoStoreWriters.Num());
-	for (TSharedPtr<IIoStoreWriter> IoStoreWriter : IoStoreWriters)
 	{
-		IoStoreWriterResults.Emplace(IoStoreWriter->GetResult().ConsumeValueOrDie());
+		TRACE_CPUPROFILER_EVENT_SCOPE(GetWriterResults);
+		IoStoreWriterResults.Reserve(IoStoreWriters.Num());
+		for (TSharedPtr<IIoStoreWriter> IoStoreWriter : IoStoreWriters)
+		{
+			IoStoreWriterResults.Emplace(IoStoreWriter->GetResult().ConsumeValueOrDie());
+		}
+	}
+
+	FGraphEventRef WriteCsvFileTask;
+	if (Arguments.CsvPath.Len() > 0)
+	{
+		WriteCsvFileTask = FFunctionGraphTask::CreateAndDispatchWhenReady([&Arguments, &IoStoreWriters, &IoStoreWriterResults]()
+		{
+			TRACE_CPUPROFILER_EVENT_SCOPE(WriteCsvFiles);
+
+			bool bPerContainerCsvFiles = FPaths::DirectoryExists(Arguments.CsvPath);
+			FChunkEntryCsv AllContainersOutCsvFile;
+			FChunkEntryCsv* Out = &AllContainersOutCsvFile;
+			if (!bPerContainerCsvFiles)
+			{
+				// When CsvPath is a filename append .utoc.csv to create a unique single csv for all container files,
+				// different from the unique single .pak.csv for all pak files.
+				FString CsvFilename = Arguments.CsvPath + TEXT(".utoc.csv");
+				AllContainersOutCsvFile.CreateOutputFile(*CsvFilename);
+			}
+
+			for (int32 Index = 0; Index < IoStoreWriters.Num(); ++Index)
+			{
+				TRACE_CPUPROFILER_EVENT_SCOPE(ListContainer);
+
+				TSharedPtr<IIoStoreWriter> Writer = IoStoreWriters[Index];
+				FIoStoreWriterResult& Result = IoStoreWriterResults[Index];
+
+				TArray<FIoStoreTocChunkInfo> Chunks;
+				{
+					IOSTORE_CPU_SCOPE(EnumerateChunks);
+					Chunks.Reserve(Result.TocEntryCount);
+					Writer->EnumerateChunks([&Chunks](FIoStoreTocChunkInfo&& ChunkInfo)
+					{
+						Chunks.Add(MoveTemp(ChunkInfo));
+						return true;
+					});
+				}
+
+				{
+					IOSTORE_CPU_SCOPE(SortChunks);
+					auto SortKey = [](const FIoStoreTocChunkInfo& ChunkInfo) { return ChunkInfo.OffsetOnDisk; };
+					Algo::SortBy(Chunks, SortKey);
+				}
+
+				{
+					IOSTORE_CPU_SCOPE(WriteCsvFile);
+					FChunkEntryCsv PerContainerOutCsvFile;
+					if (bPerContainerCsvFiles)
+					{
+						// When CsvPath is a dir, then create one unique .utoc.csv per container file
+						FString PerContainerCsvPath = Arguments.CsvPath / Result.ContainerName + TEXT(".utoc.csv");
+						PerContainerOutCsvFile.CreateOutputFile(*PerContainerCsvPath);
+						Out = &PerContainerOutCsvFile;
+					}
+					for (int32 EntryIndex=0; EntryIndex < Chunks.Num(); ++EntryIndex)
+					{
+						FIoStoreTocChunkInfo& ChunkInfo = Chunks[EntryIndex];
+						FString PackageName;
+						FPackageId PackageId;
+						if (!ChunkInfo.bHasValidFileName)
+						{
+							FString FileName = Arguments.PackageStore->GetFilenameFromChunkId(ChunkInfo.Id);
+							if (FileName.Len() > 0)
+							{
+								ChunkInfo.FileName = MoveTemp(FileName);
+							}
+							FName PackageFName = Arguments.PackageStore->GetPackageNameFromChunkId(ChunkInfo.Id);
+							if (!PackageFName.IsNone())
+							{
+								PackageName = PackageFName.ToString();
+								PackageId = FPackageId::FromName(FName(*PackageName));
+							}
+						}
+
+						Out->AddChunk(Result.ContainerName, EntryIndex, ChunkInfo, PackageId, PackageName);
+					}
+				}
+			}
+		}, TStatId(), nullptr, ENamedThreads::AnyNormalThreadHiPriTask);
 	}
-	IoStoreWriters.Empty();
 
-	IOSTORE_CPU_SCOPE(CalculateStats);
+	IOSTORE_CPU_SCOPE(OutputStats);
 
 	UE_LOG(LogIoStore, Display, TEXT("Calculating stats..."));
 	uint64 UExpSize = 0;
@@ -3833,6 +3886,14 @@ int32 CreateTarget(const FIoStoreArguments& Arguments, const FIoStoreWriterSetti
 		}
 	}
 
+	UE_LOG(LogIoStore, Display, TEXT(""));
+	if (Arguments.CsvPath.Len() > 0)
+	{
+		TRACE_CPUPROFILER_EVENT_SCOPE(WaitForCsvFiles);
+		UE_LOG(LogIoStore, Display, TEXT("Writing csv file(s) to: %s (*.utoc.csv)"), *Arguments.CsvPath);
+		FTaskGraphInterface::Get().WaitUntilTaskCompletes(WriteCsvFileTask);
+	}
+
 	return 0;
 }
 
@@ -3985,21 +4046,7 @@ int32 ListContainer(
 	const FString& ContainerPathOrWildcard,
 	const FString& CsvPath)
 {
-	struct FComputedChunkInfo
-	{
-		FIoChunkId Id;
-		FString FileName;
-		FIoChunkHash Hash;
-		EIoChunkType ChunkType;
-		uint64 Size;
-		uint64 CompressedSize;
-		uint64 Offset;
-		uint64 OffsetOnDisk;
-		int32 NumCompressedBlocks;
-		bool bIsCompressed;
-		bool bHasValidFileName;
-	};
-
+	IOSTORE_CPU_SCOPE(ListContainer);
 	TArray<FString> ContainerFilePaths;
 
 	if (IFileManager::Get().FileExists(*ContainerPathOrWildcard))
@@ -4041,18 +4088,14 @@ int32 ListContainer(
 	
 	// if CsvPath is a dir, not a file name, then write one csv per container to the dir
 	// otherwise, write all contents to one big csv
-	bool bCsvPathIsDir = IFileManager::Get().DirectoryExists(*CsvPath);
+	bool bPerContainerCsvFiles = IFileManager::Get().DirectoryExists(*CsvPath);
 
-	TUniquePtr<FOutputDeviceFile> AllContainersOutCsvFile;
-	if ( ! bCsvPathIsDir )
+	FChunkEntryCsv AllContainersOutCsvFile;
+	if (!bPerContainerCsvFiles)
 	{
-		AllContainersOutCsvFile = MakeUnique<FOutputDeviceFile>(*CsvPath, true);
-		
-		AllContainersOutCsvFile->SetSuppressEventTag(true);
+		AllContainersOutCsvFile.CreateOutputFile(*CsvPath);
 	}
 
-	bool bNeedsHeader = true;
-
 	for (const FString& ContainerFilePath : ContainerFilePaths)
 	{
 		TUniquePtr<FIoStoreReader> Reader = CreateIoStoreReader(*ContainerFilePath, KeyChain);
@@ -4062,33 +4105,19 @@ int32 ListContainer(
 			continue;
 		}
 		
-		TUniquePtr<FOutputDeviceFile> PerContainerOutCsvFile;
-		FOutputDeviceFile * Out;
+		FChunkEntryCsv PerContainerOutCsvFile;
+		FChunkEntryCsv* Out;
 		
-		if ( AllContainersOutCsvFile != nullptr )
+		if (!bPerContainerCsvFiles)
 		{
-			Out = AllContainersOutCsvFile.Get();
+			Out = &AllContainersOutCsvFile;
 		}
 		else
 		{
-			// ContainerFilePath is a .utoc
-			//	change to .csv and put in CsvPath
-			FString PerContainerCsvPath = CsvPath / FPaths::GetCleanFilename(ContainerFilePath);
-			PerContainerCsvPath = FPaths::ChangeExtension(PerContainerCsvPath, FString(TEXT(".csv")) );
-			
-			UE_LOG(LogIoStore, Display, TEXT("PerContainerCsvPath: '%s'"), *PerContainerCsvPath);
-
-			PerContainerOutCsvFile = MakeUnique<FOutputDeviceFile>(*PerContainerCsvPath, true);
-
-			Out = PerContainerOutCsvFile.Get();
-			Out->SetSuppressEventTag(true);
-			bNeedsHeader = true;
-		}
-
-		if ( bNeedsHeader )
-		{
-			Out->Log(TEXT("OrderInContainer, ChunkId, PackageId, PackageName, Filename, ContainerName, Offset, OffsetOnDisk, Size, CompressedSize, Hash, ChunkType"));
-			bNeedsHeader = false;
+			// ContainerFilePath is a .utoc, add .csv and put in CsvPath
+			FString PerContainerCsvPath = CsvPath / FPaths::GetCleanFilename(ContainerFilePath) + TEXT(".csv");
+			PerContainerOutCsvFile.CreateOutputFile(*PerContainerCsvPath);
+			Out = &PerContainerOutCsvFile;
 		}
 
 		if (!EnumHasAnyFlags(Reader->GetContainerFlags(), EIoContainerFlags::Indexed))
@@ -4099,67 +4128,39 @@ int32 ListContainer(
 		UE_LOG(LogIoStore, Display, TEXT("Listing container '%s'"), *ContainerFilePath);
 
 		FString ContainerName = FPaths::GetBaseFilename(ContainerFilePath);
-		uint64 CompressionBlockSize = Reader->GetCompressionBlockSize();
-		TArray<FIoStoreTocCompressedBlockInfo> CompressedBlocks;
-		Reader->EnumerateCompressedBlocks([&CompressedBlocks](const FIoStoreTocCompressedBlockInfo& Block) {
-			CompressedBlocks.Add(Block);
-			return true;
-			});
-
+		TArray<FIoStoreTocChunkInfo> Chunks;
 
-		TArray<FComputedChunkInfo> Chunks;
-		Reader->EnumerateChunks([&Chunks, CompressionBlockSize, &CompressedBlocks](const FIoStoreTocChunkInfo& ChunkInfo) {
-
-			int32 FirstBlockIndex = int32(ChunkInfo.Offset / CompressionBlockSize);
-			int32 LastBlockIndex = int32((Align(ChunkInfo.Offset + ChunkInfo.Size, CompressionBlockSize) - 1) / CompressionBlockSize);
-			int32 NumCompressedBlocks = LastBlockIndex - FirstBlockIndex + 1;
-			uint64 OffsetOnDisk = CompressedBlocks[FirstBlockIndex].Offset;
-
-			FComputedChunkInfo ComputedInfo{
-				ChunkInfo.Id,
-				ChunkInfo.FileName,
-				ChunkInfo.Hash,
-				ChunkInfo.ChunkType,
-				ChunkInfo.Size,
-				ChunkInfo.CompressedSize,
-				ChunkInfo.Offset,
-				OffsetOnDisk,
-				NumCompressedBlocks,
-				ChunkInfo.bIsCompressed,
-				ChunkInfo.bHasValidFileName
-			};
-
-			Chunks.Add(ComputedInfo);
-			return true;
+		{
+			IOSTORE_CPU_SCOPE(EnumerateChunks);
+			Reader->EnumerateChunks([&Chunks](FIoStoreTocChunkInfo&& ChunkInfo)
+			{
+				Chunks.Add(MoveTemp(ChunkInfo));
+				return true;
 			});
+		}
 
-		auto SortKey = [](const FComputedChunkInfo& ChunkInfo) { return ChunkInfo.OffsetOnDisk; };
-		Algo::SortBy(Chunks, SortKey);
+		{
+			IOSTORE_CPU_SCOPE(EnumerateChunks);
+			auto SortKey = [](const FIoStoreTocChunkInfo& ChunkInfo) { return ChunkInfo.OffsetOnDisk; };
+			Algo::SortBy(Chunks, SortKey);
+		}
 
-		for(int32 Index=0; Index < Chunks.Num(); ++Index)
 		{
-			const FComputedChunkInfo& ChunkInfo = Chunks[Index];
+			IOSTORE_CPU_SCOPE(WriteCsvFile);
 			FString PackageName;
-			FPackageId PackageId;
-			if (ChunkInfo.bHasValidFileName && FPackageName::TryConvertFilenameToLongPackageName(ChunkInfo.FileName, PackageName, nullptr))
+			for(int32 Index=0; Index < Chunks.Num(); ++Index)
 			{
-				PackageId = FPackageId::FromName(FName(*PackageName));
-			}
+				const FIoStoreTocChunkInfo& ChunkInfo = Chunks[Index];
 
-			Out->Logf(TEXT("%d, %s, 0x%llX, %s, %s, %s, %lld, %lld, %lld, %lld, 0x%s, %s"),
-					Index,
-					*BytesToHex(ChunkInfo.Id.GetData(), ChunkInfo.Id.GetSize()),
-					PackageId.ValueForDebugging(),
-					*PackageName,
-					*ChunkInfo.FileName,
-					*ContainerName,
-					ChunkInfo.Offset,
-					ChunkInfo.OffsetOnDisk,
-					ChunkInfo.Size,
-					ChunkInfo.CompressedSize,
-					*ChunkInfo.Hash.ToString(),
-					*LexToString(ChunkInfo.ChunkType)
-					);
+				FPackageId PackageId;
+				PackageName.Reset();
+				if (ChunkInfo.bHasValidFileName && FPackageName::TryConvertFilenameToLongPackageName(ChunkInfo.FileName, PackageName, nullptr))
+				{
+					PackageId = FPackageId::FromName(FName(*PackageName));
+				}
+
+				Out->AddChunk(ContainerName, Index, ChunkInfo, PackageId, PackageName);
+			}
 		}
 	}
 
@@ -6453,8 +6454,6 @@ bool ParseContainerGenerationArguments(FIoStoreArguments& Arguments, FIoStoreWri
 	WriterSettings.CompressionMethod = DefaultCompressionMethod;
 	WriterSettings.CompressionBlockSize = DefaultCompressionBlockSize;
 
-	WriterSettings.bEnableCsvOutput = FParse::Param(FCommandLine::Get(), TEXT("csvoutput"));
-
 	TArray<FName> CompressionFormats;
 	FString DesiredCompressionFormats;
 	if (FParse::Value(FCommandLine::Get(), TEXT("-compressionformats="), DesiredCompressionFormats) ||
@@ -6693,17 +6692,18 @@ int32 CreateIoStoreContainerFiles(const TCHAR* CmdLine)
 		}
 	}
 
+	FParse::Value(FCommandLine::Get(), TEXT("csv="), Arguments.CsvPath);
+
 	FString ArgumentValue;
 	if (FParse::Value(FCommandLine::Get(), TEXT("List="), ArgumentValue))
 	{
 		FString ContainerPathOrWildcard = MoveTemp(ArgumentValue);
-		FString CsvPath;
-		if (!FParse::Value(FCommandLine::Get(), TEXT("csv="), CsvPath))
+		if (Arguments.CsvPath.Len() == 0)
 		{
 			UE_LOG(LogIoStore, Error, TEXT("Incorrect arguments. Expected: -list=<ContainerFile> -csv=<path>"));
 		}
 
-		return ListContainer(Arguments.KeyChain, ContainerPathOrWildcard, CsvPath);
+		return ListContainer(Arguments.KeyChain, ContainerPathOrWildcard, Arguments.CsvPath);
 	}
 	else if (FParse::Value(FCommandLine::Get(), TEXT("AssetRegistryWriteback="), ArgumentValue))
 	{
diff --git a/Engine/Source/Developer/PakFileUtilities/Private/PakFileUtilities.cpp b/Engine/Source/Developer/PakFileUtilities/Private/PakFileUtilities.cpp
--- a/Engine/Source/Developer/PakFileUtilities/Private/PakFileUtilities.cpp
+++ b/Engine/Source/Developer/PakFileUtilities/Private/PakFileUtilities.cpp
@@ -72,6 +72,7 @@ int64 GDDCHits = 0;
 int64 GDDCMisses = 0;
 #endif
 
+bool ListFilesInPak(const TCHAR * InPakFilename, int64 SizeFilter, bool bIncludeDeleted, const FString& CSVFilename, bool bExtractToMountPoint, const FKeyChain& InKeyChain, bool bAppendFile);
 
 bool SignPakFile(const FString& InPakFilename, const FRSAKeyHandle InSigningKey)
 {
@@ -550,6 +551,7 @@ struct FPakCommandLineParameters
 	FString SourcePatchDiffDirectory;
 	FString InputFinalPakFilename; // This is the resulting pak file we want to end up with after we generate the pak patch.  This is used instead of passing in the raw content.
 	FString ChangedFilesOutputFilename;
+	FString CsvPath;
 	bool EncryptIndex;
 	bool UseCustomCompressor;
 	FGuid EncryptionKeyGuid;
@@ -1261,6 +1263,8 @@ void ProcessCommonCommandLine(const TCHAR* CmdLine, FPakCommandLineParameters& C
 		CmdLineParameters.SeekOptParams.Mode = ESeekOptMode::OnePass;
 	}
 	FParse::Value(CmdLine, TEXT("-patchSeekOptMode="), (int32&)CmdLineParameters.SeekOptParams.Mode);
+
+	FParse::Value(CmdLine, TEXT("csv="), CmdLineParameters.CsvPath);
 }
 
 void ProcessPakFileSpecificCommandLine(const TCHAR* CmdLine, const TArray<FString>& NonOptionArguments, TArray<FPakInputPair>& Entries, FPakCommandLineParameters& CmdLineParameters)
@@ -2994,6 +2998,43 @@ bool FPakWriterContext::Flush()
 	UE_LOG(LogPakFile, Display, TEXT("DDC Sync Read Time: %lf"), ((double)GDDCSyncReadTime) * FPlatformTime::GetSecondsPerCycle());
 	UE_LOG(LogPakFile, Display, TEXT("DDC Sync Write Time: %lf"), ((double)GDDCSyncWriteTime) * FPlatformTime::GetSecondsPerCycle());
 #endif
+
+	if (CmdLineParameters.CsvPath.Len() > 0)
+	{
+		int64 SizeFilter = 0;
+		bool bIncludeDeleted = true;
+		bool bExtractToMountPoint = true;
+
+		bool bPerPakCsvFiles = FPaths::DirectoryExists(CmdLineParameters.CsvPath);
+		FString CsvFilename;
+		if (!bPerPakCsvFiles)
+		{
+			// When CsvPath is a filename append .pak.csv to create a unique single csv for all pak files,
+			// different from the unique single .utoc.csv for all container files.
+			CsvFilename = CmdLineParameters.CsvPath + TEXT(".pak.csv");
+		}
+		bool bAppendFile = false;
+		for (TUniquePtr<FOutputPakFile>& OutputPakFile : OutputPakFiles)
+		{
+			if (bPerPakCsvFiles)
+			{
+				// When CsvPath is a dir, then create one unique .pak.csv per pak file
+				CsvFilename = CmdLineParameters.CsvPath / FPaths::GetCleanFilename(OutputPakFile->Filename) + TEXT(".csv");
+			}
+			ListFilesInPak(
+				*OutputPakFile->Filename,
+				SizeFilter,
+				bIncludeDeleted,
+				CsvFilename,
+				bExtractToMountPoint,
+				OutputPakFile->KeyChain,
+				bAppendFile);
+			if (!bPerPakCsvFiles)
+			{
+				bAppendFile = true;
+			}
+		}
+	}
 	return true;
 }
 
@@ -3190,8 +3231,9 @@ bool TestPakFile(const TCHAR* Filename, bool TestHashes)
 	}
 }
 
-bool ListFilesInPak(const TCHAR * InPakFilename, int64 SizeFilter, bool bIncludeDeleted, const FString& CSVFilename, bool bExtractToMountPoint, const FKeyChain& InKeyChain)
+bool ListFilesInPak(const TCHAR * InPakFilename, int64 SizeFilter, bool bIncludeDeleted, const FString& CSVFilename, bool bExtractToMountPoint, const FKeyChain& InKeyChain, bool bAppendFile)
 {
+	TRACE_CPUPROFILER_EVENT_SCOPE(ListFilesInPak);
 	IPlatformFile* LowerLevelPlatformFile = &FPlatformFileManager::Get().GetPlatformFile();
 	TRefCountPtr<FPakFile> PakFilePtr = new FPakFile(LowerLevelPlatformFile, InPakFilename, false);
 	FPakFile& PakFile = *PakFilePtr;
@@ -3201,7 +3243,7 @@ bool ListFilesInPak(const TCHAR * InPakFilename, int64 SizeFilter, bool bInclude
 
 	if (PakFile.IsValid())
 	{
-		UE_LOG(LogPakFile, Display, TEXT("Mount point %s"), *PakFile.GetMountPoint());
+		UE_LOG(LogPakFile, Log, TEXT("Mount point %s"), *PakFile.GetMountPoint());
 
 		TArray<FPakFile::FPakEntryIterator> Records;
 
@@ -3236,9 +3278,19 @@ bool ListFilesInPak(const TCHAR * InPakFilename, int64 SizeFilter, bool bInclude
 
 		if (CSVFilename.Len() > 0)
 		{
-			TArray<FString> Lines;
-			Lines.Empty(Records.Num()+2);
-			Lines.Add(TEXT("Filename, Offset, Size, Hash, Deleted, Compressed, CompressionMethod"));
+			TRACE_CPUPROFILER_EVENT_SCOPE(WriteCsv);
+			uint32 WriteFlags = bAppendFile ? FILEWRITE_Append : 0;
+			TUniquePtr<FArchive> OutputArchive(IFileManager::Get().CreateFileWriter(*CSVFilename, WriteFlags));
+			if (!OutputArchive.IsValid())
+			{
+				UE_LOG(LogPakFile, Display, TEXT("Failed to save CSV file %s"), *CSVFilename);
+				return false;
+			}
+
+			if (!bAppendFile)
+			{
+				OutputArchive->Logf(TEXT("Filename, Offset, Size, Hash, Deleted, Compressed, CompressionMethod"));
+			}
 			EntryIndex = 0;
 			for (auto It : Records)
 			{
@@ -3246,25 +3298,18 @@ bool ListFilesInPak(const TCHAR * InPakFilename, int64 SizeFilter, bool bInclude
 				uint8* Hash = Hashes + (EntryIndex++)*sizeof(FPakEntry::Hash);
 
 				bool bWasCompressed = Entry.CompressionMethodIndex != 0;
-
-				Lines.Add( FString::Printf(
+				OutputArchive->Logf(
 					TEXT("%s%s, %lld, %lld, %s, %s, %s, %d"),
 					*MountPoint, It.TryGetFilename() ? **It.TryGetFilename() : TEXT("<FileNamesNotLoaded>"),
 					Entry.Offset, Entry.Size,
 					*BytesToHex(Hash, sizeof(FPakEntry::Hash)),
 					Entry.IsDeleteRecord() ? TEXT("true") : TEXT("false"),
 					bWasCompressed ? TEXT("true") : TEXT("false"),
-					Entry.CompressionMethodIndex) );
-			}
-
-			if (FFileHelper::SaveStringArrayToFile(Lines, *CSVFilename) == false)
-			{
-				UE_LOG(LogPakFile, Display, TEXT("Failed to save CSV file %s"), *CSVFilename);
-			}
-			else
-			{
-				UE_LOG(LogPakFile, Display, TEXT("Saved CSV file to %s"), *CSVFilename);
+					Entry.CompressionMethodIndex);
 			}
+			OutputArchive->Flush();
+			UE_LOG(LogPakFile, Display, TEXT("Saved CSV file to %s"), *CSVFilename);
+			return true;
 		}
 
 		TSet<int32> InspectChunks;
@@ -5520,8 +5565,9 @@ bool ExecuteUnrealPak(const TCHAR* CmdLine)
 		FParse::Value(CmdLine, TEXT("csv="), CSVFilename);
 
 		bool bExtractToMountPoint = FParse::Param(CmdLine, TEXT("ExtractToMountPoint"));
+		bool bAppendFile = false;
 
-		return ListFilesInPak(*PakFilename, SizeFilter, !bExcludeDeleted, *CSVFilename, bExtractToMountPoint, KeyChain);
+		return ListFilesInPak(*PakFilename, SizeFilter, !bExcludeDeleted, *CSVFilename, bExtractToMountPoint, KeyChain, bAppendFile);
 	}
 
 	if (FParse::Param(CmdLine, TEXT("Info")))
diff --git a/Engine/Source/Runtime/Core/Private/IO/IoStore.cpp b/Engine/Source/Runtime/Core/Private/IO/IoStore.cpp
--- a/Engine/Source/Runtime/Core/Private/IO/IoStore.cpp
+++ b/Engine/Source/Runtime/Core/Private/IO/IoStore.cpp
@@ -631,6 +631,8 @@ public:
 		int32 FirstBlockIndex = int32(ChunkInfo.Offset / CompressionBlockSize);
 		int32 LastBlockIndex = int32((Align(ChunkInfo.Offset + ChunkInfo.Size, CompressionBlockSize) - 1) / CompressionBlockSize);
 
+		ChunkInfo.NumCompressedBlocks = LastBlockIndex - FirstBlockIndex + 1;
+		ChunkInfo.OffsetOnDisk = TocResource.CompressionBlocks[FirstBlockIndex].GetOffset();
 		ChunkInfo.CompressedSize = 0;
 		ChunkInfo.PartitionIndex = -1;
 		for (int32 BlockIndex = FirstBlockIndex; BlockIndex <= LastBlockIndex; ++BlockIndex)
@@ -670,14 +672,14 @@ public:
 		bVerifyHashDatabase = bInVerifyHashDatabase;
 	}
 
-	void EnumerateChunks(TFunction<bool(const FIoStoreTocChunkInfo&)>&& Callback) const
+	void EnumerateChunks(TFunction<bool(FIoStoreTocChunkInfo&&)>&& Callback) const
 	{
 		const FIoStoreTocResource& TocResource = Toc.GetTocResource();
 
 		for (int32 ChunkIndex = 0; ChunkIndex < TocResource.ChunkIds.Num(); ++ChunkIndex)
 		{
 			FIoStoreTocChunkInfo ChunkInfo = Toc.GetTocChunkInfo(ChunkIndex, 0);
-			if (!Callback(ChunkInfo))
+			if (!Callback(MoveTemp(ChunkInfo)))
 			{
 				break;
 			}
@@ -694,29 +696,9 @@ public:
 		IPlatformFile& Ipf = IPlatformFile::GetPlatformPhysical();
 		Ipf.CreateDirectoryTree(*FPaths::GetPath(TocFilePath));
 
-		FIoStatus Status = FIoStatus::Ok;
-		if (InContext.GetSettings().bEnableCsvOutput)
-		{
-			Status = EnableCsvOutput();
-		}
-
 		FPartition& Partition = Partitions.AddDefaulted_GetRef();
 		Partition.Index = 0;
 
-		return Status;
-	}
-
-	FIoStatus EnableCsvOutput()
-	{
-		FString CsvFilePath = ContainerPath + TEXT(".csv");
-		CsvArchive.Reset(IFileManager::Get().CreateFileWriter(*CsvFilePath));
-		if (!CsvArchive)
-		{
-			return FIoStatusBuilder(EIoErrorCode::FileOpenFailed) << TEXT("Failed to open IoStore CSV file '") << *CsvFilePath << TEXT("'");
-		}
-		ANSICHAR Header[] = "Name,Offset,Size\n";
-		CsvArchive->Serialize(Header, sizeof(Header) - 1);
-
 		return FIoStatus::Ok;
 	}
 
@@ -2164,14 +2146,14 @@ public:
 		return Toc.GetTocResource().Header.EncryptionKeyGuid;
 	}
 
-	void EnumerateChunks(TFunction<bool(const FIoStoreTocChunkInfo&)>&& Callback) const
+	void EnumerateChunks(TFunction<bool(FIoStoreTocChunkInfo&&)>&& Callback) const
 	{
 		const FIoStoreTocResource& TocResource = Toc.GetTocResource();
 
 		for (int32 ChunkIndex = 0; ChunkIndex < TocResource.ChunkIds.Num(); ++ChunkIndex)
 		{
 			FIoStoreTocChunkInfo ChunkInfo = Toc.GetTocChunkInfo(ChunkIndex, &ChunkFileNamesMap);
-			if (!Callback(ChunkInfo))
+			if (!Callback(MoveTemp(ChunkInfo)))
 			{
 				break;
 			}
@@ -2620,7 +2602,7 @@ public:
 		return Toc.GetTocResource().CompressionMethods;
 	}
 
-	bool EnumerateCompressedBlocksForChunk(const FIoChunkId& ChunkId, TFunction<bool(const FIoStoreTocCompressedBlockInfo&)>&& Callback)
+	bool EnumerateCompressedBlocksForChunk(const FIoChunkId& ChunkId, TFunction<bool(const FIoStoreTocCompressedBlockInfo&)>&& Callback) const
 	{
 		const FIoOffsetAndLength* OffsetAndLength = Toc.GetOffsetAndLength(ChunkId);
 		if (!OffsetAndLength)
@@ -2652,7 +2634,7 @@ public:
 		return true;
 	}
 
-	void EnumerateCompressedBlocks(TFunction<bool(const FIoStoreTocCompressedBlockInfo&)>&& Callback)
+	void EnumerateCompressedBlocks(TFunction<bool(const FIoStoreTocCompressedBlockInfo&)>&& Callback) const
 	{
 		const FIoStoreTocResource& TocResource = Toc.GetTocResource();
 
@@ -2719,7 +2701,7 @@ FGuid FIoStoreReader::GetEncryptionKeyGuid() const
 	return Impl->GetEncryptionKeyGuid();
 }
 
-void FIoStoreReader::EnumerateChunks(TFunction<bool(const FIoStoreTocChunkInfo&)>&& Callback) const
+void FIoStoreReader::EnumerateChunks(TFunction<bool(FIoStoreTocChunkInfo&&)>&& Callback) const
 {
 	Impl->EnumerateChunks(MoveTemp(Callback));
 }
diff --git a/Engine/Source/Runtime/Core/Public/IO/IoDispatcher.h b/Engine/Source/Runtime/Core/Public/IO/IoDispatcher.h
--- a/Engine/Source/Runtime/Core/Public/IO/IoDispatcher.h
+++ b/Engine/Source/Runtime/Core/Public/IO/IoDispatcher.h
@@ -967,7 +967,6 @@ struct FIoStoreWriterSettings
 	int32 CompressionMinSizeToConsiderDDC = 0;
 	uint64 MemoryMappingAlignment = 0;
 	uint64 MaxPartitionSize = 0;
-	bool bEnableCsvOutput = false;
 	bool bEnableFileRegions = false;
 	bool bCompressionEnableDDC = false;
 };
@@ -1084,8 +1083,10 @@ struct FIoStoreTocChunkInfo
 	FString FileName;
 	FIoChunkHash Hash;
 	uint64 Offset;
+	uint64 OffsetOnDisk;
 	uint64 Size;
 	uint64 CompressedSize;
+	uint32 NumCompressedBlocks;
 	int32 PartitionIndex;
 	EIoChunkType ChunkType;
 	bool bHasValidFileName;
@@ -1184,7 +1185,7 @@ public:
 	CORE_API virtual void Append(const FIoChunkId& ChunkId, FIoBuffer Chunk, const FIoWriteOptions& WriteOptions, uint64 OrderHint = MAX_uint64) = 0;
 	CORE_API virtual void Append(const FIoChunkId& ChunkId, IIoStoreWriteRequest* Request, const FIoWriteOptions& WriteOptions) = 0;
 	CORE_API virtual TIoStatusOr<FIoStoreWriterResult> GetResult() = 0;
-	CORE_API virtual void EnumerateChunks(TFunction<bool(const FIoStoreTocChunkInfo&)>&& Callback) const = 0;
+	CORE_API virtual void EnumerateChunks(TFunction<bool(FIoStoreTocChunkInfo&&)>&& Callback) const = 0;
 };
 
 class FIoStoreReader
@@ -1198,7 +1199,7 @@ public:
 	CORE_API uint32 GetVersion() const;
 	CORE_API EIoContainerFlags GetContainerFlags() const;
 	CORE_API FGuid GetEncryptionKeyGuid() const;
-	CORE_API void EnumerateChunks(TFunction<bool(const FIoStoreTocChunkInfo&)>&& Callback) const;
+	CORE_API void EnumerateChunks(TFunction<bool(FIoStoreTocChunkInfo&&)>&& Callback) const;
 	CORE_API TIoStatusOr<FIoStoreTocChunkInfo> GetChunkInfo(const FIoChunkId& Chunk) const;
 	CORE_API TIoStatusOr<FIoStoreTocChunkInfo> GetChunkInfo(const uint32 TocEntryIndex) const;
 
