IoStore: Fix compile error due to interface change coming from cl 24117138

diff --git a/Engine/Source/Developer/IoStoreUtilities/Private/IoStoreOnDemand.cpp b/Engine/Source/Developer/IoStoreUtilities/Private/IoStoreOnDemand.cpp
--- a/Engine/Source/Developer/IoStoreUtilities/Private/IoStoreOnDemand.cpp
+++ b/Engine/Source/Developer/IoStoreUtilities/Private/IoStoreOnDemand.cpp
@@ -231,7 +231,7 @@ public:
 		return WriteResult;
 	}
 
-	virtual void EnumerateChunks(TFunction<bool(const FIoStoreTocChunkInfo&)>&& Callback) const override
+	virtual void EnumerateChunks(TFunction<bool(FIoStoreTocChunkInfo&&)>&& Callback) const override
 	{
 	}
 
