Adding a bunch of iostore/unrealpak performance stats and fixing a merge error with the reference db on optional containers.

diff --git a/Engine/Source/Developer/IoStoreUtilities/Private/IoStoreUtilities.cpp b/Engine/Source/Developer/IoStoreUtilities/Private/IoStoreUtilities.cpp
--- a/Engine/Source/Developer/IoStoreUtilities/Private/IoStoreUtilities.cpp
+++ b/Engine/Source/Developer/IoStoreUtilities/Private/IoStoreUtilities.cpp
@@ -77,6 +77,9 @@ IMPLEMENT_MODULE(FDefaultModuleImpl, IoStoreUtilities);
 
 TRACE_DECLARE_MEMORY_COUNTER(IoStoreUsedFileBufferMemory, TEXT("IoStore/UsedFileBufferMemory"));
 
+// Helper to format numbers with comma separators to help readability: 1,234 vs 1234.
+static FString NumberString(uint64 N) { return FText::AsNumber(N).ToString(); }
+
 static const FName DefaultCompressionMethod = NAME_Zlib;
 static const uint64 DefaultCompressionBlockSize = 64 << 10;
 static const uint64 DefaultCompressionBlockAlignment = 64 << 10;
@@ -100,6 +103,7 @@ public:
 	int32 FulfillCount = 0;
 	int32 ContainerNotFound = 0;
 	int64 FulfillBytes = 0;
+	int64 FulfillBytesPerChunk[(int8)EIoChunkType::MAX] = {};
 	
 
 	bool Init(const FString& InGlobalContainerFileName, const FKeyChain& InDecryptionKeychain)
@@ -194,6 +198,7 @@ public:
 			return false;
 		}
 
+		FulfillBytesPerChunk[(int8)ChunkId->GetChunkType()] += TotalCompressedSize;
 		FulfillBytes += TotalCompressedSize;
 		FulfillCount++;
 
@@ -1523,13 +1528,17 @@ static void ParsePackageAssetsFromFiles(TArray<FLegacyCookedPackage*>& Packages,
 			OptionalSegmentUAssetMemoryPtr += Packages[Index]->OptionalSegmentUAssetSize;
 		}
 
+		double StartTime = FPlatformTime::Seconds();
+
+		TAtomic<uint64> TotalReadCount{ 0 };
 		TAtomic<uint64> CurrentFileIndex{ 0 };
-		ParallelFor(TotalPackageCount, [&ReadCount, &PackageAssetBuffers, &OptionalSegmentPackageAssetBuffers, &Packages, &CurrentFileIndex](int32 Index)
+		ParallelFor(TEXT("ReadingPackageAssets.PF"), TotalPackageCount, 1, [&ReadCount, &PackageAssetBuffers, &OptionalSegmentPackageAssetBuffers, &Packages, &CurrentFileIndex, &TotalReadCount](int32 Index)
 		{
 			TRACE_CPUPROFILER_EVENT_SCOPE(ReadUAssetFile);
 			FLegacyCookedPackage* Package = Packages[Index];
 			if (Package->UAssetSize)
 			{
+				TotalReadCount.IncrementExchange();
 				uint8* Buffer = PackageAssetBuffers[Index];
 				TUniquePtr<IFileHandle> FileHandle(FPlatformFileManager::Get().GetPlatformFile().OpenRead(*Package->FileName));
 				if (FileHandle)
@@ -1544,6 +1553,7 @@ static void ParsePackageAssetsFromFiles(TArray<FLegacyCookedPackage*>& Packages,
 			}
 			if (Package->OptionalSegmentUAssetSize)
 			{
+				TotalReadCount.IncrementExchange();
 				uint8* Buffer = OptionalSegmentPackageAssetBuffers[Index];
 				TUniquePtr<IFileHandle> FileHandle(FPlatformFileManager::Get().GetPlatformFile().OpenRead(*Package->OptionalSegmentFileName));
 				if (FileHandle)
@@ -1560,6 +1570,13 @@ static void ParsePackageAssetsFromFiles(TArray<FLegacyCookedPackage*>& Packages,
 			uint64 LocalFileIndex = CurrentFileIndex.IncrementExchange() + 1;
 			UE_CLOG(LocalFileIndex % 1000 == 0, LogIoStore, Display, TEXT("Reading %d/%d: '%s'"), LocalFileIndex, Packages.Num(), *Package->FileName);
 		}, EParallelForFlags::Unbalanced);
+
+		double EndTime = FPlatformTime::Seconds();
+		UE_LOG(LogIoStore, Display, TEXT("Packages read %s files in %.2f seconds, %s total bytes, %s bytes per second"), 
+			*NumberString(TotalReadCount.Load()), 
+			EndTime - StartTime, 
+			*NumberString(TotalOptionalSegmentUAssetSize + TotalUAssetSize),
+			*NumberString((int64)((TotalOptionalSegmentUAssetSize + TotalUAssetSize) / FMath::Max(.001f, EndTime - StartTime))));
 	}
 
 	{
@@ -1941,6 +1958,8 @@ static void ProcessShaderLibraries(const FIoStoreArguments& Arguments, TArray<FC
 {
 	IOSTORE_CPU_SCOPE(ProcessShaderLibraries);
 
+	double LibraryStart = FPlatformTime::Seconds();
+
 	TMap<FIoChunkId, FShaderInfo*> ChunkIdToShaderInfoMap;
 	TMap<FSHAHash, TArray<FIoChunkId>> ShaderChunkIdsByShaderMapHash;
 	TMap<FName, TSet<FSHAHash>> PackageNameToShaderMaps;
@@ -2125,6 +2144,9 @@ static void ProcessShaderLibraries(const FIoStoreArguments& Arguments, TArray<FC
 			}
 		}
 	}
+
+	double LibraryEnd = FPlatformTime::Seconds();
+	UE_LOG(LogIoStore, Display, TEXT("Shaders processed in %.2f seconds"), LibraryEnd - LibraryStart);
 }
 
 void InitializeContainerTargetsAndPackages(
@@ -2755,7 +2777,9 @@ private:
 
 		virtual void LoadSourceBufferAsync() override
 		{
+			Manager.MemorySourceReads[(int8)TargetFile.ChunkId.GetChunkType()].IncrementExchange();
 			SourceBuffer = TargetFile.SourceBuffer.GetValue();
+			Manager.MemorySourceBytes[(int8)TargetFile.ChunkId.GetChunkType()] += SourceBuffer.DataSize();
 			OnSourceBufferLoaded();
 		}
 	};
@@ -2770,6 +2794,7 @@ private:
 
 		virtual void LoadSourceBufferAsync() override
 		{
+			Manager.LooseFileSourceReads[(int8)TargetFile.ChunkId.GetChunkType()].IncrementExchange();
 			SourceBuffer = FIoBuffer(GetSourceBufferSize());
 
 			QueueEntry->FileHandle.Reset(
@@ -2778,6 +2803,8 @@ private:
 			QueueEntry->AddRef(); // Must keep it around until we've assigned the ReadRequest pointer
 			FAsyncFileCallBack Callback = [this](bool, IAsyncReadRequest* ReadRequest)
 			{
+				Manager.LooseFileSourceBytes[(int8)TargetFile.ChunkId.GetChunkType()] += SourceBuffer.DataSize();
+
 				if (TargetFile.ChunkType == EContainerChunkType::PackageData)
 				{
 					SourceBuffer = Manager.PackageStoreOptimizer.CreatePackageBuffer(TargetFile.Package->OptimizedPackage, SourceBuffer, bHasUpdatedExportBundleRegions ? nullptr : &FileRegions);
@@ -2807,11 +2834,13 @@ private:
 
 		virtual void LoadSourceBufferAsync() override
 		{
+			Manager.ZenSourceReads[(int8)TargetFile.ChunkId.GetChunkType()].IncrementExchange();
 			Manager.PackageStore->ReadChunkAsync(
 				TargetFile.ChunkId,
 				[this](TIoStatusOr<FIoBuffer> Status)
 				{
 					SourceBuffer = Status.ConsumeValueOrDie();
+					Manager.ZenSourceBytes[(int8)TargetFile.ChunkId.GetChunkType()] += SourceBuffer.DataSize();
 					if (TargetFile.ChunkType == EContainerChunkType::PackageData)
 					{
 						check(TargetFile.Package->UAssetSize > 0);
@@ -3016,6 +3045,14 @@ private:
 	TAtomic<uint64> UsedBufferMemory { 0 };
 	FEvent* MemoryAvailableEvent;
 
+public:
+	TAtomic<uint64> ZenSourceReads[(int8)EIoChunkType::MAX] { 0 };
+	TAtomic<uint64> ZenSourceBytes[(int8)EIoChunkType::MAX]{ 0 };
+	TAtomic<uint64> MemorySourceReads[(int8)EIoChunkType::MAX]{ 0 };
+	TAtomic<uint64> MemorySourceBytes[(int8)EIoChunkType::MAX]{ 0 };
+	TAtomic<uint64> LooseFileSourceReads[(int8)EIoChunkType::MAX]{ 0 };
+	TAtomic<uint64> LooseFileSourceBytes[(int8)EIoChunkType::MAX]{ 0 };
+
 	static constexpr uint64 BufferMemoryLimit = 2ull << 30;
 };
 
@@ -3471,7 +3508,7 @@ int32 CreateTarget(const FIoStoreArguments& Arguments, const FIoStoreWriterSetti
 					if (!ContainerTarget->OptionalSegmentOutputPath.IsEmpty())
 					{
 						ContainerTarget->OptionalSegmentIoStoreWriter = IoStoreWriterContext->CreateContainer(*ContainerTarget->OptionalSegmentOutputPath, ContainerSettings);
-						ContainerTarget->OptionalSegmentIoStoreWriter->SetReferenceChunkDatabase(ChunkDatabase);
+					    ContainerTarget->OptionalSegmentIoStoreWriter->SetReferenceChunkDatabase(ChunkDatabase);
 						IoStoreWriters.Add(ContainerTarget->OptionalSegmentIoStoreWriter);
 					}
 				}
@@ -3689,10 +3726,95 @@ int32 CreateTarget(const FIoStoreArguments& Arguments, const FIoStoreWriterSetti
 			}
 			else
 			{
-				UE_LOG(LogIoStore, Display, TEXT("Hashing chunks (%llu/%llu) %llu from hash database..."), Progress.HashedChunksCount, Progress.TotalChunksCount, Progress.HashDbChunksCount);
+			UE_LOG(LogIoStore, Display, TEXT("Hashing chunks (%llu/%llu)..."), Progress.HashedChunksCount, Progress.TotalChunksCount);
 			}
 		}
 	}
+
+	{
+		FIoStoreWriterContext::FProgress Progress = IoStoreWriterContext->GetProgress();
+		if (Progress.HashDbChunksCount)
+		{
+			UE_LOG(LogIoStore, Display, TEXT("%s / %s hashes were loaded from the hash database, by type:"), *NumberString(Progress.HashDbChunksCount), *NumberString(Progress.TotalChunksCount));
+		
+			for (uint8 i = 0; i < (uint8)EIoChunkType::MAX; i++)
+			{
+				if (Progress.HashDbChunksByType[i])
+				{
+					UE_LOG(LogIoStore, Display, TEXT("    %-26s %s"), *LexToString((EIoChunkType)i), *NumberString(Progress.HashDbChunksByType[i]));
+				}
+			}
+		}
+		if (Progress.RefDbChunksCount)
+		{
+			FIoStoreChunkDatabase& TypedChunkDatabase = ((FIoStoreChunkDatabase&)*ChunkDatabase);
+			UE_LOG(LogIoStore, Display, TEXT("%s / %s chunks for %s bytes were loaded from the reference chunk database, by type:"), *NumberString(Progress.RefDbChunksCount), *NumberString(Progress.TotalChunksCount), *NumberString(TypedChunkDatabase.FulfillBytes));
+
+			for (uint8 i = 0; i < (uint8)EIoChunkType::MAX; i++)
+			{
+				if (Progress.RefDbChunksByType[i])
+				{
+					UE_LOG(LogIoStore, Display, TEXT("    %-26s: %s for %s bytes"), *LexToString((EIoChunkType)i), *NumberString(Progress.RefDbChunksByType[i]), *NumberString(TypedChunkDatabase.FulfillBytesPerChunk[i]));
+				}
+			}
+		}
+		if (Progress.CompressedChunksCount)
+		{
+			UE_LOG(LogIoStore, Display, TEXT("%s / %s chunks attempted to compress, by type:"), *NumberString(Progress.CompressedChunksCount), *NumberString(Progress.TotalChunksCount));
+
+			for (uint8 i = 0; i < (uint8)EIoChunkType::MAX; i++)
+			{
+				if (Progress.CompressedChunksByType[i] || Progress.BeginCompressChunksByType[i])
+				{
+					UE_LOG(LogIoStore, Display, TEXT("    %-26s %s / %s"), *LexToString((EIoChunkType)i), *NumberString(Progress.CompressedChunksByType[i]), *NumberString(Progress.BeginCompressChunksByType[i]));
+				}
+			}
+		}
+		UE_LOG(LogIoStore, Display, TEXT("Source bytes read:"));
+		uint64 ZenTotalBytes = 0;
+		for (uint64 b : WriteRequestManager.ZenSourceBytes)
+		{
+			ZenTotalBytes += b;
+		}
+
+		UE_LOG(LogIoStore, Display, TEXT("    Zen: %34s"), *NumberString(ZenTotalBytes));
+		for (uint8 i = 0; i < (uint8)EIoChunkType::MAX; i++)
+		{
+			if (WriteRequestManager.ZenSourceReads[i])
+			{
+				UE_LOG(LogIoStore, Display, TEXT("        %-22s %12s bytes over %s reads"), *LexToString((EIoChunkType)i), *NumberString(WriteRequestManager.ZenSourceBytes[i].Load()), *NumberString(WriteRequestManager.ZenSourceReads[i].Load()));
+			}
+		}
+
+		uint64 LooseTotalBytes = 0;
+		for (uint64 b : WriteRequestManager.LooseFileSourceBytes)
+		{
+			LooseTotalBytes += b;
+		}
+		UE_LOG(LogIoStore, Display, TEXT("    Loose File: %27s"), *NumberString(LooseTotalBytes));
+		for (uint8 i = 0; i < (uint8)EIoChunkType::MAX; i++)
+		{
+			if (WriteRequestManager.LooseFileSourceReads[i])
+			{
+				UE_LOG(LogIoStore, Display, TEXT("        %-22s %12s bytes over %s reads"), *LexToString((EIoChunkType)i), *NumberString(WriteRequestManager.LooseFileSourceBytes[i].Load()), *NumberString(WriteRequestManager.LooseFileSourceReads[i].Load()));
+			}
+		}
+
+		uint64 MemoryTotalBytes = 0;
+		for (uint64 b : WriteRequestManager.MemorySourceBytes)
+		{
+			MemoryTotalBytes += b;
+		}
+		UE_LOG(LogIoStore, Display, TEXT("    Memory: %31s"), *NumberString(MemoryTotalBytes));
+		for (uint8 i = 0; i < (uint8)EIoChunkType::MAX; i++)
+		{
+			if (WriteRequestManager.MemorySourceReads[i])
+			{
+				UE_LOG(LogIoStore, Display, TEXT("        %-22s %12s bytes over %s reads"), *LexToString((EIoChunkType)i), *NumberString(WriteRequestManager.MemorySourceBytes[i].Load()), *NumberString(WriteRequestManager.MemorySourceReads[i].Load()));
+			}
+		}
+	}
+
 	if (GeneralIoWriterSettings.bCompressionEnableDDC)
 	{
 		FIoStoreWriterContext::FProgress Progress = IoStoreWriterContext->GetProgress();
diff --git a/Engine/Source/Runtime/Core/Private/IO/IoStore.cpp b/Engine/Source/Runtime/Core/Private/IO/IoStore.cpp
--- a/Engine/Source/Runtime/Core/Private/IO/IoStore.cpp
+++ b/Engine/Source/Runtime/Core/Private/IO/IoStore.cpp
@@ -352,13 +352,32 @@ public:
 	{
 		FIoStoreWriterContext::FProgress Progress;
 		Progress.HashDbChunksCount = HashDbChunksCount.Load();
+		for (uint8 i=0; i<(uint8)EIoChunkType::MAX; i++)
+		{
+			Progress.HashDbChunksByType[i] = HashDbChunksByType[i].Load();
+		}
+
 		Progress.TotalChunksCount = TotalChunksCount.Load();
 		Progress.HashedChunksCount = HashedChunksCount.Load();
 		Progress.CompressedChunksCount = CompressedChunksCount.Load();
+		for (uint8 i = 0; i < (uint8)EIoChunkType::MAX; i++)
+		{
+			Progress.CompressedChunksByType[i] = CompressedChunksByType[i].Load();
+		}
+		for (uint8 i = 0; i < (uint8)EIoChunkType::MAX; i++)
+		{
+			Progress.BeginCompressChunksByType[i] = BeginCompressChunksByType[i].Load();
+		}
 		Progress.SerializedChunksCount = SerializedChunksCount.Load();
 		Progress.ScheduledCompressionTasksCount = ScheduledCompressionTasksCount.Load();
 		Progress.CompressionDDCHitCount = CompressionDDCHitCount.Load();
 		Progress.CompressionDDCMissCount = CompressionDDCMissCount.Load();
+		Progress.RefDbChunksCount = RefDbChunksCount.Load();
+		for (uint8 i = 0; i < (uint8)EIoChunkType::MAX; i++)
+		{
+			Progress.RefDbChunksByType[i] = RefDbChunksByType[i].Load();
+		}
+
 		return Progress;
 	}
 
@@ -443,8 +462,15 @@ private:
 	TAtomic<uint64> TotalChunksCount{ 0 };
 	TAtomic<uint64> HashedChunksCount{ 0 };
 	TAtomic<uint64> HashDbChunksCount{ 0 };
+	TAtomic<uint64> HashDbChunksByType[(int8)EIoChunkType::MAX] = {0};
+	TAtomic<uint64> RefDbChunksCount{ 0 };
+	TAtomic<uint64> RefDbChunksByType[(int8)EIoChunkType::MAX] = { 0 };
 	TAtomic<uint64> CompressedChunksCount{ 0 };
+	TAtomic<uint64> CompressedChunksByType[(int8)EIoChunkType::MAX] = { 0 };
+	TAtomic<uint64> BeginCompressChunksByType[(int8)EIoChunkType::MAX] = { 0 };
 	TAtomic<uint64> SerializedChunksCount{ 0 };
+	TAtomic<uint64> WriteCycleCount{ 0 };
+	TAtomic<uint64> WriteByteCount{ 0 };
 	TAtomic<uint64> ScheduledCompressionTasksCount{ 0 };
 	TAtomic<uint64> CompressionDDCHitCount{ 0 };
 	TAtomic<uint64> CompressionDDCMissCount{ 0 };
@@ -791,6 +817,7 @@ public:
 			{
 				// If we aren't validating then we just use it and bail.
 				WriterContext->HashDbChunksCount.IncrementExchange();
+				WriterContext->HashDbChunksByType[(int8)Entry->ChunkId.GetChunkType()].IncrementExchange();
 				WriterContext->HashedChunksCount.IncrementExchange();
 				return;
 			}
@@ -1439,6 +1466,8 @@ private:
 
 	void BeginCompress(FIoStoreWriteQueueEntry* Entry)
 	{
+		WriterContext->BeginCompressChunksByType[(int8)Entry->ChunkId.GetChunkType()].IncrementExchange();
+
 		FName CompressionMethod = NAME_None;
 		const FIoStoreWriterSettings& WriterSettings = WriterContext->WriterSettings;
 		if (ContainerSettings.IsCompressed() && !Entry->Options.bForceUncompressed && !Entry->Options.bIsMemoryMapped)
@@ -1519,6 +1548,8 @@ private:
 
 			if (GotChunk)
 			{
+				WriterContext->RefDbChunksCount.IncrementExchange();
+				WriterContext->RefDbChunksByType[(int8)Entry->ChunkId.GetChunkType()].IncrementExchange();
 				// Lambda handles dispatch subsequents
 				return;
 			}
@@ -1581,6 +1612,7 @@ private:
 				int32 CompressedBlocksCount = Entry->CompressedBlocksCount.IncrementExchange();
 				if (CompressedBlocksCount + 1 == Entry->ChunkBlocks.Num())
 				{
+					WriterContext->CompressedChunksByType[(int8)Entry->ChunkId.GetChunkType()].IncrementExchange();
 					WriterContext->CompressedChunksCount.IncrementExchange();
 					Entry->FinishCompressionBarrier->DispatchSubsequents();
 				}
@@ -1766,7 +1798,8 @@ private:
 		}
 		delete Entry->Request;
 		Entry->Request = nullptr;
-
+		uint64 WriteStartCycles = FPlatformTime::Cycles64();
+		uint64 WriteBytes = 0;
 		if (Entry->Padding > 0)
 		{
 			if (PaddingBuffer.Num() < Entry->Padding)
@@ -1776,6 +1809,7 @@ private:
 			{
 				TRACE_CPUPROFILER_EVENT_SCOPE(WritePaddingToContainer);
 				TargetPartition->ContainerFileHandle->Serialize(PaddingBuffer.GetData(), Entry->Padding);
+				WriteBytes += Entry->Padding;
 			}
 		}
 		check(Entry->Offset == TargetPartition->ContainerFileHandle->Tell());
@@ -1784,10 +1818,14 @@ private:
 			{
 				TRACE_CPUPROFILER_EVENT_SCOPE(WriteBlockToContainer);
 				TargetPartition->ContainerFileHandle->Serialize(ChunkBlock.IoBuffer->Data(), ChunkBlock.Size);
+				WriteBytes += ChunkBlock.Size;
 			}
 			WriterContext->FreeCompressionBuffer(ChunkBlock.IoBuffer, Entry->ChunkBlocks.Num());
 			ChunkBlock.IoBuffer = nullptr;
 		}
+		uint64 WriteEndCycles = FPlatformTime::Cycles64();
+		WriterContext->WriteCycleCount.AddExchange(WriteEndCycles - WriteStartCycles);
+		WriterContext->WriteByteCount.AddExchange(WriteBytes);
 		WriterContext->SerializedChunksCount.IncrementExchange();
 	}
 
@@ -1861,6 +1899,8 @@ void FIoStoreWriterContextImpl::Flush()
 		AllEntries.Append(IoStoreWriter->Entries);
 	}
 
+	double WritesStart = FPlatformTime::Seconds();
+
 	for (FIoStoreWriteQueueEntry* Entry : AllEntries)
 	{
 		ScheduleCompression(Entry);
@@ -1874,6 +1914,15 @@ void FIoStoreWriterContextImpl::Flush()
 		}
 	}
 
+	double WritesEnd = FPlatformTime::Seconds();
+	double WritesSeconds = FPlatformTime::ToSeconds64(WriteCycleCount.Load());
+	UE_LOG(LogIoStore, Display, TEXT("Writing and compressing took %.2f seconds, writes to disk took %.2f seconds for %s bytes @ %s bytes per second."), 
+		WritesEnd - WritesStart,
+		WritesSeconds,
+		*FText::AsNumber(WriteByteCount.Load()).ToString(),
+		*FText::AsNumber((int64)((double)WriteByteCount.Load() / WritesSeconds)).ToString()
+		);
+
 	AllEntries.Empty();
 
 	// Classically there were so few writers that this didn't need to be multi threaded, but it
@@ -1884,7 +1933,21 @@ void FIoStoreWriterContextImpl::Flush()
 		IoStoreWriters[Index]->Finalize(); 
 	});
 	double FinalizeEnd = FPlatformTime::Seconds();
-	UE_LOG(LogIoStore, Display, TEXT("Finalize took %.1f seconds for %d writers"), FinalizeEnd - FinalizeStart, IoStoreWriters.Num());
+	int64 TotalTocSize = 0;
+	for (TSharedPtr<IIoStoreWriter> Writer : IoStoreWriters)
+	{
+		if (Writer->GetResult().IsOk())
+		{
+			TotalTocSize += Writer->GetResult().ValueOrDie().TocSize;
+		}
+	}
+
+	UE_LOG(LogIoStore, Display, TEXT("Finalize took %.1f seconds for %d writers to write %s bytes, %s bytes per second"), 
+		FinalizeEnd - FinalizeStart, 
+		IoStoreWriters.Num(), 
+		*FText::AsNumber(TotalTocSize).ToString(), 
+		*FText::AsNumber((int64)((double)TotalTocSize / (FinalizeEnd - FinalizeStart))).ToString()
+		);
 }
 
 void FIoStoreWriterContextImpl::BeginCompressionThreadFunc()
diff --git a/Engine/Source/Runtime/Core/Public/IO/IoDispatcher.h b/Engine/Source/Runtime/Core/Public/IO/IoDispatcher.h
--- a/Engine/Source/Runtime/Core/Public/IO/IoDispatcher.h
+++ b/Engine/Source/Runtime/Core/Public/IO/IoDispatcher.h
@@ -1044,13 +1044,23 @@ public:
 	{
 		uint64 TotalChunksCount = 0;
 		uint64 HashedChunksCount = 0;
-		// Number of chunks where we avoided reading and hashing, and instead used the result from the hashdb
+		// Number of chunks where we avoided reading and hashing, and instead used the result from the hashdb, and their types
 		uint64 HashDbChunksCount = 0;
+		uint64 HashDbChunksByType[(int8)EIoChunkType::MAX] = { 0 };
+		// Number of chunks that were passed to the compressor (i.e. passed the various opt-outs), and their types
 		uint64 CompressedChunksCount = 0;
+		uint64 CompressedChunksByType[(int8)EIoChunkType::MAX] = { 0 };
 		uint64 SerializedChunksCount = 0;
 		uint64 ScheduledCompressionTasksCount = 0;
 		uint64 CompressionDDCHitCount = 0;
 		uint64 CompressionDDCMissCount = 0;
+
+		// The number of chunk retrieved from the reference cache database, and their types.
+		uint64 RefDbChunksCount{ 0 };
+		uint64 RefDbChunksByType[(int8)EIoChunkType::MAX] = { 0 };
+		
+		// The type of chunk that landed in BeginCompress before any opt-outs.
+		uint64 BeginCompressChunksByType[(int8)EIoChunkType::MAX] = { 0 };
 	};
 
 	CORE_API FIoStoreWriterContext();
