Missing asset registry is a display now so we don't get warning spam when we know it's missing.

diff --git a/Engine/Source/Developer/IoStoreUtilities/Private/IoStoreUtilities.cpp b/Engine/Source/Developer/IoStoreUtilities/Private/IoStoreUtilities.cpp
--- a/Engine/Source/Developer/IoStoreUtilities/Private/IoStoreUtilities.cpp
+++ b/Engine/Source/Developer/IoStoreUtilities/Private/IoStoreUtilities.cpp
@@ -3419,7 +3419,7 @@ int32 CreateTarget(const FIoStoreArguments& Arguments, const FIoStoreWriterSetti
 	TSharedPtr<IIoStoreWriterHashDatabase> HashDatabase = MakeShared<FIoStoreHashDb>();
 	if (((FIoStoreHashDb&)*HashDatabase).Initialize(Arguments.CookedDir) == false)
 	{
-		UE_LOG(LogIoStore, Warning, TEXT("Unabled to initialize the hash database from the asset registry!"));
+		UE_LOG(LogIoStore, Display, TEXT("Unabled to initialize the hash database from the asset registry!"));
 	}
 	if (Arguments.bVerifyHashDatabase)
 	{
