ZenCook: Added support for optional package segments

diff --git a/Engine/Source/Developer/IoStoreUtilities/Private/IoStoreUtilities.cpp b/Engine/Source/Developer/IoStoreUtilities/Private/IoStoreUtilities.cpp
--- a/Engine/Source/Developer/IoStoreUtilities/Private/IoStoreUtilities.cpp
+++ b/Engine/Source/Developer/IoStoreUtilities/Private/IoStoreUtilities.cpp
@@ -798,14 +798,14 @@ public:
 		(new FAutoDeleteAsyncTask<FReadChunkTask>(ZenStoreClient.Get(), ChunkId, MoveTemp(Callback)))->StartBackgroundTask();
 	}
 	
-	TIoStatusOr<FIoBuffer> ReadPackageHeaderFromZen(FPackageId PackageId)
+	TIoStatusOr<FIoBuffer> ReadPackageHeaderFromZen(FPackageId PackageId, uint16 ChunkIndex)
 	{
 		if (const FPackageStoreEntryResource* Entry = PackageIdToEntry.Find(PackageId))
 		{
 			FIoReadOptions ReadOptions;
 			ReadOptions.SetRange(0, 64 << 10);
 			
-			TIoStatusOr<FIoBuffer> Status = ZenStoreClient->ReadChunk(CreateIoChunkId(PackageId.Value(), 0, EIoChunkType::ExportBundleData), ReadOptions.GetOffset(), ReadOptions.GetSize());
+			TIoStatusOr<FIoBuffer> Status = ZenStoreClient->ReadChunk(CreateIoChunkId(PackageId.Value(), ChunkIndex, EIoChunkType::ExportBundleData), ReadOptions.GetOffset(), ReadOptions.GetSize());
 			if (!Status.IsOk())
 			{
 				return Status;
@@ -817,7 +817,7 @@ public:
 			{
 				ReadOptions.SetRange(0, HeaderSize);
 
-				Status = ZenStoreClient->ReadChunk(CreateIoChunkId(PackageId.Value(), 0, EIoChunkType::ExportBundleData), ReadOptions.GetOffset(), ReadOptions.GetSize());
+				Status = ZenStoreClient->ReadChunk(CreateIoChunkId(PackageId.Value(), ChunkIndex, EIoChunkType::ExportBundleData), ReadOptions.GetOffset(), ReadOptions.GetSize());
 				if (!Status.IsOk())
 				{
 					return Status;
@@ -1631,12 +1631,20 @@ static void ParsePackageAssetsFromZen(FCookedPackageStore& PackageStore, TArray<
 		&Packages](int32 Index)
 	{
 		FLegacyCookedPackage* Package = Packages[Index];
-		TIoStatusOr<FIoBuffer> HeaderBuffer = PackageStore.ReadPackageHeaderFromZen(Package->GlobalPackageId);
+		TIoStatusOr<FIoBuffer> HeaderBuffer = PackageStore.ReadPackageHeaderFromZen(Package->GlobalPackageId, 0);
 		Package->UExpSize = Package->UAssetSize - HeaderBuffer.ValueOrDie().DataSize();
 		Package->UAssetSize = HeaderBuffer.ValueOrDie().DataSize();
 		const FPackageStoreEntryResource* PackageStoreEntry = Package->PackageStoreEntry;
 		Package->OptimizedPackage = PackageStoreOptimizer.CreatePackageFromZenPackageHeader(Package->PackageName, HeaderBuffer.ValueOrDie(), PackageStoreEntry->ExportInfo.ExportBundleCount, PackageStoreEntry->ImportedPackageIds);
 		check(Package->OptimizedPackage->GetId() == Package->GlobalPackageId);
+		if (Package->OptionalSegmentUAssetSize)
+		{
+			TIoStatusOr<FIoBuffer> OptionalSegmentHeaderBuffer = PackageStore.ReadPackageHeaderFromZen(Package->GlobalPackageId, 1);
+			Package->OptionalSegmentUExpSize = Package->OptionalSegmentUAssetSize - OptionalSegmentHeaderBuffer.ValueOrDie().DataSize();
+			Package->OptionalSegmentUAssetSize = OptionalSegmentHeaderBuffer.ValueOrDie().DataSize();
+			Package->OptimizedOptionalSegmentPackage = PackageStoreOptimizer.CreatePackageFromZenPackageHeader(Package->PackageName, OptionalSegmentHeaderBuffer.ValueOrDie(), PackageStoreEntry->OptionalSegmentExportInfo.ExportBundleCount, PackageStoreEntry->OptionalSegmentImportedPackageIds);
+			check(Package->OptimizedOptionalSegmentPackage->GetId() == Package->GlobalPackageId);
+		}
 	}, EParallelForFlags::Unbalanced);
 }
 
diff --git a/Engine/Source/Developer/IoStoreUtilities/Private/ZenStoreWriter.cpp b/Engine/Source/Developer/IoStoreUtilities/Private/ZenStoreWriter.cpp
--- a/Engine/Source/Developer/IoStoreUtilities/Private/ZenStoreWriter.cpp
+++ b/Engine/Source/Developer/IoStoreUtilities/Private/ZenStoreWriter.cpp
@@ -60,6 +60,11 @@ FMD5Hash IoHashToMD5(const FIoHash& IoHash)
 	return Hash;
 }
 
+FZenStoreWriter::FPackageDataEntry::~FPackageDataEntry()
+{
+
+}
+
 struct FZenStoreWriter::FZenCommitInfo
 {
 	IPackageWriter::FCommitPackageInfo CommitInfo;
@@ -221,6 +226,9 @@ FZenStoreWriter::~FZenStoreWriter()
 void FZenStoreWriter::WritePackageData(const FPackageInfo& Info, FLargeMemoryWriter& ExportsArchive, const TArray<FFileRegion>& FileRegions)
 {
 	check(Info.ChunkId.IsValid());
+	FPendingPackageState& ExistingState = GetPendingPackage(Info.PackageName);
+	FPackageDataEntry& Entry = ExistingState.PackageData.AddDefaulted_GetRef();
+
 	int64 DataSize = ExportsArchive.TotalSize();
 	FIoBuffer PackageData(FIoBuffer::AssumeOwnership, ExportsArchive.ReleaseOwnership(), DataSize);
 
@@ -228,20 +236,20 @@ void FZenStoreWriter::WritePackageData(const FPackageInfo& Info, FLargeMemoryWri
 
 	FIoBuffer CookedHeaderBuffer = FIoBuffer(PackageData.Data(), Info.HeaderSize, PackageData);
 	FIoBuffer CookedExportsBuffer = FIoBuffer(PackageData.Data() + Info.HeaderSize, PackageData.DataSize() - Info.HeaderSize, PackageData);
-	TUniquePtr<FPackageStorePackage> Package{PackageStoreOptimizer->CreatePackageFromCookedHeader(Info.PackageName, CookedHeaderBuffer)};
-	PackageStoreOptimizer->FinalizePackage(Package.Get());
+	Entry.OptimizedPackage.Reset(PackageStoreOptimizer->CreatePackageFromCookedHeader(Info.PackageName, CookedHeaderBuffer));
+	PackageStoreOptimizer->FinalizePackage(Entry.OptimizedPackage.Get());
 	TArray<FFileRegion> FileRegionsCopy(FileRegions);
 	for (FFileRegion& Region : FileRegionsCopy)
 	{
 		// Adjust regions so they are relative to the start of the exports buffer
 		Region.Offset -= Info.HeaderSize;
 	}
-	FIoBuffer PackageBuffer = PackageStoreOptimizer->CreatePackageBuffer(Package.Get(), CookedExportsBuffer, &FileRegionsCopy);
+	FIoBuffer PackageBuffer = PackageStoreOptimizer->CreatePackageBuffer(Entry.OptimizedPackage.Get(), CookedExportsBuffer, &FileRegionsCopy);
 	PackageStoreManifest.AddPackageData(Info.PackageName, Info.LooseFilePath, Info.ChunkId);
 	for (FFileRegion& Region : FileRegionsCopy)
 	{
 		// Adjust regions once more so they are relative to the exports bundle buffer
-		Region.Offset -= Package->GetHeaderSize();
+		Region.Offset -= Entry.OptimizedPackage->GetHeaderSize();
 	}
 	//WriteFileRegions(*FPaths::ChangeExtension(Info.LooseFilePath, FString(".uexp") + FFileRegion::RegionsFileExtension), FileRegionsCopy);
 
@@ -249,10 +257,6 @@ void FZenStoreWriter::WritePackageData(const FPackageInfo& Info, FLargeMemoryWri
 
 	FCbObjectId ChunkOid = ToObjectId(Info.ChunkId);
 
-	FPendingPackageState& ExistingState = GetPendingPackage(Info.PackageName);
-
-	FPackageDataEntry& Entry = ExistingState.PackageData;
-
 	Entry.CompressedPayload = Async(EAsyncExecution::TaskGraph, [this, PackageBuffer]()
 	{ 
 		return FCompressedBuffer::Compress(FSharedBuffer::MakeView(PackageBuffer.GetView()), Compressor, CompressionLevel);
@@ -260,18 +264,7 @@ void FZenStoreWriter::WritePackageData(const FPackageInfo& Info, FLargeMemoryWri
 
 	Entry.Info				= Info;
 	Entry.ChunkId			= ChunkOid;
-	Entry.PackageStoreEntry = PackageStoreOptimizer->CreatePackageStoreEntry(Package.Get(), nullptr); // TODO: Can we separate out the optional segment package store entry and do this when we commit instead?
 	Entry.IsValid			= true;
-
-	if (EntryCreatedEvent.IsBound())
-	{
-		IPackageStoreWriter::FEntryCreatedEventArgs EntryCreatedEventArgs
-		{
-			TargetPlatformFName,
-			Entry.PackageStoreEntry
-		};
-		EntryCreatedEvent.Broadcast(EntryCreatedEventArgs);
-	}
 }
 
 void FZenStoreWriter::WriteIoStorePackageData(const FPackageInfo& Info, const FIoBuffer& PackageData, const FPackageStoreEntryResource& PackageStoreEntry, const TArray<FFileRegion>& FileRegions)
@@ -287,7 +280,7 @@ void FZenStoreWriter::WriteIoStorePackageData(const FPackageInfo& Info, const FI
 
 	FPendingPackageState& ExistingState = GetPendingPackage(Info.PackageName);
 
-	FPackageDataEntry& Entry = ExistingState.PackageData;
+	FPackageDataEntry& Entry = ExistingState.PackageData.AddDefaulted_GetRef();
 
 	PackageData.EnsureOwned();
 
@@ -298,7 +291,6 @@ void FZenStoreWriter::WriteIoStorePackageData(const FPackageInfo& Info, const FI
 
 	Entry.Info				= Info;
 	Entry.ChunkId			= ChunkOid;
-	Entry.PackageStoreEntry = PackageStoreEntry;
 	Entry.IsValid			= true;
 }
 
@@ -698,12 +690,30 @@ void FZenStoreWriter::CommitPackageInternal(FZenCommitInfo&& ZenCommitInfo)
 
 	if (CommitInfo.Status == ECommitStatus::Success && bWritePackage)
 	{
-		FPackageDataEntry& PkgData = PackageState->PackageData;
-		checkf(PkgData.IsValid, TEXT("CommitPackage called with bSucceeded but without first calling WritePackageData"))
+		checkf(!PackageState->PackageData.IsEmpty(), TEXT("CommitPackage called with bSucceeded but without first calling WritePackageData"))
 		checkf(EnumHasAllFlags(CommitInfo.WriteOptions, EWriteOptions::Write), TEXT("Partial EWriteOptions::Write options are not yet implemented."));
 		checkf(!EnumHasAnyFlags(CommitInfo.WriteOptions, EWriteOptions::SaveForDiff), TEXT("-diffonly -savefordiff is not yet implemented."));
 
+		FPackageStoreEntryResource PackageStoreEntry;
 		{
+			FPackageDataEntry* PkgData = nullptr;
+			FPackageDataEntry* OptionalSegmentPkgData = nullptr;
+			for (FPackageDataEntry& PackageDataEntry : PackageState->PackageData)
+			{
+				check(PackageDataEntry.Info.MultiOutputIndex <= 1);
+				if (PackageDataEntry.Info.MultiOutputIndex == 0)
+				{
+					check(!PkgData);
+					PkgData = &PackageDataEntry;
+				}
+				else if (PackageDataEntry.Info.MultiOutputIndex == 1)
+				{
+					check(!OptionalSegmentPkgData);
+					OptionalSegmentPkgData = &PackageDataEntry;
+				}
+			}
+			PackageStoreEntry = PackageStoreOptimizer->CreatePackageStoreEntry(PkgData->OptimizedPackage.Get(), OptionalSegmentPkgData ? OptionalSegmentPkgData->OptimizedPackage.Get() : nullptr);
+
 			FWriteScopeLock _(EntriesLock);
 			CommitEventArgs.EntryIndex = PackageNameToIndex.FindOrAdd(CommitInfo.PackageName, PackageStoreEntries.Num());
 			if (CommitEventArgs.EntryIndex == PackageStoreEntries.Num())
@@ -711,17 +721,21 @@ void FZenStoreWriter::CommitPackageInternal(FZenCommitInfo&& ZenCommitInfo)
 				PackageStoreEntries.Emplace();
 				CookedPackagesInfo.Emplace();
 			}
+			PackageStoreEntries[CommitEventArgs.EntryIndex] = PackageStoreEntry;
 		}
 		
-		PackageStoreEntries[CommitEventArgs.EntryIndex] = PkgData.PackageStoreEntry;
+		if (EntryCreatedEvent.IsBound())
+		{
+			IPackageStoreWriter::FEntryCreatedEventArgs EntryCreatedEventArgs
+			{
+				TargetPlatformFName,
+				PackageStoreEntry
+			};
+			EntryCreatedEvent.Broadcast(EntryCreatedEventArgs);
+		}
 		
 		FMD5 PkgHashGen;
 		FCbPackage OplogEntry;
-
-		FCbAttachment PkgDataAttachment = FCbAttachment(PkgData.CompressedPayload.Get());
-		PkgHashGen.Update(PkgDataAttachment.GetHash().GetBytes(), sizeof(FIoHash::ByteArray));
-		OplogEntry.AddAttachment(PkgDataAttachment);
-		
 		// Commit attachments
 		FOplogCookInfo& CookInfo = CookedPackagesInfo[CommitEventArgs.EntryIndex];
 		CookInfo.bUpToDate = true;
@@ -766,14 +780,29 @@ void FZenStoreWriter::CommitPackageInternal(FZenCommitInfo&& ZenCommitInfo)
 		FString PackageNameKey = CommitInfo.PackageName.ToString();
 		PackageNameKey.ToLowerInline();
 		OplogEntryDesc << "key" << PackageNameKey;
+		OplogEntryDesc << "packagestoreentry" << PackageStoreEntry;
+		
+		OplogEntryDesc.BeginArray("packagedata");
+		
+		for (FPackageDataEntry& PkgData : PackageState->PackageData)
+		{
+			if (bComputeHash)
+			{
+				PackageState->PackageHashes->ChunkHashes.Add(PkgData.Info.ChunkId, PkgData.CompressedPayload.Get().GetRawHash());
+			}
 
-		OplogEntryDesc.BeginObject("package");
-		OplogEntryDesc << "id" << PkgData.ChunkId;
-		OplogEntryDesc << "data" << PkgDataAttachment;
-		OplogEntryDesc.EndObject();
+			FCbAttachment PkgDataAttachment = FCbAttachment(PkgData.CompressedPayload.Get());
+			PkgHashGen.Update(PkgDataAttachment.GetHash().GetBytes(), sizeof(FIoHash::ByteArray));
+			OplogEntry.AddAttachment(PkgDataAttachment);
+
+			OplogEntryDesc.BeginObject();
+			OplogEntryDesc << "id" << PkgData.ChunkId;
+			OplogEntryDesc << "data" << PkgDataAttachment;
+			OplogEntryDesc.EndObject();
+		}
+
+		OplogEntryDesc.EndArray();
 
-		OplogEntryDesc << "packagestoreentry" << PkgData.PackageStoreEntry;
-		
 		if (PackageState->BulkData.Num())
 		{
 			OplogEntryDesc.BeginArray("bulkdata");
@@ -852,11 +881,11 @@ void FZenStoreWriter::CommitPackageInternal(FZenCommitInfo&& ZenCommitInfo)
 	}
 	else if (CommitInfo.Status == ECommitStatus::Success && bComputeHash)
 	{
-		FPackageDataEntry& PkgData = PackageState->PackageData;
-		checkf(PkgData.IsValid, TEXT("CommitPackage called with bSucceeded but without first calling WritePackageData"));
+		checkf(!PackageState->PackageData.IsEmpty(), TEXT("CommitPackage called with bSucceeded but without first calling WritePackageData"));
 		
 		FMD5 PkgHashGen;
 		
+		for (FPackageDataEntry& PkgData : PackageState->PackageData)
 		{
 			FCompressedBuffer Payload = PkgData.CompressedPayload.Get();
 			FIoHash IoHash = Payload.GetRawHash();
diff --git a/Engine/Source/Developer/IoStoreUtilities/Public/ZenStoreWriter.h b/Engine/Source/Developer/IoStoreUtilities/Public/ZenStoreWriter.h
--- a/Engine/Source/Developer/IoStoreUtilities/Public/ZenStoreWriter.h
+++ b/Engine/Source/Developer/IoStoreUtilities/Public/ZenStoreWriter.h
@@ -119,10 +119,12 @@ private:
 
 	struct FPackageDataEntry
 	{
+		~FPackageDataEntry();
+
 		TFuture<FCompressedBuffer> CompressedPayload;
 		FPackageInfo Info;
 		FCbObjectId ChunkId;
-		FPackageStoreEntryResource PackageStoreEntry;
+		TUniquePtr<class FPackageStorePackage> OptimizedPackage;
 		bool IsValid = false;
 	};
 
@@ -137,7 +139,7 @@ private:
 	struct FPendingPackageState
 	{
 		FName PackageName;
-		FPackageDataEntry PackageData;
+		TArray<FPackageDataEntry> PackageData;
 		TArray<FBulkDataEntry> BulkData;
 		TArray<FFileDataEntry> FileData;
 		TRefCountPtr<FPackageHashes> PackageHashes;
diff --git a/Engine/Source/Runtime/Core/Private/IO/PackageStore.cpp b/Engine/Source/Runtime/Core/Private/IO/PackageStore.cpp
--- a/Engine/Source/Runtime/Core/Private/IO/PackageStore.cpp
+++ b/Engine/Source/Runtime/Core/Private/IO/PackageStore.cpp
@@ -52,6 +52,8 @@ FArchive& operator<<(FArchive& Ar, FPackageStoreEntryResource& PackageStoreEntry
 	Ar << PackageStoreEntry.Region;
 	Ar << PackageStoreEntry.ExportInfo;
 	Ar << PackageStoreEntry.ImportedPackageIds;
+	Ar << PackageStoreEntry.OptionalSegmentExportInfo;
+	Ar << PackageStoreEntry.OptionalSegmentImportedPackageIds;
 
 	if (Ar.IsLoading())
 	{
@@ -92,6 +94,18 @@ FCbWriter& operator<<(FCbWriter& Writer, const FPackageStoreEntryResource& Packa
 		Writer.EndArray();
 	}
 
+	Writer << "optionalsegmentexportinfo" << PackageStoreEntry.OptionalSegmentExportInfo;
+
+	if (PackageStoreEntry.OptionalSegmentImportedPackageIds.Num())
+	{
+		Writer.BeginArray("optionalsegmentimportedpackageids");
+		for (const FPackageId& ImportedPackageId : PackageStoreEntry.OptionalSegmentImportedPackageIds)
+		{
+			Writer << ImportedPackageId.Value();
+		}
+		Writer.EndArray();
+	}
+
 	Writer.EndObject();
 
 	return Writer;
@@ -125,6 +139,16 @@ FPackageStoreEntryResource FPackageStoreEntryResource::FromCbObject(const FCbObj
 		}
 	}
 
+	Entry.OptionalSegmentExportInfo = FPackageStoreExportInfo::FromCbObject(Obj["optionalsegmentexportinfo"].AsObject());
+
+	if (Obj["optionalsegmentimportedpackageids"])
+	{
+		for (FCbFieldView ArrayField : Obj["optionalsegmentimportedpackageids"])
+		{
+			Entry.OptionalSegmentImportedPackageIds.Add(FPackageId::FromValue(ArrayField.AsUInt64()));
+		}
+	}
+
 	return Entry;
 }
 
