Mode in iostore utilities to check for invalid cross-container references.

diff --git a/Engine/Source/Developer/IoStoreUtilities/Private/IoStoreUtilities.cpp b/Engine/Source/Developer/IoStoreUtilities/Private/IoStoreUtilities.cpp
--- a/Engine/Source/Developer/IoStoreUtilities/Private/IoStoreUtilities.cpp
+++ b/Engine/Source/Developer/IoStoreUtilities/Private/IoStoreUtilities.cpp
@@ -13,12 +13,14 @@
 #include "Misc/App.h"
 #include "Misc/CommandLine.h"
 #include "Misc/ConfigCacheIni.h"
+#include "Misc/Optional.h"
 #include "Misc/PackageName.h"
 #include "Misc/Paths.h"
 #include "Misc/Base64.h"
 #include "Misc/AES.h"
 #include "Misc/CoreDelegates.h"
 #include "Misc/KeyChainUtilities.h"
+#include "Misc/WildcardString.h"
 #include "Modules/ModuleManager.h"
 #include "Serialization/Archive.h"
 #include "Serialization/JsonReader.h"
@@ -37,6 +39,7 @@
 #include "UObject/ObjectResource.h"
 #include "UObject/Package.h"
 #include "UObject/UObjectHash.h"
+#include "Algo/AnyOf.h"
 #include "Algo/Find.h"
 #include "Misc/FileHelper.h"
 #include "Misc/ScopeLock.h"
@@ -4635,12 +4638,7 @@ int32 ProfileReadSpeed(const TCHAR* InCommandLine, const FKeyChain& InKeyChain)
 	return 0;
 }
 
-int32 Describe(
-	const FString& GlobalContainerPath,
-	const FKeyChain& KeyChain,
-	const FString& PackageFilter,
-	const FString& OutPath,
-	bool bIncludeExportHashes)
+namespace DescribeUtils
 {
 	struct FPackageDesc;
 
@@ -4721,410 +4719,517 @@ int32 Describe(
 		TArray<TArray<FExportBundleEntryDesc>, TInlineAllocator<1>> ExportBundles;
 	};
 
-	if (!IFileManager::Get().FileExists(*GlobalContainerPath))
+	// Info loaded about a set of containers for the purposes of dumping to text in Describe or exploring some other way for debugging 
+	struct FContainerPackageInfo
 	{
-		UE_LOG(LogIoStore, Error, TEXT("Global container '%s' doesn't exist."), *GlobalContainerPath);
-		return -1;
-	}
-
-	TUniquePtr<FIoStoreReader> GlobalReader = CreateIoStoreReader(*GlobalContainerPath, KeyChain);
-	if (!GlobalReader.IsValid())
-	{
-		UE_LOG(LogIoStore, Warning, TEXT("Failed reading global container '%s'"), *GlobalContainerPath);
-		return -1;
-	}
-
-	UE_LOG(LogIoStore, Display, TEXT("Loading script imports..."));
+		TArray<FContainerDesc*> Containers;
+		TArray<FPackageDesc*> Packages;
+		TMap<FPackageObjectIndex, FScriptObjectDesc> ScriptObjectByGlobalIdMap;
+		TMap<FPublicExportKey, FExportDesc*> ExportByKeyMap;
 
-	TIoStatusOr<FIoBuffer> ScriptObjectsBuffer = GlobalReader->Read(CreateIoChunkId(0, 0, EIoChunkType::ScriptObjects), FIoReadOptions());
-	if (!ScriptObjectsBuffer.IsOk())
-	{
-		UE_LOG(LogIoStore, Warning, TEXT("Failed reading initial load meta chunk from global container '%s'"), *GlobalContainerPath);
-		return -1;
-	}
+		FContainerPackageInfo() = default;
+		FContainerPackageInfo(const FContainerPackageInfo&) = delete;
+		FContainerPackageInfo(FContainerPackageInfo&&) = default;
+		FContainerPackageInfo& operator=(const FContainerPackageInfo&) = delete;
+		FContainerPackageInfo& operator=(FContainerPackageInfo&&) = default;
+		~FContainerPackageInfo()
+		{
+			for (FPackageDesc* PackageDesc : Packages)
+			{
+				delete PackageDesc;
+			}
+			for (FContainerDesc* ContainerDesc : Containers)
+			{
+				delete ContainerDesc;
+			}
+		}
 
-	TMap<FPackageObjectIndex, FScriptObjectDesc> ScriptObjectByGlobalIdMap;
-	FLargeMemoryReader ScriptObjectsArchive(ScriptObjectsBuffer.ValueOrDie().Data(), ScriptObjectsBuffer.ValueOrDie().DataSize());
-	TArray<FDisplayNameEntryId> GlobalNameMap = LoadNameBatch(ScriptObjectsArchive);
-	int32 NumScriptObjects = 0;
-	ScriptObjectsArchive << NumScriptObjects;
-	const FScriptObjectEntry* ScriptObjectEntries = reinterpret_cast<const FScriptObjectEntry*>(ScriptObjectsBuffer.ValueOrDie().Data() + ScriptObjectsArchive.Tell());
-	for (int32 ScriptObjectIndex = 0; ScriptObjectIndex < NumScriptObjects; ++ScriptObjectIndex)
-	{
-		const FScriptObjectEntry& ScriptObjectEntry = ScriptObjectEntries[ScriptObjectIndex];
-		FMappedName MappedName = ScriptObjectEntry.Mapped;
-		check(MappedName.IsGlobal());
-		FScriptObjectDesc& ScriptObjectDesc = ScriptObjectByGlobalIdMap.Add(ScriptObjectEntry.GlobalIndex);
-		ScriptObjectDesc.Name = GlobalNameMap[MappedName.GetIndex()].ToName(MappedName.GetNumber());
-		ScriptObjectDesc.GlobalImportIndex = ScriptObjectEntry.GlobalIndex;
-		ScriptObjectDesc.OuterIndex = ScriptObjectEntry.OuterIndex;
-	}
-	for (auto& KV : ScriptObjectByGlobalIdMap)
-	{
-		FScriptObjectDesc& ScriptObjectDesc = KV.Get<1>();
-		if (ScriptObjectDesc.FullName.IsNone())
+		FString PackageObjectIndexToString(const FPackageDesc* Package, const FPackageObjectIndex& PackageObjectIndex, bool bIncludeName)
 		{
-			TArray<FScriptObjectDesc*> ScriptObjectStack;
-			FScriptObjectDesc* Current = &ScriptObjectDesc;
-			FString FullName;
-			while (Current)
+			if (PackageObjectIndex.IsNull())
 			{
-				if (!Current->FullName.IsNone())
+				return TEXT("<null>");
+			}
+			else if (PackageObjectIndex.IsPackageImport())
+			{
+				FPublicExportKey Key = FPublicExportKey::FromPackageImport(PackageObjectIndex, Package->ImportedPackageIds, Package->ImportedPublicExportHashes);
+				FExportDesc* ExportDesc = ExportByKeyMap.FindRef(Key);
+				if (ExportDesc && bIncludeName)
 				{
-					FullName = Current->FullName.ToString();
-					break;
+					return FString::Printf(TEXT("0x%llX '%s'"), PackageObjectIndex.Value(), *ExportDesc->FullName.ToString());
+				}
+				else
+				{
+					return FString::Printf(TEXT("0x%llX"), PackageObjectIndex.Value());
+				}
+			}
+			else if (PackageObjectIndex.IsScriptImport())
+			{
+				const FScriptObjectDesc* ScriptObjectDesc = ScriptObjectByGlobalIdMap.Find(PackageObjectIndex);
+				if (ScriptObjectDesc && bIncludeName)
+				{
+					return FString::Printf(TEXT("0x%llX '%s'"), PackageObjectIndex.Value(), *ScriptObjectDesc->FullName.ToString());
+				}
+				else
+				{
+					return FString::Printf(TEXT("0x%llX"), PackageObjectIndex.Value());
 				}
-				ScriptObjectStack.Push(Current);
-				Current = ScriptObjectByGlobalIdMap.Find(Current->OuterIndex);
 			}
-			while (ScriptObjectStack.Num() > 0)
+			else if (PackageObjectIndex.IsExport())
 			{
-				Current = ScriptObjectStack.Pop();
-				FullName /= Current->Name.ToString();
-				Current->FullName = FName(FullName);
+				return FString::Printf(TEXT("%d"), PackageObjectIndex.Value());
+			}
+			else
+			{
+				return FString::Printf(TEXT("0x%llX"), PackageObjectIndex.Value());
 			}
 		}
-	}
-
-	FString Directory = FPaths::GetPath(GlobalContainerPath);
-	FPaths::NormalizeDirectoryName(Directory);
+	};
 
-	TArray<FString> FoundContainerFiles;
-	IFileManager::Get().FindFiles(FoundContainerFiles, *(Directory / TEXT("*.utoc")), true, false);
-	TArray<FString> ContainerFilePaths;
-	for (const FString& Filename : FoundContainerFiles)
+	// Try and read all packages inside the containers and the links between them for debugging/analysis
+	TOptional<FContainerPackageInfo> TryGetContainerPackageInfo(
+		const FString& GlobalContainerPath,
+		const FKeyChain& KeyChain,
+		bool bIncludeExportHashes
+	)
 	{
-		ContainerFilePaths.Emplace(Directory / Filename);
-	}
-
-	UE_LOG(LogIoStore, Display, TEXT("Loading containers..."));
+		if (!IFileManager::Get().FileExists(*GlobalContainerPath))
+		{
+			UE_LOG(LogIoStore, Error, TEXT("Global container '%s' doesn't exist."), *GlobalContainerPath);
+			return {};
+		}
 
-	TArray<TUniquePtr<FIoStoreReader>> Readers;
+		TUniquePtr<FIoStoreReader> GlobalReader = CreateIoStoreReader(*GlobalContainerPath, KeyChain);
+		if (!GlobalReader.IsValid())
+		{
+			UE_LOG(LogIoStore, Warning, TEXT("Failed reading global container '%s'"), *GlobalContainerPath);
+			return {};
+		}
 
-	struct FLoadContainerHeaderJob
-	{
-		FName ContainerName;
-		FContainerDesc* ContainerDesc = nullptr;
-		TArray<FPackageDesc*> Packages;
-		FIoStoreReader* Reader = nullptr;
-		TArray<FIoContainerHeaderLocalizedPackage> RawLocalizedPackages;
-		TArray<FIoContainerHeaderPackageRedirect> RawPackageRedirects;
-	};
+		UE_LOG(LogIoStore, Display, TEXT("Loading script imports..."));
 
-	TArray<FLoadContainerHeaderJob> LoadContainerHeaderJobs;
+		TIoStatusOr<FIoBuffer> ScriptObjectsBuffer = GlobalReader->Read(CreateIoChunkId(0, 0, EIoChunkType::ScriptObjects), FIoReadOptions());
+		if (!ScriptObjectsBuffer.IsOk())
+		{
+			UE_LOG(LogIoStore, Warning, TEXT("Failed reading initial load meta chunk from global container '%s'"), *GlobalContainerPath);
+			return {};
+		}
 
-	for (const FString& ContainerFilePath : ContainerFilePaths)
-	{
-		TUniquePtr<FIoStoreReader> Reader = CreateIoStoreReader(*ContainerFilePath, KeyChain);
-		if (!Reader.IsValid())
+		TMap<FPackageObjectIndex, FScriptObjectDesc> ScriptObjectByGlobalIdMap;
+		FLargeMemoryReader ScriptObjectsArchive(ScriptObjectsBuffer.ValueOrDie().Data(), ScriptObjectsBuffer.ValueOrDie().DataSize());
+		TArray<FDisplayNameEntryId> GlobalNameMap = LoadNameBatch(ScriptObjectsArchive);
+		int32 NumScriptObjects = 0;
+		ScriptObjectsArchive << NumScriptObjects;
+		const FScriptObjectEntry* ScriptObjectEntries = reinterpret_cast<const FScriptObjectEntry*>(ScriptObjectsBuffer.ValueOrDie().Data() + ScriptObjectsArchive.Tell());
+		for (int32 ScriptObjectIndex = 0; ScriptObjectIndex < NumScriptObjects; ++ScriptObjectIndex)
 		{
-			UE_LOG(LogIoStore, Warning, TEXT("Failed to read container '%s'"), *ContainerFilePath);
-			continue;
+			const FScriptObjectEntry& ScriptObjectEntry = ScriptObjectEntries[ScriptObjectIndex];
+			FMappedName MappedName = ScriptObjectEntry.Mapped;
+			check(MappedName.IsGlobal());
+			FScriptObjectDesc& ScriptObjectDesc = ScriptObjectByGlobalIdMap.Add(ScriptObjectEntry.GlobalIndex);
+			ScriptObjectDesc.Name = GlobalNameMap[MappedName.GetIndex()].ToName(MappedName.GetNumber());
+			ScriptObjectDesc.GlobalImportIndex = ScriptObjectEntry.GlobalIndex;
+			ScriptObjectDesc.OuterIndex = ScriptObjectEntry.OuterIndex;
+		}
+		for (auto& KV : ScriptObjectByGlobalIdMap)
+		{
+			FScriptObjectDesc& ScriptObjectDesc = KV.Get<1>();
+			if (ScriptObjectDesc.FullName.IsNone())
+			{
+				TArray<FScriptObjectDesc*> ScriptObjectStack;
+				FScriptObjectDesc* Current = &ScriptObjectDesc;
+				FString FullName;
+				while (Current)
+				{
+					if (!Current->FullName.IsNone())
+					{
+						FullName = Current->FullName.ToString();
+						break;
+					}
+					ScriptObjectStack.Push(Current);
+					Current = ScriptObjectByGlobalIdMap.Find(Current->OuterIndex);
+				}
+				while (ScriptObjectStack.Num() > 0)
+				{
+					Current = ScriptObjectStack.Pop();
+					FullName /= Current->Name.ToString();
+					Current->FullName = FName(FullName);
+				}
+			}
 		}
 
-		FLoadContainerHeaderJob& LoadContainerHeaderJob = LoadContainerHeaderJobs.AddDefaulted_GetRef();
-		LoadContainerHeaderJob.Reader = Reader.Get();
-		LoadContainerHeaderJob.ContainerName = FName(FPaths::GetBaseFilename(ContainerFilePath));
-		
-		Readers.Emplace(MoveTemp(Reader));
-	}
-	
-	TAtomic<int32> TotalPackageCount{ 0 };
-	ParallelFor(LoadContainerHeaderJobs.Num(), [&LoadContainerHeaderJobs, &TotalPackageCount](int32 Index)
-	{
-		TRACE_CPUPROFILER_EVENT_SCOPE(LoadContainerHeader);
+		FString Directory = FPaths::GetPath(GlobalContainerPath);
+		FPaths::NormalizeDirectoryName(Directory);
 
-		FLoadContainerHeaderJob& Job = LoadContainerHeaderJobs[Index];
+		TArray<FString> FoundContainerFiles;
+		IFileManager::Get().FindFiles(FoundContainerFiles, *(Directory / TEXT("*.utoc")), true, false);
+		TArray<FString> ContainerFilePaths;
+		for (const FString& Filename : FoundContainerFiles)
+		{
+			ContainerFilePaths.Emplace(Directory / Filename);
+		}
 
-		FContainerDesc* ContainerDesc = new FContainerDesc();
-		ContainerDesc->Name = Job.ContainerName;
-		ContainerDesc->Id = Job.Reader->GetContainerId();
-		ContainerDesc->EncryptionKeyGuid = Job.Reader->GetEncryptionKeyGuid();
-		EIoContainerFlags Flags = Job.Reader->GetContainerFlags();
-		ContainerDesc->bCompressed = bool(Flags & EIoContainerFlags::Compressed);
-		ContainerDesc->bEncrypted = bool(Flags & EIoContainerFlags::Encrypted);
-		ContainerDesc->bSigned = bool(Flags & EIoContainerFlags::Signed);
-		ContainerDesc->bIndexed = bool(Flags & EIoContainerFlags::Indexed);
-		Job.ContainerDesc = ContainerDesc;
+		UE_LOG(LogIoStore, Display, TEXT("Loading containers..."));
 
-		TIoStatusOr<FIoBuffer> IoBuffer = Job.Reader->Read(CreateIoChunkId(Job.Reader->GetContainerId().Value(), 0, EIoChunkType::ContainerHeader), FIoReadOptions());
-		if (IoBuffer.IsOk())
-		{
-			FMemoryReaderView Ar(MakeArrayView(IoBuffer.ValueOrDie().Data(), IoBuffer.ValueOrDie().DataSize()));
-			FIoContainerHeader ContainerHeader;
-			Ar << ContainerHeader;
+		TArray<TUniquePtr<FIoStoreReader>> Readers;
 
-			Job.RawLocalizedPackages = ContainerHeader.LocalizedPackages;
-			Job.RawPackageRedirects = ContainerHeader.PackageRedirects;
+		struct FLoadContainerHeaderJob
+		{
+			FName ContainerName;
+			FContainerDesc* ContainerDesc = nullptr;
+			TArray<FPackageDesc*> Packages;
+			FIoStoreReader* Reader = nullptr;
+			TArray<FIoContainerHeaderLocalizedPackage> RawLocalizedPackages;
+			TArray<FIoContainerHeaderPackageRedirect> RawPackageRedirects;
+		};
 
-			TArrayView<FFilePackageStoreEntry> StoreEntries(reinterpret_cast<FFilePackageStoreEntry*>(ContainerHeader.StoreEntries.GetData()), ContainerHeader.PackageIds.Num());
+		TArray<FLoadContainerHeaderJob> LoadContainerHeaderJobs;
 
-			int32 PackageIndex = 0;
-			Job.Packages.Reserve(StoreEntries.Num());
-			for (FFilePackageStoreEntry& ContainerEntry : StoreEntries)
+		for (const FString& ContainerFilePath : ContainerFilePaths)
+		{
+			TUniquePtr<FIoStoreReader> Reader = CreateIoStoreReader(*ContainerFilePath, KeyChain);
+			if (!Reader.IsValid())
 			{
-				const FPackageId& PackageId = ContainerHeader.PackageIds[PackageIndex++];
-				FPackageDesc* PackageDesc = new FPackageDesc();
-				PackageDesc->PackageId = PackageId;
-				PackageDesc->Exports.SetNum(ContainerEntry.ExportCount);
-				PackageDesc->ExportBundleCount = ContainerEntry.ExportBundleCount;
-				PackageDesc->ImportedPackageIds = TArrayView<FPackageId>(ContainerEntry.ImportedPackages.Data(), ContainerEntry.ImportedPackages.Num());
-				Job.Packages.Add(PackageDesc);
-				++TotalPackageCount;
+				UE_LOG(LogIoStore, Warning, TEXT("Failed to read container '%s'"), *ContainerFilePath);
+				continue;
 			}
+
+			FLoadContainerHeaderJob& LoadContainerHeaderJob = LoadContainerHeaderJobs.AddDefaulted_GetRef();
+			LoadContainerHeaderJob.Reader = Reader.Get();
+			LoadContainerHeaderJob.ContainerName = FName(FPaths::GetBaseFilename(ContainerFilePath));
+			
+			Readers.Emplace(MoveTemp(Reader));
 		}
-	}, EParallelForFlags::Unbalanced);
+		
+		TAtomic<int32> TotalPackageCount{ 0 };
+		ParallelFor(LoadContainerHeaderJobs.Num(), [&LoadContainerHeaderJobs, &TotalPackageCount](int32 Index)
+		{
+			TRACE_CPUPROFILER_EVENT_SCOPE(LoadContainerHeader);
 
-	struct FLoadPackageSummaryJob
-	{
-		FPackageDesc* PackageDesc = nullptr;
-		FIoChunkId ChunkId;
-		TArray<FLoadContainerHeaderJob*, TInlineAllocator<1>> Containers;
-	};
+			FLoadContainerHeaderJob& Job = LoadContainerHeaderJobs[Index];
 
-	TArray<FLoadPackageSummaryJob> LoadPackageSummaryJobs;
+			FContainerDesc* ContainerDesc = new FContainerDesc();
+			ContainerDesc->Name = Job.ContainerName;
+			ContainerDesc->Id = Job.Reader->GetContainerId();
+			ContainerDesc->EncryptionKeyGuid = Job.Reader->GetEncryptionKeyGuid();
+			EIoContainerFlags Flags = Job.Reader->GetContainerFlags();
+			ContainerDesc->bCompressed = bool(Flags & EIoContainerFlags::Compressed);
+			ContainerDesc->bEncrypted = bool(Flags & EIoContainerFlags::Encrypted);
+			ContainerDesc->bSigned = bool(Flags & EIoContainerFlags::Signed);
+			ContainerDesc->bIndexed = bool(Flags & EIoContainerFlags::Indexed);
+			Job.ContainerDesc = ContainerDesc;
 
-	TArray<FContainerDesc*> Containers;
-	TArray<FPackageDesc*> Packages;
-	TMap<FPackageId, FPackageDesc*> PackageByIdMap;
-	TMap<FPackageId, FLoadPackageSummaryJob*> PackageJobByIdMap;
-	Containers.Reserve(LoadContainerHeaderJobs.Num());
-	Packages.Reserve(TotalPackageCount);
-	PackageByIdMap.Reserve(TotalPackageCount);
-	PackageJobByIdMap.Reserve(TotalPackageCount);
-	LoadPackageSummaryJobs.Reserve(TotalPackageCount);
-	for (FLoadContainerHeaderJob& LoadContainerHeaderJob : LoadContainerHeaderJobs)
-	{
-		Containers.Add(LoadContainerHeaderJob.ContainerDesc);
-		for (FPackageDesc* PackageDesc : LoadContainerHeaderJob.Packages)
-		{
-			FLoadPackageSummaryJob*& UniquePackageJob = PackageJobByIdMap.FindOrAdd(PackageDesc->PackageId);
-			if (!UniquePackageJob)
+			TIoStatusOr<FIoBuffer> IoBuffer = Job.Reader->Read(CreateIoChunkId(Job.Reader->GetContainerId().Value(), 0, EIoChunkType::ContainerHeader), FIoReadOptions());
+			if (IoBuffer.IsOk())
 			{
-				Packages.Add(PackageDesc);
-				PackageByIdMap.Add(PackageDesc->PackageId, PackageDesc);
-				FLoadPackageSummaryJob& LoadPackageSummaryJob = LoadPackageSummaryJobs.AddDefaulted_GetRef();
-				LoadPackageSummaryJob.PackageDesc = PackageDesc;
-				LoadPackageSummaryJob.ChunkId = CreateIoChunkId(PackageDesc->PackageId.Value(), 0, EIoChunkType::ExportBundleData);
-				UniquePackageJob = &LoadPackageSummaryJob;
+				FMemoryReaderView Ar(MakeArrayView(IoBuffer.ValueOrDie().Data(), IoBuffer.ValueOrDie().DataSize()));
+				FIoContainerHeader ContainerHeader;
+				Ar << ContainerHeader;
+
+				Job.RawLocalizedPackages = ContainerHeader.LocalizedPackages;
+				Job.RawPackageRedirects = ContainerHeader.PackageRedirects;
+
+				TArrayView<FFilePackageStoreEntry> StoreEntries(reinterpret_cast<FFilePackageStoreEntry*>(ContainerHeader.StoreEntries.GetData()), ContainerHeader.PackageIds.Num());
+
+				int32 PackageIndex = 0;
+				Job.Packages.Reserve(StoreEntries.Num());
+				for (FFilePackageStoreEntry& ContainerEntry : StoreEntries)
+				{
+					const FPackageId& PackageId = ContainerHeader.PackageIds[PackageIndex++];
+					FPackageDesc* PackageDesc = new FPackageDesc();
+					PackageDesc->PackageId = PackageId;
+					PackageDesc->Exports.SetNum(ContainerEntry.ExportCount);
+					PackageDesc->ExportBundleCount = ContainerEntry.ExportBundleCount;
+					PackageDesc->ImportedPackageIds = TArrayView<FPackageId>(ContainerEntry.ImportedPackages.Data(), ContainerEntry.ImportedPackages.Num());
+					Job.Packages.Add(PackageDesc);
+					++TotalPackageCount;
+				}
 			}
-			UniquePackageJob->Containers.Add(&LoadContainerHeaderJob);
-		}
-	}
-	for (FLoadContainerHeaderJob& LoadContainerHeaderJob : LoadContainerHeaderJobs)
-	{
-		for (const auto& RedirectPair : LoadContainerHeaderJob.RawPackageRedirects)
-		{
-			FPackageRedirect& PackageRedirect = LoadContainerHeaderJob.ContainerDesc->PackageRedirects.AddDefaulted_GetRef();
-			PackageRedirect.Source = PackageByIdMap.FindRef(RedirectPair.SourcePackageId);
-			PackageRedirect.Target = PackageByIdMap.FindRef(RedirectPair.TargetPackageId);
-		}
-		for (const auto& LocalizedPackage : LoadContainerHeaderJob.RawLocalizedPackages)
-		{
-			LoadContainerHeaderJob.ContainerDesc->LocalizedPackages.Add(PackageByIdMap.FindRef(LocalizedPackage.SourcePackageId));
-		}
-	}
-	
-	ParallelFor(LoadPackageSummaryJobs.Num(), [&LoadPackageSummaryJobs, bIncludeExportHashes](int32 Index)
-	{
-		TRACE_CPUPROFILER_EVENT_SCOPE(LoadPackageSummary);
+		}, EParallelForFlags::Unbalanced);
 
-		FLoadPackageSummaryJob& Job = LoadPackageSummaryJobs[Index];
-		for (FLoadContainerHeaderJob* LoadContainerHeaderJob : Job.Containers)
+		struct FLoadPackageSummaryJob
 		{
-			TIoStatusOr<FIoStoreTocChunkInfo> ChunkInfo = LoadContainerHeaderJob->Reader->GetChunkInfo(Job.ChunkId);
-			check(ChunkInfo.IsOk());
-			FPackageLocation& Location = Job.PackageDesc->Locations.AddDefaulted_GetRef();
-			Location.Container = LoadContainerHeaderJob->ContainerDesc;
-			Location.Offset = ChunkInfo.ValueOrDie().Offset;
-		}
+			FPackageDesc* PackageDesc = nullptr;
+			FIoChunkId ChunkId;
+			TArray<FLoadContainerHeaderJob*, TInlineAllocator<1>> Containers;
+		};
 
-		FIoStoreReader* Reader = Job.Containers[0]->Reader;
-		FIoReadOptions ReadOptions;
-		if (!bIncludeExportHashes)
-		{
-			ReadOptions.SetRange(0, 16 << 10);
+		TArray<FLoadPackageSummaryJob> LoadPackageSummaryJobs;
+
+		TArray<FContainerDesc*> Containers;
+		TArray<FPackageDesc*> Packages;
+		TMap<FPackageId, FPackageDesc*> PackageByIdMap;
+		TMap<FPackageId, FLoadPackageSummaryJob*> PackageJobByIdMap;
+		Containers.Reserve(LoadContainerHeaderJobs.Num());
+		Packages.Reserve(TotalPackageCount);
+		PackageByIdMap.Reserve(TotalPackageCount);
+		PackageJobByIdMap.Reserve(TotalPackageCount);
+		LoadPackageSummaryJobs.Reserve(TotalPackageCount);
+		for (FLoadContainerHeaderJob& LoadContainerHeaderJob : LoadContainerHeaderJobs)
+		{
+			Containers.Add(LoadContainerHeaderJob.ContainerDesc);
+			for (FPackageDesc* PackageDesc : LoadContainerHeaderJob.Packages)
+			{
+				FLoadPackageSummaryJob*& UniquePackageJob = PackageJobByIdMap.FindOrAdd(PackageDesc->PackageId);
+				if (!UniquePackageJob)
+				{
+					Packages.Add(PackageDesc);
+					PackageByIdMap.Add(PackageDesc->PackageId, PackageDesc);
+					FLoadPackageSummaryJob& LoadPackageSummaryJob = LoadPackageSummaryJobs.AddDefaulted_GetRef();
+					LoadPackageSummaryJob.PackageDesc = PackageDesc;
+					LoadPackageSummaryJob.ChunkId = CreateIoChunkId(PackageDesc->PackageId.Value(), 0, EIoChunkType::ExportBundleData);
+					UniquePackageJob = &LoadPackageSummaryJob;
+				}
+				UniquePackageJob->Containers.Add(&LoadContainerHeaderJob);
+			}
 		}
-		TIoStatusOr<FIoBuffer> IoBuffer = Reader->Read(Job.ChunkId, ReadOptions);
-		check(IoBuffer.IsOk());
-		const uint8* PackageSummaryData = IoBuffer.ValueOrDie().Data();
-		const FZenPackageSummary* PackageSummary = reinterpret_cast<const FZenPackageSummary*>(PackageSummaryData);
-		if (PackageSummary->HeaderSize > IoBuffer.ValueOrDie().DataSize())
+		for (FLoadContainerHeaderJob& LoadContainerHeaderJob : LoadContainerHeaderJobs)
 		{
-			ReadOptions.SetRange(0, PackageSummary->HeaderSize);
-			IoBuffer = Reader->Read(Job.ChunkId, ReadOptions);
-			PackageSummaryData = IoBuffer.ValueOrDie().Data();
-			PackageSummary = reinterpret_cast<const FZenPackageSummary*>(PackageSummaryData);
+			for (const auto& RedirectPair : LoadContainerHeaderJob.RawPackageRedirects)
+			{
+				FPackageRedirect& PackageRedirect = LoadContainerHeaderJob.ContainerDesc->PackageRedirects.AddDefaulted_GetRef();
+				PackageRedirect.Source = PackageByIdMap.FindRef(RedirectPair.SourcePackageId);
+				PackageRedirect.Target = PackageByIdMap.FindRef(RedirectPair.TargetPackageId);
+			}
+			for (const auto& LocalizedPackage : LoadContainerHeaderJob.RawLocalizedPackages)
+			{
+				LoadContainerHeaderJob.ContainerDesc->LocalizedPackages.Add(PackageByIdMap.FindRef(LocalizedPackage.SourcePackageId));
+			}
 		}
+		
+		ParallelFor(LoadPackageSummaryJobs.Num(), [&LoadPackageSummaryJobs, bIncludeExportHashes](int32 Index)
+		{
+			TRACE_CPUPROFILER_EVENT_SCOPE(LoadPackageSummary);
 
-		TArrayView<const uint8> HeaderDataView(PackageSummaryData + sizeof(FZenPackageSummary), PackageSummary->HeaderSize - sizeof(FZenPackageSummary));
-		FMemoryReaderView HeaderDataReader(HeaderDataView);
+			FLoadPackageSummaryJob& Job = LoadPackageSummaryJobs[Index];
+			for (FLoadContainerHeaderJob* LoadContainerHeaderJob : Job.Containers)
+			{
+				TIoStatusOr<FIoStoreTocChunkInfo> ChunkInfo = LoadContainerHeaderJob->Reader->GetChunkInfo(Job.ChunkId);
+				check(ChunkInfo.IsOk());
+				FPackageLocation& Location = Job.PackageDesc->Locations.AddDefaulted_GetRef();
+				Location.Container = LoadContainerHeaderJob->ContainerDesc;
+				Location.Offset = ChunkInfo.ValueOrDie().Offset;
+			}
 
-		FZenPackageVersioningInfo VersioningInfo;
-		if (PackageSummary->bHasVersioningInfo)
-		{
-			HeaderDataReader << VersioningInfo;
-		}
+			FIoStoreReader* Reader = Job.Containers[0]->Reader;
+			FIoReadOptions ReadOptions;
+			if (!bIncludeExportHashes)
+			{
+				ReadOptions.SetRange(0, 16 << 10);
+			}
+			TIoStatusOr<FIoBuffer> IoBuffer = Reader->Read(Job.ChunkId, ReadOptions);
+			check(IoBuffer.IsOk());
+			const uint8* PackageSummaryData = IoBuffer.ValueOrDie().Data();
+			const FZenPackageSummary* PackageSummary = reinterpret_cast<const FZenPackageSummary*>(PackageSummaryData);
+			if (PackageSummary->HeaderSize > IoBuffer.ValueOrDie().DataSize())
+			{
+				ReadOptions.SetRange(0, PackageSummary->HeaderSize);
+				IoBuffer = Reader->Read(Job.ChunkId, ReadOptions);
+				PackageSummaryData = IoBuffer.ValueOrDie().Data();
+				PackageSummary = reinterpret_cast<const FZenPackageSummary*>(PackageSummaryData);
+			}
 
-		TArray<FDisplayNameEntryId> PackageNameMap;
-		{
-			TRACE_CPUPROFILER_EVENT_SCOPE(LoadNameBatch);
-			PackageNameMap = LoadNameBatch(HeaderDataReader);
-		}
+			TArrayView<const uint8> HeaderDataView(PackageSummaryData + sizeof(FZenPackageSummary), PackageSummary->HeaderSize - sizeof(FZenPackageSummary));
+			FMemoryReaderView HeaderDataReader(HeaderDataView);
 
-		Job.PackageDesc->PackageName = PackageNameMap[PackageSummary->Name.GetIndex()].ToName(PackageSummary->Name.GetNumber());
-		Job.PackageDesc->PackageFlags = PackageSummary->PackageFlags;
-		Job.PackageDesc->NameCount = PackageNameMap.Num();
-		
-		Job.PackageDesc->ImportedPublicExportHashes = MakeArrayView<const uint64>(reinterpret_cast<const uint64*>(PackageSummaryData + PackageSummary->ImportedPublicExportHashesOffset), (PackageSummary->ImportMapOffset - PackageSummary->ImportedPublicExportHashesOffset) / sizeof(uint64));
-
-		const FPackageObjectIndex* ImportMap = reinterpret_cast<const FPackageObjectIndex*>(PackageSummaryData + PackageSummary->ImportMapOffset);
-		Job.PackageDesc->Imports.SetNum((PackageSummary->ExportMapOffset - PackageSummary->ImportMapOffset) / sizeof(FPackageObjectIndex));
-		for (int32 ImportIndex = 0; ImportIndex < Job.PackageDesc->Imports.Num(); ++ImportIndex)
-		{
-			FImportDesc& ImportDesc = Job.PackageDesc->Imports[ImportIndex];
-			ImportDesc.GlobalImportIndex = ImportMap[ImportIndex];
-		}
-
-		const FExportMapEntry* ExportMap = reinterpret_cast<const FExportMapEntry*>(PackageSummaryData + PackageSummary->ExportMapOffset);
-		for (int32 ExportIndex = 0; ExportIndex < Job.PackageDesc->Exports.Num(); ++ExportIndex)
-		{
-			const FExportMapEntry& ExportMapEntry = ExportMap[ExportIndex];
-			FExportDesc& ExportDesc = Job.PackageDesc->Exports[ExportIndex];
-			ExportDesc.Package = Job.PackageDesc;
-			ExportDesc.Name = PackageNameMap[ExportMapEntry.ObjectName.GetIndex()].ToName(ExportMapEntry.ObjectName.GetNumber());
-			ExportDesc.OuterIndex = ExportMapEntry.OuterIndex;
-			ExportDesc.ClassIndex = ExportMapEntry.ClassIndex;
-			ExportDesc.SuperIndex = ExportMapEntry.SuperIndex;
-			ExportDesc.TemplateIndex = ExportMapEntry.TemplateIndex;
-			ExportDesc.PublicExportHash = ExportMapEntry.PublicExportHash;
-			ExportDesc.SerialSize = ExportMapEntry.CookedSerialSize;
-		}
-
-		const FExportBundleHeader* ExportBundleHeaders = reinterpret_cast<const FExportBundleHeader*>(PackageSummaryData + PackageSummary->GraphDataOffset);
-		const FExportBundleEntry* ExportBundleEntries = reinterpret_cast<const FExportBundleEntry*>(PackageSummaryData + PackageSummary->ExportBundleEntriesOffset);
-		uint64 CurrentExportOffset = PackageSummary->HeaderSize;
-		for (int32 ExportBundleIndex = 0; ExportBundleIndex < Job.PackageDesc->ExportBundleCount; ++ExportBundleIndex)
-		{
-			TArray<FExportBundleEntryDesc>& ExportBundleDesc = Job.PackageDesc->ExportBundles.AddDefaulted_GetRef();
-			const FExportBundleHeader* ExportBundle = ExportBundleHeaders + ExportBundleIndex;
-			const FExportBundleEntry* BundleEntry = ExportBundleEntries + ExportBundle->FirstEntryIndex;
-			const FExportBundleEntry* BundleEntryEnd = BundleEntry + ExportBundle->EntryCount;
-			check(BundleEntry <= BundleEntryEnd);
-			while (BundleEntry < BundleEntryEnd)
-			{
-				FExportBundleEntryDesc& EntryDesc = ExportBundleDesc.AddDefaulted_GetRef();
-				EntryDesc.CommandType = FExportBundleEntry::EExportCommandType(BundleEntry->CommandType);
-				EntryDesc.LocalExportIndex = BundleEntry->LocalExportIndex;
-				EntryDesc.Export = &Job.PackageDesc->Exports[BundleEntry->LocalExportIndex];
-				if (BundleEntry->CommandType == FExportBundleEntry::ExportCommandType_Serialize)
-				{
-					EntryDesc.Export->SerialOffset = CurrentExportOffset;
-					CurrentExportOffset += EntryDesc.Export->SerialSize;
+			FZenPackageVersioningInfo VersioningInfo;
+			if (PackageSummary->bHasVersioningInfo)
+			{
+				HeaderDataReader << VersioningInfo;
+			}
+
+			TArray<FDisplayNameEntryId> PackageNameMap;
+			{
+				TRACE_CPUPROFILER_EVENT_SCOPE(LoadNameBatch);
+				PackageNameMap = LoadNameBatch(HeaderDataReader);
+			}
 
-					if (bIncludeExportHashes)
+			Job.PackageDesc->PackageName = PackageNameMap[PackageSummary->Name.GetIndex()].ToName(PackageSummary->Name.GetNumber());
+			Job.PackageDesc->PackageFlags = PackageSummary->PackageFlags;
+			Job.PackageDesc->NameCount = PackageNameMap.Num();
+			
+			Job.PackageDesc->ImportedPublicExportHashes = MakeArrayView<const uint64>(reinterpret_cast<const uint64*>(PackageSummaryData + PackageSummary->ImportedPublicExportHashesOffset), (PackageSummary->ImportMapOffset - PackageSummary->ImportedPublicExportHashesOffset) / sizeof(uint64));
+
+			const FPackageObjectIndex* ImportMap = reinterpret_cast<const FPackageObjectIndex*>(PackageSummaryData + PackageSummary->ImportMapOffset);
+			Job.PackageDesc->Imports.SetNum((PackageSummary->ExportMapOffset - PackageSummary->ImportMapOffset) / sizeof(FPackageObjectIndex));
+			for (int32 ImportIndex = 0; ImportIndex < Job.PackageDesc->Imports.Num(); ++ImportIndex)
+			{
+				FImportDesc& ImportDesc = Job.PackageDesc->Imports[ImportIndex];
+				ImportDesc.GlobalImportIndex = ImportMap[ImportIndex];
+			}
+
+			const FExportMapEntry* ExportMap = reinterpret_cast<const FExportMapEntry*>(PackageSummaryData + PackageSummary->ExportMapOffset);
+			for (int32 ExportIndex = 0; ExportIndex < Job.PackageDesc->Exports.Num(); ++ExportIndex)
+			{
+				const FExportMapEntry& ExportMapEntry = ExportMap[ExportIndex];
+				FExportDesc& ExportDesc = Job.PackageDesc->Exports[ExportIndex];
+				ExportDesc.Package = Job.PackageDesc;
+				ExportDesc.Name = PackageNameMap[ExportMapEntry.ObjectName.GetIndex()].ToName(ExportMapEntry.ObjectName.GetNumber());
+				ExportDesc.OuterIndex = ExportMapEntry.OuterIndex;
+				ExportDesc.ClassIndex = ExportMapEntry.ClassIndex;
+				ExportDesc.SuperIndex = ExportMapEntry.SuperIndex;
+				ExportDesc.TemplateIndex = ExportMapEntry.TemplateIndex;
+				ExportDesc.PublicExportHash = ExportMapEntry.PublicExportHash;
+				ExportDesc.SerialSize = ExportMapEntry.CookedSerialSize;
+			}
+
+			const FExportBundleHeader* ExportBundleHeaders = reinterpret_cast<const FExportBundleHeader*>(PackageSummaryData + PackageSummary->GraphDataOffset);
+			const FExportBundleEntry* ExportBundleEntries = reinterpret_cast<const FExportBundleEntry*>(PackageSummaryData + PackageSummary->ExportBundleEntriesOffset);
+			uint64 CurrentExportOffset = PackageSummary->HeaderSize;
+			for (int32 ExportBundleIndex = 0; ExportBundleIndex < Job.PackageDesc->ExportBundleCount; ++ExportBundleIndex)
+			{
+				TArray<FExportBundleEntryDesc>& ExportBundleDesc = Job.PackageDesc->ExportBundles.AddDefaulted_GetRef();
+				const FExportBundleHeader* ExportBundle = ExportBundleHeaders + ExportBundleIndex;
+				const FExportBundleEntry* BundleEntry = ExportBundleEntries + ExportBundle->FirstEntryIndex;
+				const FExportBundleEntry* BundleEntryEnd = BundleEntry + ExportBundle->EntryCount;
+				check(BundleEntry <= BundleEntryEnd);
+				while (BundleEntry < BundleEntryEnd)
+				{
+					FExportBundleEntryDesc& EntryDesc = ExportBundleDesc.AddDefaulted_GetRef();
+					EntryDesc.CommandType = FExportBundleEntry::EExportCommandType(BundleEntry->CommandType);
+					EntryDesc.LocalExportIndex = BundleEntry->LocalExportIndex;
+					EntryDesc.Export = &Job.PackageDesc->Exports[BundleEntry->LocalExportIndex];
+					if (BundleEntry->CommandType == FExportBundleEntry::ExportCommandType_Serialize)
 					{
-						check(EntryDesc.Export->SerialOffset + EntryDesc.Export->SerialSize <= IoBuffer.ValueOrDie().DataSize());
-						FSHA1::HashBuffer(IoBuffer.ValueOrDie().Data() + EntryDesc.Export->SerialOffset, EntryDesc.Export->SerialSize, EntryDesc.Export->ExportHash.Hash);
+						EntryDesc.Export->SerialOffset = CurrentExportOffset;
+						CurrentExportOffset += EntryDesc.Export->SerialSize;
+
+						if (bIncludeExportHashes)
+						{
+							check(EntryDesc.Export->SerialOffset + EntryDesc.Export->SerialSize <= IoBuffer.ValueOrDie().DataSize());
+							FSHA1::HashBuffer(IoBuffer.ValueOrDie().Data() + EntryDesc.Export->SerialOffset, EntryDesc.Export->SerialSize, EntryDesc.Export->ExportHash.Hash);
+						}
 					}
+					++BundleEntry;
 				}
-				++BundleEntry;
 			}
-		}
-	}, EParallelForFlags::Unbalanced);
-
-	UE_LOG(LogIoStore, Display, TEXT("Connecting imports and exports..."));
-	TMap<FPublicExportKey, FExportDesc*> ExportByKeyMap;
-	{
-		TRACE_CPUPROFILER_EVENT_SCOPE(ConnectImportsAndExports);
+		}, EParallelForFlags::Unbalanced);
 
-		for (FPackageDesc* PackageDesc : Packages)
+		UE_LOG(LogIoStore, Display, TEXT("Connecting imports and exports..."));
+		TMap<FPublicExportKey, FExportDesc*> ExportByKeyMap;
 		{
-			for (FExportDesc& ExportDesc : PackageDesc->Exports)
+			TRACE_CPUPROFILER_EVENT_SCOPE(ConnectImportsAndExports);
+
+			for (FPackageDesc* PackageDesc : Packages)
 			{
-				if (ExportDesc.PublicExportHash)
+				for (FExportDesc& ExportDesc : PackageDesc->Exports)
 				{
-					FPublicExportKey Key = FPublicExportKey::MakeKey(PackageDesc->PackageId, ExportDesc.PublicExportHash);
-					ExportByKeyMap.Add(Key, &ExportDesc);
+					if (ExportDesc.PublicExportHash)
+					{
+						FPublicExportKey Key = FPublicExportKey::MakeKey(PackageDesc->PackageId, ExportDesc.PublicExportHash);
+						ExportByKeyMap.Add(Key, &ExportDesc);
+					}
 				}
 			}
-		}
 
-		ParallelFor(Packages.Num(), [&Packages](int32 Index)
-		{
-			FPackageDesc* PackageDesc = Packages[Index];
-			for (FExportDesc& ExportDesc : PackageDesc->Exports)
+			ParallelFor(Packages.Num(), [&Packages](int32 Index)
 			{
-				if (ExportDesc.FullName.IsNone())
+				FPackageDesc* PackageDesc = Packages[Index];
+				for (FExportDesc& ExportDesc : PackageDesc->Exports)
 				{
-					TRACE_CPUPROFILER_EVENT_SCOPE(GenerateExportFullName);
-
-					TArray<FExportDesc*> ExportStack;
-					
-					FExportDesc* Current = &ExportDesc;
-					TStringBuilder<2048> FullNameBuilder;
-					TCHAR NameBuffer[FName::StringBufferSize];
-					for (;;)
+					if (ExportDesc.FullName.IsNone())
 					{
-						if (!Current->FullName.IsNone())
+						TRACE_CPUPROFILER_EVENT_SCOPE(GenerateExportFullName);
+
+						TArray<FExportDesc*> ExportStack;
+						
+						FExportDesc* Current = &ExportDesc;
+						TStringBuilder<2048> FullNameBuilder;
+						TCHAR NameBuffer[FName::StringBufferSize];
+						for (;;)
 						{
-							Current->FullName.ToString(NameBuffer);
-							FullNameBuilder.Append(NameBuffer);
-							break;
+							if (!Current->FullName.IsNone())
+							{
+								Current->FullName.ToString(NameBuffer);
+								FullNameBuilder.Append(NameBuffer);
+								break;
+							}
+							ExportStack.Push(Current);
+							if (Current->OuterIndex.IsNull())
+							{
+								PackageDesc->PackageName.ToString(NameBuffer);
+								FullNameBuilder.Append(NameBuffer);
+								break;
+							}
+							Current = &PackageDesc->Exports[Current->OuterIndex.Value()];
 						}
-						ExportStack.Push(Current);
-						if (Current->OuterIndex.IsNull())
+						while (ExportStack.Num() > 0)
 						{
-							PackageDesc->PackageName.ToString(NameBuffer);
+							Current = ExportStack.Pop(false);
+							FullNameBuilder.Append(TEXT("."));
+							Current->Name.ToString(NameBuffer);
 							FullNameBuilder.Append(NameBuffer);
-							break;
+							Current->FullName = FName(FullNameBuilder);
 						}
-						Current = &PackageDesc->Exports[Current->OuterIndex.Value()];
-					}
-					while (ExportStack.Num() > 0)
-					{
-						Current = ExportStack.Pop(false);
-						FullNameBuilder.Append(TEXT("/"));
-						Current->Name.ToString(NameBuffer);
-						FullNameBuilder.Append(NameBuffer);
-						Current->FullName = FName(FullNameBuilder);
 					}
 				}
-			}
-		}, EParallelForFlags::Unbalanced);
+			}, EParallelForFlags::Unbalanced);
 
-		for (FPackageDesc* PackageDesc : Packages)
-		{
-			for (FImportDesc& Import : PackageDesc->Imports)
+			for (FPackageDesc* PackageDesc : Packages)
 			{
-				if (!Import.GlobalImportIndex.IsNull())
+				for (FImportDesc& Import : PackageDesc->Imports)
 				{
-					if (Import.GlobalImportIndex.IsPackageImport())
+					if (!Import.GlobalImportIndex.IsNull())
 					{
-						FPublicExportKey Key = FPublicExportKey::FromPackageImport(Import.GlobalImportIndex, PackageDesc->ImportedPackageIds, PackageDesc->ImportedPublicExportHashes);
-						Import.Export = ExportByKeyMap.FindRef(Key);
-						if (!Import.Export)
+						if (Import.GlobalImportIndex.IsPackageImport())
 						{
-							UE_LOG(LogIoStore, Warning, TEXT("Missing import: 0x%llX in package 0x%llX '%s'"), Import.GlobalImportIndex.Value(), PackageDesc->PackageId.ValueForDebugging(), *PackageDesc->PackageName.ToString());
+							FPublicExportKey Key = FPublicExportKey::FromPackageImport(Import.GlobalImportIndex, PackageDesc->ImportedPackageIds, PackageDesc->ImportedPublicExportHashes);
+							Import.Export = ExportByKeyMap.FindRef(Key);
+							if (!Import.Export)
+							{
+								UE_LOG(LogIoStore, Warning, TEXT("Missing import: 0x%llX in package 0x%llX '%s'"), Import.GlobalImportIndex.Value(), PackageDesc->PackageId.ValueForDebugging(), *PackageDesc->PackageName.ToString());
+							}
+							else
+							{
+								Import.Name = Import.Export->FullName;
+							}
 						}
 						else
 						{
-							Import.Name = Import.Export->FullName;
+							FScriptObjectDesc* ScriptObjectDesc = ScriptObjectByGlobalIdMap.Find(Import.GlobalImportIndex);
+							check(ScriptObjectDesc);
+							Import.Name = ScriptObjectDesc->FullName;
 						}
 					}
-					else
-					{
-						FScriptObjectDesc* ScriptObjectDesc = ScriptObjectByGlobalIdMap.Find(Import.GlobalImportIndex);
-						check(ScriptObjectDesc);
-						Import.Name = ScriptObjectDesc->FullName;
-					}
 				}
 			}
 		}
+
+		return { 
+			FContainerPackageInfo{
+				MoveTemp(Containers),
+				MoveTemp(Packages),
+				MoveTemp(ScriptObjectByGlobalIdMap),
+				MoveTemp(ExportByKeyMap),
+			} 
+		};
+	}
+}
+
+int32 Describe(
+	const FString& GlobalContainerPath,
+	const FKeyChain& KeyChain,
+	const FString& PackageFilter,
+	const FString& OutPath,
+	bool bIncludeExportHashes)
+{
+	using namespace DescribeUtils;
+	TOptional<FContainerPackageInfo> MaybeInfo = DescribeUtils::TryGetContainerPackageInfo(GlobalContainerPath, KeyChain, bIncludeExportHashes);
+	if (!MaybeInfo.IsSet())
+	{
+		return -1;
 	}
+	FContainerPackageInfo& Info = MaybeInfo.GetValue();
+
+	const TArray<FContainerDesc*>& Containers = Info.Containers;
+	const TArray<FPackageDesc*>& Packages = Info.Packages;
+	const TMap<FPackageObjectIndex, FScriptObjectDesc>& ScriptObjectByGlobalIdMap = Info.ScriptObjectByGlobalIdMap;
+	const TMap<FPublicExportKey, FExportDesc*>& ExportByKeyMap = Info.ExportByKeyMap;
 
 	UE_LOG(LogIoStore, Display, TEXT("Collecting output packages..."));
 	TArray<const FPackageDesc*> OutputPackages;
+	TSet<FPackageId> RelevantPackages;
+	TSet<FContainerDesc*> RelevantContainers;
 	{
 		TRACE_CPUPROFILER_EVENT_SCOPE(CollectOutputPackages);
 
@@ -5187,6 +5292,11 @@ int32 Describe(
 				{
 					Visited.Add(PackageDesc);
 					OutputPackages.Add(PackageDesc);
+					RelevantPackages.Add(PackageDesc->PackageId);
+					for (const FPackageLocation& Location : PackageDesc->Locations)
+					{
+						RelevantContainers.Add(Location.Container);
+					}
 					for (const FImportDesc& Import : PackageDesc->Imports)
 					{
 						if (Import.Export && Import.Export->Package)
@@ -5237,7 +5347,7 @@ int32 Describe(
 			}
 			else if (PackageObjectIndex.IsScriptImport())
 			{
-				FScriptObjectDesc* ScriptObjectDesc = ScriptObjectByGlobalIdMap.Find(PackageObjectIndex);
+				const FScriptObjectDesc* ScriptObjectDesc = ScriptObjectByGlobalIdMap.Find(PackageObjectIndex);
 				if (ScriptObjectDesc && bIncludeName)
 				{
 					return FString::Printf(TEXT("0x%llX '%s'"), PackageObjectIndex.Value(), *ScriptObjectDesc->FullName.ToString());
@@ -5259,6 +5369,11 @@ int32 Describe(
 
 		for (const FContainerDesc* ContainerDesc : Containers)
 		{
+			if (RelevantContainers.Num() > 0 && !RelevantContainers.Contains(ContainerDesc))
+			{
+				continue;
+			}
+
 			OutputOverride->Logf(ELogVerbosity::Display, TEXT("********************************************"));
 			OutputOverride->Logf(ELogVerbosity::Display, TEXT("Container '%s' Summary"), *ContainerDesc->Name.ToString());
 			OutputOverride->Logf(ELogVerbosity::Display, TEXT("--------------------------------------------"));
@@ -5274,11 +5389,20 @@ int32 Describe(
 
 			if (ContainerDesc->LocalizedPackages.Num())
 			{
-				OutputOverride->Logf(ELogVerbosity::Display, TEXT("--------------------------------------------"));
-				OutputOverride->Logf(ELogVerbosity::Display, TEXT("Localized Packages"));
-				OutputOverride->Logf(ELogVerbosity::Display, TEXT("=========="));
+				bool bNeedsHeader = true;
 				for (const FPackageDesc* LocalizedPackage : ContainerDesc->LocalizedPackages)
 				{
+					if (RelevantPackages.Num() > 0 && !RelevantPackages.Contains(LocalizedPackage->PackageId))
+					{
+						continue;
+					}
+					if (bNeedsHeader)
+					{
+						OutputOverride->Logf(ELogVerbosity::Display, TEXT("--------------------------------------------"));
+						OutputOverride->Logf(ELogVerbosity::Display, TEXT("Localized Packages"));
+						OutputOverride->Logf(ELogVerbosity::Display, TEXT("=========="));
+						bNeedsHeader = false;
+					}
 					OutputOverride->Logf(ELogVerbosity::Display, TEXT("\t*************************"));
 					OutputOverride->Logf(ELogVerbosity::Display, TEXT("\t\t           Source: 0x%llX '%s'"), LocalizedPackage->PackageId.ValueForDebugging(), *LocalizedPackage->PackageName.ToString());
 				}
@@ -5286,11 +5410,20 @@ int32 Describe(
 
 			if (ContainerDesc->PackageRedirects.Num())
 			{
-				OutputOverride->Logf(ELogVerbosity::Display, TEXT("--------------------------------------------"));
-				OutputOverride->Logf(ELogVerbosity::Display, TEXT("Package Redirects"));
-				OutputOverride->Logf(ELogVerbosity::Display, TEXT("=========="));
+				bool bNeedsHeader = true;
 				for (const FPackageRedirect& Redirect : ContainerDesc->PackageRedirects)
 				{
+					if (RelevantPackages.Num() > 0 && !RelevantPackages.Contains(Redirect.Source->PackageId) && !RelevantPackages.Contains(Redirect.Target->PackageId))
+					{
+						continue;
+					}
+					if (bNeedsHeader)
+					{
+						OutputOverride->Logf(ELogVerbosity::Display, TEXT("--------------------------------------------"));
+						OutputOverride->Logf(ELogVerbosity::Display, TEXT("Package Redirects"));
+						OutputOverride->Logf(ELogVerbosity::Display, TEXT("=========="));
+						bNeedsHeader = false;
+					}
 					OutputOverride->Logf(ELogVerbosity::Display, TEXT("\t*************************"));
 					OutputOverride->Logf(ELogVerbosity::Display, TEXT("\t\t           Source: 0x%llX '%s'"), Redirect.Source->PackageId.ValueForDebugging(), *Redirect.Source->PackageName.ToString());
 					OutputOverride->Logf(ELogVerbosity::Display, TEXT("\t\t           Target: 0x%llX '%s'"), Redirect.Target->PackageId.ValueForDebugging(), *Redirect.Target->PackageName.ToString());
@@ -5378,15 +5511,280 @@ int32 Describe(
 		}
 	}
 
-	for (FPackageDesc* PackageDesc : Packages)
+
+	return 0;
+}
+
+int32 ValidateCrossContainerRefs(
+	const FString& GlobalContainerPath,
+	const FKeyChain& KeyChain,
+	const FString& ConfigPath,
+	const FString& OutPath
+	)
+{
+	TMultiMap<FString, FString> ValidEdges;
+	TArray<FString> IgnoreRefsFromAssets, IgnoreRefsToAssets;
+	if (FString ConfigString; FFileHelper::LoadFileToString(ConfigString, *ConfigPath))
 	{
-		delete PackageDesc;
+		TSharedPtr<FJsonObject> RootConfigObject;
+		TSharedRef<TJsonReader<>> Reader = TJsonReaderFactory<>::Create(ConfigString);
+		FJsonSerializer::Deserialize(Reader, RootConfigObject);
+		if (!RootConfigObject.IsValid())
+		{
+			UE_LOG(LogIoStore, Error, TEXT("Failed to read configuration from file %s, check if it is valid json"), *ConfigPath);
+			return -1;
+		}
+		if (TSharedPtr<FJsonObject> ValidEdgesObj = RootConfigObject->GetObjectField(TEXT("ValidEdges")); ValidEdgesObj.IsValid())
+		{
+			for (const TPair<FString, TSharedPtr<FJsonValue>>& EdgePair : ValidEdgesObj->Values) 
+			{
+				const FString& From = EdgePair.Key;
+				FString ToString;
+				if (EdgePair.Value->TryGetString(ToString))
+				{
+					ValidEdges.Add(From, ToString);
+				}
+				else if (TArray<TSharedPtr<FJsonValue>>* ToArray; EdgePair.Value->TryGetArray(ToArray))
+				{
+					for (const TSharedPtr<FJsonValue>& ToValue : *ToArray)
+					{
+						if (ToValue->TryGetString(ToString))
+						{
+							ValidEdges.Add(From, ToString);
+						}
+					}
+				}
+			}
+		}
+		const TArray<TSharedPtr<FJsonValue>> IgnoreRefsFromArray = RootConfigObject->GetArrayField(TEXT("IgnoreRefsFrom"));
+		for (const TSharedPtr<FJsonValue>& Entry : IgnoreRefsFromArray)
+		{
+			if (FString Value; Entry->TryGetString(Value))
+			{
+				IgnoreRefsFromAssets.Add(MoveTemp(Value));
+			}
+		}
+		const TArray<TSharedPtr<FJsonValue>> IgnoreRefsToArray = RootConfigObject->GetArrayField(TEXT("IgnoreRefsTo"));
+		for (const TSharedPtr<FJsonValue>& Entry : IgnoreRefsToArray)
+		{
+			if (FString Value; Entry->TryGetString(Value))
+			{
+				IgnoreRefsToAssets.Add(MoveTemp(Value));
+			}
+		}
+	}
+	else
+	{
+		UE_LOG(LogIoStore, Error, TEXT("Failed to load config file %s"), *ConfigPath);
+		return -1;
+	}
+
+	if (ValidEdges.Num() == 0)
+	{
+		UE_LOG(LogIoStore, Error, TEXT("No valid edges configured, nothing to validate"));
+		return -1;
+	}
+
+	if (ValidEdges.FindPair(TEXT(""), TEXT("")))
+	{
+		UE_LOG(LogIoStore, Error, TEXT("Configuration contains all-to-all edge (empty string to empty string), nothing to validate"));
+		return -1;
+	}
+
+	using namespace DescribeUtils;
+	TOptional<FContainerPackageInfo> MaybeInfo = DescribeUtils::TryGetContainerPackageInfo(GlobalContainerPath, KeyChain, false);
+	if (!MaybeInfo.IsSet())
+	{
+		return -1;
+	}
+	FContainerPackageInfo& Info = MaybeInfo.GetValue();
+
+	const TArray<FContainerDesc*>& Containers = Info.Containers;
+	TArray<FPackageDesc*>& Packages = Info.Packages;
+	const TMap<FPackageObjectIndex, FScriptObjectDesc>& ScriptObjectByGlobalIdMap = Info.ScriptObjectByGlobalIdMap;
+	const TMap<FPublicExportKey, FExportDesc*>& ExportByKeyMap = Info.ExportByKeyMap;
+
+	// Expand container prefixes from config into full container names and produce transitive closure
+	TMultiMap<const FContainerDesc*, const FContainerDesc*> FinalValidEdges;
+	{
+		TMultiMap<FString, const FContainerDesc* > ShortNameToContainer;
+		for (auto It = ValidEdges.CreateIterator(); It; ++It)
+		{
+			for (const FContainerDesc* Desc : Containers)
+			{
+				// Empty strings mean 'all containers'
+				if (It.Key().Len() == 0 || Desc->Name.ToString().StartsWith(It.Key()))
+				{
+					ShortNameToContainer.AddUnique(It.Key(), Desc);
+				}
+				if (It.Value().Len() == 0 || Desc->Name.ToString().StartsWith(It.Value()))
+				{
+					ShortNameToContainer.AddUnique(It.Value(), Desc);
+				}
+			}
+		}
+		TMultiMap<const FContainerDesc*, const FContainerDesc*> ValidDirectEdges;
+		for (const TPair<FString, FString>& Pair : ValidEdges)
+		{
+			if (Pair.Key.Len() == 0) 
+			{ 
+				// Do 'all containers' later after handling explicit containers 
+				continue;
+			}
+			for (auto FromIt = ShortNameToContainer.CreateKeyIterator(Pair.Key); FromIt; ++FromIt)
+			{
+				for (auto ToIt = ShortNameToContainer.CreateKeyIterator(Pair.Value); ToIt; ++ToIt)
+				{
+					ValidDirectEdges.AddUnique(FromIt.Value(), ToIt.Value());
+				}
+			}
+		}
+		
+		TSet<const FContainerDesc*> UnassignedFromContainers;
+		for (const FContainerDesc* Container : Containers) 
+		{
+			if (!ValidDirectEdges.Contains(Container)) 
+			{
+				UnassignedFromContainers.Add(Container);
+			}
+		}
+		
+		for (const TPair<FString, FString>& Pair : ValidEdges)
+		{
+			if (Pair.Key.Len() == 0) 
+			{ 
+				for (auto FromIt = ShortNameToContainer.CreateKeyIterator(Pair.Key); FromIt; ++FromIt)
+				{
+					if (!UnassignedFromContainers.Contains(FromIt.Value()))
+					{
+						continue;
+					}
+
+					for (auto ToIt = ShortNameToContainer.CreateKeyIterator(Pair.Value); ToIt; ++ToIt)
+					{
+						ValidDirectEdges.Add(FromIt.Value(), ToIt.Value());
+					}
+				}
+			}
+		}
+
+		// Create a transitive closure (i.e. if it is valid for packages in container A to import packages in B, and B to import packages in C, then it is valid for packages in A to import packages in C)
+		for (const FContainerDesc* StartContainer : Containers)
+		{
+			TSet<const FContainerDesc*> SeenContainers;
+			TArray<const FContainerDesc*> Queue;
+			Queue.Add(StartContainer);
+			SeenContainers.Add(StartContainer);
+			while (Queue.Num() > 0)
+			{
+				const FContainerDesc* ToContainer = Queue.Pop();
+				FinalValidEdges.Add(StartContainer, ToContainer);
+				for (auto ToIt = ValidDirectEdges.CreateKeyIterator(ToContainer); ToIt; ++ToIt)
+				{
+					if (!SeenContainers.Contains(ToIt.Value()))
+					{
+						SeenContainers.Add(ToIt.Value());
+						Queue.Add(ToIt.Value());
+					}
+				}
+			}
+		}
 	}
-	for (FContainerDesc* ContainerDesc : Containers)
+
+	UE_LOG(LogIoStore, Display, TEXT("Final valid edges: %d"), FinalValidEdges.Num());
+	for (const TPair<const FContainerDesc*, const FContainerDesc*>& Pair : FinalValidEdges)
 	{
-		delete ContainerDesc;
+		UE_LOG(LogIoStore, Display, TEXT("%s -> %s"), *Pair.Key->Name.ToString(), *Pair.Value->Name.ToString());
 	}
 
+	TMap<TTuple<const FContainerDesc*, const FContainerDesc*>, TSet<TTuple<const FPackageDesc*, const FPackageDesc*>>> Errors;
+	Algo::SortBy(Packages, [](FPackageDesc* Desc) { return Desc->PackageName; }, FNameLexicalLess());
+	for (const FPackageDesc* Package : Packages)
+	{
+		if (Package->Locations.Num() > 1)
+		{
+			UE_LOG(LogIoStore, Error, TEXT("ValidateCrossContainerRefs does not support assets with multiple locations. (%s)"), *Package->PackageName.ToString());
+			continue;
+		}
+		
+		bool bSkip = Algo::AnyOf(IgnoreRefsFromAssets, [Package](const FString& IgnoreString)
+		{
+			FString PackageNameString = Package->PackageName.ToString();
+			if (PackageNameString == IgnoreString)
+			{
+				return true;
+			}
+			else if (FWildcardString::ContainsWildcards(*IgnoreString) && FWildcardString::IsMatch(*IgnoreString, *PackageNameString))
+			{
+				return true;
+			}
+			return false;
+		});
+		if (bSkip)
+		{
+			continue;
+		}
+
+		const FContainerDesc* Container = Package->Locations[0].Container;
+		bool bNeedsHeader = true;
+		for (const FImportDesc& Import : Package->Imports)
+		{
+			FPackageDesc* ImportPackage = Import.Export ? Import.Export->Package : nullptr;
+			if (!ImportPackage) 
+			{
+				UE_CLOG(Import.Name != FName() && !FPackageName::IsScriptPackage(*Import.Name.ToString()),
+					LogIoStore, Error, TEXT("Unresolved import of package %s by package %s"), *Import.Name.ToString(), *Package->PackageName.ToString());
+				continue;
+			}
+			if (ImportPackage->Locations.Num() > 1)
+			{
+				UE_LOG(LogIoStore, Error, TEXT("ValidateCrossContainerRefs does not support assets with multiple locations. (%s)"), *Package->PackageName.ToString());
+				continue;
+			}
+
+			bSkip = Algo::AnyOf(IgnoreRefsFromAssets, [ImportPackage](const FString& IgnoreString)
+			{
+				FString PackageNameString = ImportPackage->PackageName.ToString();
+				if (PackageNameString == IgnoreString)
+				{
+					return true;
+				}
+				else if (FWildcardString::ContainsWildcards(*IgnoreString) && FWildcardString::IsMatch(*IgnoreString, *PackageNameString))
+				{
+					return true;
+				}
+				return false;
+			});
+			if (bSkip)
+			{
+				continue;
+			}
+
+			const FContainerDesc* ImportContainer = ImportPackage->Locations[0].Container;
+			if (!FinalValidEdges.FindPair(Container, ImportContainer))
+			{
+				Errors.FindOrAdd(TTuple<const FContainerDesc*, const FContainerDesc*>(Container, ImportContainer))
+					.Add({Package, ImportPackage});
+			}
+		}
+	}
+	FOutputDevice* OutputOverride = GWarn;
+	FString OutputFilename;
+	TUniquePtr<FOutputDeviceFile> OutputBuffer;
+	if (!OutPath.IsEmpty())
+	{
+		OutputBuffer = MakeUnique<FOutputDeviceFile>(*OutPath, true);
+		OutputBuffer->SetSuppressEventTag(true);
+		OutputOverride = OutputBuffer.Get();
+	}
+	for (const TPair<TTuple<const FContainerDesc*, const FContainerDesc*>, TSet<TTuple<const FPackageDesc*, const FPackageDesc*>>>& Pair : Errors)
+	{
+		OutputOverride->Logf(ELogVerbosity::Display, TEXT("%s -> %s"), *Pair.Key.Key->Name.ToString(), *Pair.Key.Value->Name.ToString());
+		for (const TTuple<const FPackageDesc*, const FPackageDesc*>& Error : Pair.Value) 
+		{
+			OutputOverride->Logf(ELogVerbosity::Display, TEXT("\t%s -> %s"), *Error.Key->PackageName.ToString(), *Error.Value->PackageName.ToString());
+		} 
+	}
 	return 0;
 }
 
@@ -6860,6 +7258,15 @@ int32 CreateIoStoreContainerFiles(const TCHAR* CmdLine)
 		bool bIncludeExportHashes = FParse::Param(FCommandLine::Get(), TEXT("IncludeExportHashes"));
 		return Describe(ContainerPathOrWildcard, Arguments.KeyChain, PackageFilter, OutPath, bIncludeExportHashes);
 	}
+	else if (FParse::Value(FCommandLine::Get(), TEXT("ValidateCrossContainerRefs="), ArgumentValue))
+	{
+		FString ContainerPathOrWildcard = ArgumentValue;
+		FString ConfigPath;
+		FParse::Value(FCommandLine::Get(), TEXT("Config="), ConfigPath);
+		FString OutPath;
+		FParse::Value(FCommandLine::Get(), TEXT("DumpToFile="), OutPath);
+		return ValidateCrossContainerRefs(ContainerPathOrWildcard, Arguments.KeyChain, ConfigPath, OutPath);
+	}
 	else if (FParse::Param(FCommandLine::Get(), TEXT("ProfileReadSpeed")))
 	{
 		// Load the .UTOC file provided and read it in its entirety, cmdline parsed in function
