Fix C++20 compile error. FContainerPackageInfo can no longer be aggregate initialized in C++20.

diff --git a/Engine/Source/Developer/IoStoreUtilities/Private/IoStoreUtilities.cpp b/Engine/Source/Developer/IoStoreUtilities/Private/IoStoreUtilities.cpp
--- a/Engine/Source/Developer/IoStoreUtilities/Private/IoStoreUtilities.cpp
+++ b/Engine/Source/Developer/IoStoreUtilities/Private/IoStoreUtilities.cpp
@@ -4721,6 +4721,17 @@ namespace DescribeUtils
 		TMap<FPublicExportKey, FExportDesc*> ExportByKeyMap;
 
 		FContainerPackageInfo() = default;
+		FContainerPackageInfo(
+			TArray<FContainerDesc*> InContainers,
+			TArray<FPackageDesc*> InPackages,
+			TMap<FPackageObjectIndex, FScriptObjectDesc> InScriptObjectByGlobalIdMap,
+			TMap<FPublicExportKey, FExportDesc*> InExportByKeyMap)
+			: Containers(MoveTemp(InContainers))
+			, Packages(MoveTemp(InPackages))
+			, ScriptObjectByGlobalIdMap(MoveTemp(InScriptObjectByGlobalIdMap))
+			, ExportByKeyMap(MoveTemp(InExportByKeyMap))
+		{}
+
 		FContainerPackageInfo(const FContainerPackageInfo&) = delete;
 		FContainerPackageInfo(FContainerPackageInfo&&) = default;
 		FContainerPackageInfo& operator=(const FContainerPackageInfo&) = delete;
