ZenCook: Workaround for not being able to reliably start the zen server from UAT
#rnx

diff --git a/Engine/Source/Developer/IoStoreUtilities/Private/ZenStoreWriter.cpp b/Engine/Source/Developer/IoStoreUtilities/Private/ZenStoreWriter.cpp
--- a/Engine/Source/Developer/IoStoreUtilities/Private/ZenStoreWriter.cpp
+++ b/Engine/Source/Developer/IoStoreUtilities/Private/ZenStoreWriter.cpp
@@ -597,6 +597,44 @@ void FZenStoreWriter::EndCook(const FCookInfo& Info)
 		{
 			UE_LOG(LogSavePackage, Error, TEXT("Failed saving package store manifest file '%s'"), *PackageStoreManifestFilePath);
 		}
+
+		{
+			// Temporary solution until we can reliably read the oplog from UAT
+			TArray<FString> CookedFiles;
+			TIoStatusOr<FCbObject> OplogStatus = HttpClient->GetOplog().Get();
+			if (OplogStatus.IsOk())
+			{
+				FCbObject Oplog = OplogStatus.ConsumeValueOrDie();
+				for (FCbField& OplogEntry : Oplog["entries"].AsArray())
+				{
+					FCbObject OplogObj = OplogEntry.AsObject();
+					for (FCbField& ChunkEntry : OplogEntry["packagedata"].AsArray())
+					{
+						FCbObject ChunkObj = ChunkEntry.AsObject();
+						if (ChunkObj["filename"])
+						{
+							CookedFiles.Add(FString(ChunkObj["filename"].AsString()));
+						}
+					}
+					for (FCbField& ChunkEntry : OplogEntry["bulkdata"].AsArray())
+					{
+						FCbObject ChunkObj = ChunkEntry.AsObject();
+						if (ChunkObj["filename"])
+						{
+							CookedFiles.Add(FString(ChunkObj["filename"].AsString()));
+						}
+					}
+				}
+				if (!FFileHelper::SaveStringArrayToFile(CookedFiles, *FPaths::Combine(MetadataDirectoryPath, TEXT("cookedfiles.manifest"))))
+				{
+					UE_LOG(LogSavePackage, Error, TEXT("Failed writing UAT file manifest"));
+				}
+			}
+			else
+			{
+				UE_LOG(LogSavePackage, Error, TEXT("Failed reading oplog"));
+			}
+		}
 	}
 
 	UE_LOG(LogZenStoreWriter, Display, TEXT("Input:\t%d Packages"), PackageStoreOptimizer->GetTotalPackageCount());
diff --git a/Engine/Source/Programs/AutomationTool/Scripts/CopyBuildToStagingDirectory.Automation.cs b/Engine/Source/Programs/AutomationTool/Scripts/CopyBuildToStagingDirectory.Automation.cs
--- a/Engine/Source/Programs/AutomationTool/Scripts/CopyBuildToStagingDirectory.Automation.cs
+++ b/Engine/Source/Programs/AutomationTool/Scripts/CopyBuildToStagingDirectory.Automation.cs
@@ -761,6 +761,22 @@ namespace AutomationScripts
 					return;
 				}
 
+				{
+					// Temporary solution until we can reliably read the oplog from UAT
+					FileReference CookedFilesManifestFile = FileReference.Combine(SC.MetadataDir, "cookedfiles.manifest");
+					if (FileReference.Exists(CookedFilesManifestFile))
+					{
+						SC.PackageStoreManifest.ZenCookedFiles = new List<string>();
+						string[] CookedFiles = FileReference.ReadAllLines(CookedFilesManifestFile);
+						foreach (string RelativeFilename in CookedFiles)
+						{
+							CommandUtils.LogInformation("{0}", RelativeFilename);
+							SC.PackageStoreManifest.ZenCookedFiles.Add(FileReference.Combine(SC.PlatformCookDir, RelativeFilename).FullName);
+						}
+						return;
+					}
+				}
+
 				byte[] ManifestData = File.ReadAllBytes(PackageStoreManifestFile.FullName);
 				CbObject ManifestObject = new CbField(ManifestData).AsObject();
 				CbObject ZenServerObject = ManifestObject["zenserver"].AsObject();
