IoStore: Refactor to prepare for staging without a global optimization pass
* Moved package redirection out of the package store optimizer
* Moved iocontainer header creation out of the package store optimizer
* Moved package load order sorting out of the package store optimizer

#rnx

diff --git a/Engine/Source/Developer/IoStoreUtilities/Internal/PackageStoreOptimizer.h b/Engine/Source/Developer/IoStoreUtilities/Internal/PackageStoreOptimizer.h
--- a/Engine/Source/Developer/IoStoreUtilities/Internal/PackageStoreOptimizer.h
+++ b/Engine/Source/Developer/IoStoreUtilities/Internal/PackageStoreOptimizer.h
@@ -89,71 +89,16 @@ public:
 		return Id;
 	}
 
-	uint32 GetLoadOrder() const
-	{
-		return LoadOrder;
-	}
-
 	uint64 GetHeaderSize() const
 	{
 		return HeaderSize;
 	}
 
-	uint64 GetImportMapSize() const
-	{
-		return ImportMapSize;
-	}
-
-	uint64 GetExportMapSize() const
-	{
-		return ExportMapSize;
-	}
-
-	uint64 GetExportBundleCount() const
-	{
-		return GraphData.ExportBundles.Num();
-	}
-
-	uint64 GetExportBundleEntriesSize() const
-	{
-		return ExportBundleEntriesSize;
-	}
-	
-	uint64 GetGraphDataSize() const
-	{
-		return GraphDataSize;
-	}
-
 	uint64 GetNameCount() const
 	{
 		return NameMapBuilder.GetNameMap().Num();
 	}
 
-	uint64 GetNameMapSize() const
-	{
-		return NameMapSize;
-	}
-
-	const TArray<FPackageId>& GetImportedPackageIds() const
-	{
-		return ImportedPackageIds;
-	}
-
-	void RedirectFrom(FName SourcePackageName)
-	{
-		SourceName = SourcePackageName;
-	}
-
-	void AddShaderMapHash(const FSHAHash& ShaderMapHash)
-	{
-		ShaderMapHashes.Add(ShaderMapHash);
-	}
-
-	const TSet<FSHAHash>& GetShaderMapHashes() const
-	{
-		return ShaderMapHashes;
-	}
-
 	const bool HasEditorData() const
 	{
 		return (PackageFlags & PKG_FilterEditorOnly) == 0;
@@ -266,8 +211,6 @@ private:
 
 	FPackageId Id;
 	FName Name;
-	FName SourceName;
-	FString Region;
 
 	TOptional<FZenPackageVersioningInfo> VersioningInfo;
 
@@ -279,7 +222,6 @@ private:
 
 	TArray<FPackageId> ImportedPackageIds;
 	TArray<uint64> ImportedPublicExportHashes;
-	TSet<FSHAHash> ShaderMapHashes;
 	TArray<FBulkDataMapEntry> BulkDataEntries;
 
 	FIoBuffer HeaderBuffer;
@@ -288,15 +230,6 @@ private:
 	uint32 CookedHeaderSize = 0;
 	uint64 HeaderSize = 0;
 	uint64 ExportsSerialSize = 0;
-	uint64 ImportedPublicExportHashesSize = 0;
-	uint64 ImportMapSize = 0;
-	uint64 ExportMapSize = 0;
-	uint64 ExportBundleEntriesSize = 0;
-	uint64 GraphDataSize = 0;
-	uint64 NameMapSize = 0;
-	uint64 VersioningInfoSize = 0;
-	uint64 BulkDataMapSize = 0;
-	uint32 LoadOrder = 0;
 	
 	TArray<FExportBundleGraphNode> ExportBundleGraphNodes;
 	FPackageStorePackage::FExportBundle* CurrentBundle = nullptr;
@@ -304,39 +237,12 @@ private:
 	bool bTemporaryMark = false;
 	bool bPermanentMark = false;
 
-	bool bIsRedirected = false;
-
 	friend class FPackageStoreOptimizer;
 };
 
 class FPackageStoreOptimizer
 {
 public:
-	uint64 GetTotalPackageCount() const
-	{
-		return TotalPackageCount;
-	}
-
-	uint64 GetTotalExportBundleCount() const
-	{
-		return TotalExportBundleCount;
-	}
-
-	uint64 GetTotalExportBundleEntryCount() const
-	{
-		return TotalExportBundleEntryCount;
-	}
-
-	uint64 GetTotalInternalBundleArcsCount() const
-	{
-		return TotalInternalBundleArcsCount;
-	}
-
-	uint64 GetTotalExternalBundleArcsCount() const
-	{
-		return TotalExternalBundleArcsCount;
-	}
-
 	uint64 GetTotalScriptObjectCount() const
 	{
 		return TotalScriptObjectCount;
@@ -348,14 +254,11 @@ public:
 	FPackageStorePackage* CreateMissingPackage(const FName& Name) const;
 	FPackageStorePackage* CreatePackageFromCookedHeader(const FName& Name, const FIoBuffer& CookedHeaderBuffer) const;
 	FPackageStorePackage* CreatePackageFromZenPackageHeader(const FName& Name, const FIoBuffer& Buffer, int32 ExportBundleCount, const TArrayView<const FPackageId>& ImportedPackageIds) const;
-	void FinalizePackage(FPackageStorePackage* Package);
+	void FinalizePackage(FPackageStorePackage* Package) const;
 	FIoBuffer CreatePackageBuffer(const FPackageStorePackage* Package, const FIoBuffer& CookedExportsBuffer, TArray<FFileRegion>* InOutFileRegions) const;
 	FPackageStoreEntryResource CreatePackageStoreEntry(const FPackageStorePackage* Package, const FPackageStorePackage* OptionalSegmentPackage) const;
-	FIoContainerHeader CreateContainerHeader(const FIoContainerId& ContainerId, TArrayView<const FPackageStoreEntryResource> PackageStoreEntries) const;
-	FIoContainerHeader CreateOptionalContainerHeader(const FIoContainerId& ContainerId, TArrayView<const FPackageStoreEntryResource> PackageStoreEntries) const;
 	IOSTOREUTILITIES_API FIoBuffer CreateScriptObjectsBuffer() const;
 	void LoadScriptObjectsBuffer(const FIoBuffer& ScriptObjectsBuffer);
-	void ProcessRedirects(const TMap<FPackageId, FPackageStorePackage*>& PackagesMap, bool bIsBuildingDLC) const;
 	void OptimizeExportBundles(const TMap<FPackageId, FPackageStorePackage*>& PackagesMap);
 
 private:
@@ -414,18 +317,10 @@ private:
 	TArray<FPackageStorePackage::FExportGraphNode*> SortExportGraphNodesInLoadOrder(FPackageStorePackage* Package, FExportGraphEdges& Edges) const;
 	TArray<FPackageStorePackage::FExportBundleGraphNode*> SortExportBundleGraphNodesInLoadOrder(const TArray<FPackageStorePackage*>& Packages, FExportBundleGraphEdges& Edges) const;
 	void CreateExportBundles(FPackageStorePackage* Package) const;
-	bool VerifyRedirect(const FPackageStorePackage* SourcePackage, FPackageStorePackage& TargetPackage, bool bIsBuildingDLC) const;
 	void FinalizePackageHeader(FPackageStorePackage* Package) const;
 	void FindScriptObjectsRecursive(FPackageObjectIndex OuterIndex, UObject* Object);
 	void FindScriptObjects();
-	FIoContainerHeader CreateContainerHeaderInternal(const FIoContainerId& ContainerId, TArrayView<const FPackageStoreEntryResource> PackageStoreEntries, bool bIsOptional) const;
 
 	TMap<FPackageObjectIndex, FScriptObjectData> ScriptObjectsMap;
-	uint64 TotalPackageCount = 0;
-	uint64 TotalExportBundleCount = 0;
-	uint64 TotalExportBundleEntryCount = 0;
-	uint64 TotalInternalBundleArcsCount = 0;
-	uint64 TotalExternalBundleArcsCount = 0;
 	uint64 TotalScriptObjectCount = 0;
-	uint32 NextLoadOrder = 0;
 };
diff --git a/Engine/Source/Developer/IoStoreUtilities/Private/IoStoreUtilities.cpp b/Engine/Source/Developer/IoStoreUtilities/Private/IoStoreUtilities.cpp
--- a/Engine/Source/Developer/IoStoreUtilities/Private/IoStoreUtilities.cpp
+++ b/Engine/Source/Developer/IoStoreUtilities/Private/IoStoreUtilities.cpp
@@ -578,7 +578,7 @@ struct FShaderInfo
 	FIoBuffer CodeIoBuffer;
 	// the smaller the order, the more likely this will be loaded
 	uint32 LoadOrderFactor = MAX_uint32;
-	TSet<struct FLegacyCookedPackage*> ReferencedByPackages;
+	TSet<struct FCookedPackage*> ReferencedByPackages;
 	TMap<FContainerTargetSpec*, EShaderType> TypeInContainer;
 
 	// This is used to ensure build determinism and such must be stable
@@ -594,23 +594,32 @@ struct FShaderInfo
 	}
 };
 
-struct FLegacyCookedPackage
+struct FCookedPackage
 {
 	FPackageId GlobalPackageId;
 	FName PackageName;
-	FName RedirectFromPackageName;
-	FString FileName;
+	FName SourcePackageName; // Source package name for redirected package
+	TArray<FShaderInfo*> Shaders;
+	TArray<FSHAHash> ShaderMapHashes;
 	uint64 UAssetSize = 0;
-	uint64 UExpSize = 0;
-	FString OptionalSegmentFileName;
 	uint64 OptionalSegmentUAssetSize = 0;
-	uint64 OptionalSegmentUExpSize = 0;
+	uint64 UExpSize = 0;
 	uint64 TotalBulkDataSize = 0;
-	uint64 DiskLayoutOrder = MAX_uint64;
+	uint64 LoadOrder = MAX_uint64; // Ordered by dependencies
+	uint64 DiskLayoutOrder = MAX_uint64; // Final order after considering input order files
+	FPackageStoreEntryResource PackageStoreEntry;
 	FPackageStorePackage* OptimizedPackage = nullptr;
 	FPackageStorePackage* OptimizedOptionalSegmentPackage = nullptr;
-	TArray<FShaderInfo*> Shaders;
-	const FPackageStoreEntryResource* PackageStoreEntry = nullptr;
+	int32 PreOrderNumber = -1; // Used for sorting in load order
+	bool bPermanentMark = false; // Used for sorting in load order
+	bool bIsLocalized = false;
+};
+
+struct FLegacyCookedPackage
+	: public FCookedPackage
+{
+	FString FileName;
+	FString OptionalSegmentFileName;
 };
 
 enum class EContainerChunkType
@@ -629,7 +638,7 @@ enum class EContainerChunkType
 struct FContainerTargetFile
 {
 	FContainerTargetSpec* ContainerTarget = nullptr;
-	FLegacyCookedPackage* Package = nullptr;
+	FCookedPackage* Package = nullptr;
 	FString NormalizedSourcePath;
 	TOptional<FIoBuffer> SourceBuffer;
 	FString DestinationPath;
@@ -903,7 +912,7 @@ struct FContainerTargetSpec
 	TArray<FContainerTargetFile> TargetFiles;
 	TArray<TUniquePtr<FIoStoreReader>> PatchSourceReaders;
 	EIoContainerFlags ContainerFlags = EIoContainerFlags::None;
-	TArray<FLegacyCookedPackage*> Packages;
+	TArray<FCookedPackage*> Packages;
 	TSet<FShaderInfo*> GlobalShaders;
 	TSet<FShaderInfo*> SharedShaders;
 	TSet<FShaderInfo*> UniqueShaders;
@@ -911,8 +920,8 @@ struct FContainerTargetSpec
 	bool bGenerateDiffPatch = false;
 };
 
-using FPackageNameMap = TMap<FName, FLegacyCookedPackage*>;
-using FPackageIdMap = TMap<FPackageId, FLegacyCookedPackage*>;
+using FPackageNameMap = TMap<FName, FCookedPackage*>;
+using FPackageIdMap = TMap<FPackageId, FCookedPackage*>;
 
 class FChunkEntryCsv
 {
@@ -960,7 +969,102 @@ private:
 	TUniquePtr<FArchive> OutputArchive;
 };
 
+void SortPackagesInLoadOrderRecursive(TArray<FCookedPackage*>& Result, FCookedPackage* Package, TArray<FCookedPackage*>& S, TArray<FCookedPackage*>& P, int32& C, const TMap<FCookedPackage*, TArray<FCookedPackage*>>& ReverseEdges)
+{
+	Package->PreOrderNumber = C;
+	++C;
+	S.Push(Package);
+	P.Push(Package);
+	const TArray<FCookedPackage*>* FindParents = ReverseEdges.Find(Package);
+	if (FindParents)
+	{
+		for (FCookedPackage* Parent : *FindParents)
+		{
+			if (!Parent->bPermanentMark)
+			{
+				if (Parent->PreOrderNumber < 0)
+				{
+					SortPackagesInLoadOrderRecursive(Result, Parent, S, P, C, ReverseEdges);
+				}
+				else
+				{
+					while (P.Top()->PreOrderNumber > Parent->PreOrderNumber)
+					{
+						P.Pop();
+					}
+				}
+			}
+		}
+	}
+	if (P.Top() == Package)
+	{
+		FCookedPackage* InStronglyConnectedComponent;
+		do
+		{
+			InStronglyConnectedComponent = S.Top();
+			S.Pop();
+			InStronglyConnectedComponent->bPermanentMark = true;
+			Result.Add(InStronglyConnectedComponent);
+		} while (InStronglyConnectedComponent != Package);
+		P.Pop();
+	}
+}
+
+void SortPackagesInLoadOrder(TArray<FCookedPackage*>& Packages, const TMap<FPackageId, FCookedPackage*>& PackagesMap)
+{
+	TRACE_CPUPROFILER_EVENT_SCOPE(SortPackagesInLoadOrder);
+	Algo::Sort(Packages, [](const FCookedPackage* A, const FCookedPackage* B)
+		{
+			return A->GlobalPackageId < B->GlobalPackageId;
+		});
 
+	TMap<FCookedPackage*, TArray<FCookedPackage*>> ReverseEdges;
+	for (FCookedPackage* Package : Packages)
+	{
+		for (FPackageId ImportedPackageId : Package->PackageStoreEntry.ImportedPackageIds)
+		{
+			FCookedPackage* FindImportedPackage = PackagesMap.FindRef(ImportedPackageId);
+			if (FindImportedPackage)
+			{
+				TArray<FCookedPackage*>& SourceArray = ReverseEdges.FindOrAdd(FindImportedPackage);
+				SourceArray.Add(Package);
+			}
+		}
+	}
+	for (auto& KV : ReverseEdges)
+	{
+		TArray<FCookedPackage*>& SourceArray = KV.Value;
+		Algo::Sort(SourceArray, [](const FCookedPackage* A, const FCookedPackage* B)
+			{
+				return A->GlobalPackageId < B->GlobalPackageId;
+			});
+	}
+
+	// Path based strongly connected components + topological sort of the components
+	TArray<FCookedPackage*> Result;
+	Result.Reserve(Packages.Num());
+	TArray<int32> PackagePreOrderNumbers;
+	TArray<FCookedPackage*> S;
+	TArray<FCookedPackage*> P;
+	for (FCookedPackage* Package : Packages)
+	{
+		if (!Package->bPermanentMark)
+		{
+			S.Reset();
+			P.Reset();
+			int32 C = 0;
+			SortPackagesInLoadOrderRecursive(Result, Package, S, P, C, ReverseEdges);
+		}
+	}
+	check(Result.Num() == Packages.Num());
+	Algo::Reverse(Result);
+	Swap(Packages, Result);
+	uint64 LoadOrder = 0;
+	for (FCookedPackage* Package : Packages)
+	{
+		Package->LoadOrder = LoadOrder++;
+	}
+}
 
 class FClusterStatsCsv
 {
@@ -1024,7 +1128,7 @@ FClusterStatsCsv ClusterStatsCsv;
 //			Cluster packages B, D, then A, C
 //			Then reassemble clusters A, C, B, D
 static void AssignPackagesDiskOrder(
-	const TArray<FLegacyCookedPackage*>& Packages,
+	const TArray<FCookedPackage*>& Packages,
 	const TArray<FFileOrderMap>& OrderMaps,
 	const FPackageIdMap& PackageIdMap,
 	bool bClusterByOrderFilePriority
@@ -1034,7 +1138,7 @@ static void AssignPackagesDiskOrder(
 
 	struct FCluster
 	{
-		TArray<FLegacyCookedPackage*> Packages;
+		TArray<FCookedPackage*> Packages;
 		int32 OrderFileIndex; // Index in OrderMaps of the FFileOrderMap which contained Packages.Last()
 		int32 ClusterSequence;
 
@@ -1046,16 +1150,16 @@ static void AssignPackagesDiskOrder(
 	};
 
 	TArray<FCluster*> Clusters;
-	TSet<FLegacyCookedPackage*> AssignedPackages;
-	TArray<FLegacyCookedPackage*> ProcessStack;
+	TSet<FCookedPackage*> AssignedPackages;
+	TArray<FCookedPackage*> ProcessStack;
 
 	struct FPackageAndOrder
 	{
-		FLegacyCookedPackage* Package = nullptr;
+		FCookedPackage* Package = nullptr;
 		int64 LocalOrder;
 		const FFileOrderMap* OrderMap;
 
-		FPackageAndOrder(FLegacyCookedPackage* InPackage, int64 InLocalOrder, const FFileOrderMap* InOrderMap)
+		FPackageAndOrder(FCookedPackage* InPackage, int64 InLocalOrder, const FFileOrderMap* InOrderMap)
 			: Package(InPackage)
 			, LocalOrder(InLocalOrder)
 			, OrderMap(InOrderMap)
@@ -1079,12 +1183,12 @@ static void AssignPackagesDiskOrder(
 
 	TArray<FPackageAndOrder> SortedPackages;
 	SortedPackages.Reserve(Packages.Num());
-	for (FLegacyCookedPackage* Package : Packages)
+	for (FCookedPackage* Package : Packages)
 	{
 		// Default to the fallback order map 
 		// Reverse the bundle load order for the fallback map (so that packages are considered before their imports)
 		const FFileOrderMap* UsedOrderMap = &FallbackOrderMap;
-		int64 LocalOrder = -int64(Package->OptimizedPackage->GetLoadOrder());
+		int64 LocalOrder = -int64(Package->LoadOrder);
 
 		for (const FFileOrderMap* OrderMap : PriorityOrderMaps)
 		{
@@ -1141,12 +1245,12 @@ static void AssignPackagesDiskOrder(
 	// Keep these containers outside of inner loops to reuse allocated memory across iterations
 	// No need to allocate & free them on every single iteration, just reset element count before using them
 	TSet<FCluster*> ClustersToRead;
-	TSet<FLegacyCookedPackage*> VisitedDeps;
-	TArray<FLegacyCookedPackage*> DepQueue;
+	TSet<FCookedPackage*> VisitedDeps;
+	TArray<FCookedPackage*> DepQueue;
 	TArray<FCluster*> OrderedClustersToRead;
 
 	int32 ClusterSequence = 0;
-	TMap<FLegacyCookedPackage*, FCluster*> PackageToCluster;
+	TMap<FCookedPackage*, FCluster*> PackageToCluster;
 	for (FPackageAndOrder& Entry : SortedPackages)
 	{
 		checkSlow(Entry.OrderMap); // Without this, Entry.OrderMap != LastBlameOrderMap convinces static analysis that Entry.OrderMap may be null
@@ -1174,7 +1278,7 @@ static void AssignPackagesDiskOrder(
 			int64 ClusterBytes = 0;
 			while (ProcessStack.Num())
 			{
-				FLegacyCookedPackage* PackageToProcess = ProcessStack.Pop(false);
+				FCookedPackage* PackageToProcess = ProcessStack.Pop(false);
 				if (!AssignedPackages.Contains(PackageToProcess))
 				{
 					AssignedPackages.Add(PackageToProcess);
@@ -1185,10 +1289,10 @@ static void AssignPackagesDiskOrder(
 					AssignedBulkSize += PackageToProcess->TotalBulkDataSize;
 					
 					TArray<FPackageId> AllReferencedPackageIds;
-					AllReferencedPackageIds.Append(PackageToProcess->OptimizedPackage->GetImportedPackageIds());
+					AllReferencedPackageIds.Append(PackageToProcess->PackageStoreEntry.ImportedPackageIds);
 					for (const FPackageId& ReferencedPackageId : AllReferencedPackageIds)
 					{
-						FLegacyCookedPackage* FindReferencedPackage = PackageIdMap.FindRef(ReferencedPackageId);
+						FCookedPackage* FindReferencedPackage = PackageIdMap.FindRef(ReferencedPackageId);
 						if (FindReferencedPackage)
 						{
 							ProcessStack.Push(FindReferencedPackage);
@@ -1197,7 +1301,7 @@ static void AssignPackagesDiskOrder(
 				}
 			}
 
-			for (FLegacyCookedPackage* Package : Cluster->Packages)
+			for (FCookedPackage* Package : Cluster->Packages)
 			{
 				int64 BytesToRead = 0;
 
@@ -1208,7 +1312,7 @@ static void AssignPackagesDiskOrder(
 				DepQueue.Push(Package);
 				while (DepQueue.Num() > 0)
 				{
-					FLegacyCookedPackage* Cursor = DepQueue.Pop();
+					FCookedPackage* Cursor = DepQueue.Pop();
 					if( VisitedDeps.Contains(Cursor) == false)
 					{
 						VisitedDeps.Add(Cursor);
@@ -1218,9 +1322,9 @@ static void AssignPackagesDiskOrder(
 							ClustersToRead.Add(ReadCluster);
 						}
 						
-						for (const FPackageId& ImportedPackageId : Cursor->OptimizedPackage->GetImportedPackageIds())
+						for (const FPackageId& ImportedPackageId : Cursor->PackageStoreEntry.ImportedPackageIds)
 						{
-							FLegacyCookedPackage* FindReferencedPackage = PackageIdMap.FindRef(ImportedPackageId);
+							FCookedPackage* FindReferencedPackage = PackageIdMap.FindRef(ImportedPackageId);
 							if (FindReferencedPackage)
 							{
 								DepQueue.Push(FindReferencedPackage);
@@ -1261,16 +1365,16 @@ static void AssignPackagesDiskOrder(
 	
 	for (FCluster* Cluster : Clusters)
 	{
-		Algo::Sort(Cluster->Packages, [](const FLegacyCookedPackage* A, const FLegacyCookedPackage* B)
+		Algo::Sort(Cluster->Packages, [](const FCookedPackage* A, const FCookedPackage* B)
 		{
-			return A->OptimizedPackage->GetLoadOrder() < B->OptimizedPackage->GetLoadOrder();
+			return A->LoadOrder < B->LoadOrder;
 		});
 	}
 
 	uint64 LayoutIndex = 0;
 	for (FCluster* Cluster : Clusters)
 	{
-		for (FLegacyCookedPackage* Package : Cluster->Packages)
+		for (FCookedPackage* Package : Cluster->Packages)
 		{
 			Package->DiskLayoutOrder = LayoutIndex++;
 		}
@@ -1282,7 +1386,7 @@ static void AssignPackagesDiskOrder(
 
 static void CreateDiskLayout(
 	const TArray<FContainerTargetSpec*>& ContainerTargets,
-	const TArray<FLegacyCookedPackage*>& Packages,
+	const TArray<FCookedPackage*>& Packages,
 	const TArray<FFileOrderMap>& OrderMaps,
 	const FPackageIdMap& PackageIdMap,
 	bool bClusterByOrderFilePriority)
@@ -1428,18 +1532,18 @@ FContainerTargetSpec* AddContainer(
 	return ContainerTargetSpec;
 }
 
-FLegacyCookedPackage& FindOrAddPackage(
+FCookedPackage& FindOrAddPackage(
 	const FIoStoreArguments& Arguments,
 	const FName& PackageName,
-	TArray<FLegacyCookedPackage*>& Packages,
+	TArray<FCookedPackage*>& Packages,
 	FPackageNameMap& PackageNameMap,
 	FPackageIdMap& PackageIdMap)
 {
-	FLegacyCookedPackage* Package = PackageNameMap.FindRef(PackageName);
+	FCookedPackage* Package = PackageNameMap.FindRef(PackageName);
 	if (!Package)
 	{
 		FPackageId PackageId = FPackageId::FromName(PackageName);
-		if (FLegacyCookedPackage* FindById = PackageIdMap.FindRef(PackageId))
+		if (FCookedPackage* FindById = PackageIdMap.FindRef(PackageId))
 		{
 			UE_LOG(LogIoStore, Fatal, TEXT("Package name hash collision \"%s\" and \"%s"), *FindById->PackageName.ToString(), *PackageName.ToString());
 		}
@@ -1449,7 +1553,7 @@ FLegacyCookedPackage& FindOrAddPackage(
 			UE_LOG(LogIoStore, Fatal, TEXT("Package name hash collision with base game package \"%s\" and \"%s"), *ReleasedPackageName->ToString(), *PackageName.ToString());
 		}
 
-		Package = new FLegacyCookedPackage();
+		Package = new FCookedPackage();
 		Package->PackageName = PackageName;
 		Package->GlobalPackageId = PackageId;
 		Packages.Add(Package);
@@ -1460,10 +1564,10 @@ FLegacyCookedPackage& FindOrAddPackage(
 	return *Package;
 }
 
-FLegacyCookedPackage* FindOrAddPackageFromFileName(
+FLegacyCookedPackage* FindOrAddLegacyPackage(
 	const FIoStoreArguments& Arguments,
 	const TCHAR* FileName,
-	TArray<FLegacyCookedPackage*>& Packages,
+	TArray<FCookedPackage*>& Packages,
 	FPackageNameMap& PackageNameMap,
 	FPackageIdMap& PackageIdMap)
 {
@@ -1473,10 +1577,31 @@ FLegacyCookedPackage* FindOrAddPackageFromFileName(
 		return nullptr;
 	}
 
-	return &FindOrAddPackage(Arguments, PackageName, Packages, PackageNameMap, PackageIdMap);
+	FCookedPackage* Package = PackageNameMap.FindRef(PackageName);
+	if (!Package)
+	{
+		FPackageId PackageId = FPackageId::FromName(PackageName);
+		if (FCookedPackage* FindById = PackageIdMap.FindRef(PackageId))
+		{
+			UE_LOG(LogIoStore, Fatal, TEXT("Package name hash collision \"%s\" and \"%s"), *FindById->PackageName.ToString(), *PackageName.ToString());
+		}
+
+		if (const FName* ReleasedPackageName = Arguments.ReleasedPackages.PackageIdToName.Find(PackageId))
+		{
+			UE_LOG(LogIoStore, Fatal, TEXT("Package name hash collision with base game package \"%s\" and \"%s"), *ReleasedPackageName->ToString(), *PackageName.ToString());
+		}
+
+		Package = new FLegacyCookedPackage();
+		Package->PackageName = PackageName;
+		Package->GlobalPackageId = PackageId;
+		Packages.Add(Package);
+		PackageNameMap.Add(PackageName, Package);
+		PackageIdMap.Add(PackageId, Package);
+	}
+	return static_cast<FLegacyCookedPackage*>(Package);
 }
 
-static void ParsePackageAssetsFromFiles(TArray<FLegacyCookedPackage*>& Packages, const FPackageStoreOptimizer& PackageStoreOptimizer)
+static void ParsePackageAssetsFromFiles(TArray<FCookedPackage*>& Packages, const FPackageStoreOptimizer& PackageStoreOptimizer)
 {
 	IOSTORE_CPU_SCOPE(ParsePackageAssetsFromFiles);
 	UE_LOG(LogIoStore, Display, TEXT("Parsing packages..."));
@@ -1506,10 +1631,11 @@ static void ParsePackageAssetsFromFiles(TArray<FLegacyCookedPackage*>& Packages,
 
 		uint64 TotalUAssetSize = 0;
 		uint64 TotalOptionalSegmentUAssetSize = 0;
-		for (const FLegacyCookedPackage* Package : Packages)
+		for (const FCookedPackage* Package : Packages)
 		{
-			TotalUAssetSize += Package->UAssetSize;
-			TotalOptionalSegmentUAssetSize += Package->OptionalSegmentUAssetSize;
+			const FLegacyCookedPackage* LegacyPackage = static_cast<const FLegacyCookedPackage*>(Package);
+			TotalUAssetSize += LegacyPackage->UAssetSize;
+			TotalOptionalSegmentUAssetSize += LegacyPackage->OptionalSegmentUAssetSize;
 		}
 		UAssetMemory = reinterpret_cast<uint8*>(FMemory::Malloc(TotalUAssetSize));
 		uint8* UAssetMemoryPtr = UAssetMemory;
@@ -1518,10 +1644,11 @@ static void ParsePackageAssetsFromFiles(TArray<FLegacyCookedPackage*>& Packages,
 
 		for (int32 Index = 0; Index < TotalPackageCount; ++Index)
 		{
+			FLegacyCookedPackage* Package = static_cast<FLegacyCookedPackage*>(Packages[Index]);
 			PackageAssetBuffers[Index] = UAssetMemoryPtr;
-			UAssetMemoryPtr += Packages[Index]->UAssetSize;
+			UAssetMemoryPtr += Package->UAssetSize;
 			OptionalSegmentPackageAssetBuffers[Index] = OptionalSegmentUAssetMemoryPtr;
-			OptionalSegmentUAssetMemoryPtr += Packages[Index]->OptionalSegmentUAssetSize;
+			OptionalSegmentUAssetMemoryPtr += Package->OptionalSegmentUAssetSize;
 		}
 
 		double StartTime = FPlatformTime::Seconds();
@@ -1531,7 +1658,7 @@ static void ParsePackageAssetsFromFiles(TArray<FLegacyCookedPackage*>& Packages,
 		ParallelFor(TEXT("ReadingPackageAssets.PF"), TotalPackageCount, 1, [&ReadCount, &PackageAssetBuffers, &OptionalSegmentPackageAssetBuffers, &Packages, &CurrentFileIndex, &TotalReadCount](int32 Index)
 		{
 			TRACE_CPUPROFILER_EVENT_SCOPE(ReadUAssetFile);
-			FLegacyCookedPackage* Package = Packages[Index];
+			FLegacyCookedPackage* Package = static_cast<FLegacyCookedPackage*>(Packages[Index]);
 			if (Package->UAssetSize)
 			{
 				TotalReadCount.IncrementExchange();
@@ -1586,7 +1713,7 @@ static void ParsePackageAssetsFromFiles(TArray<FLegacyCookedPackage*>& Packages,
 				&Packages,
 				&PackageStoreOptimizer](int32 Index)
 		{
-			FLegacyCookedPackage* Package = Packages[Index];
+			FLegacyCookedPackage* Package = static_cast<FLegacyCookedPackage*>(Packages[Index]);
 
 			if (Package->UAssetSize)
 			{
@@ -1607,6 +1734,10 @@ static void ParsePackageAssetsFromFiles(TArray<FLegacyCookedPackage*>& Packages,
 				Package->OptimizedOptionalSegmentPackage = PackageStoreOptimizer.CreatePackageFromCookedHeader(Package->PackageName, OptionalSegmentCookedHeaderBuffer);
 				check(Package->OptimizedOptionalSegmentPackage->GetId() == Package->GlobalPackageId);
 			}
+
+			// The entry created here will have the correct set of imported packages but the export info will be updated later
+			Package->PackageStoreEntry = PackageStoreOptimizer.CreatePackageStoreEntry(Package->OptimizedPackage, Package->OptimizedOptionalSegmentPackage);
+
 		}, EParallelForFlags::Unbalanced);
 	}
 
@@ -1614,7 +1745,7 @@ static void ParsePackageAssetsFromFiles(TArray<FLegacyCookedPackage*>& Packages,
 	FMemory::Free(OptionalSegmentUAssetMemory);
 }
 
-static void ParsePackageAssetsFromZen(FCookedPackageStore& PackageStore, TArray<FLegacyCookedPackage*>& Packages, const FPackageStoreOptimizer& PackageStoreOptimizer)
+static void ParsePackageAssetsFromZen(FCookedPackageStore& PackageStore, TArray<FCookedPackage*>& Packages, const FPackageStoreOptimizer& PackageStoreOptimizer)
 {
 	IOSTORE_CPU_SCOPE(ParsePackageAssetsFromZen);
 	UE_LOG(LogIoStore, Display, TEXT("Parsing packages..."));
@@ -1626,19 +1757,18 @@ static void ParsePackageAssetsFromZen(FCookedPackageStore& PackageStore, TArray<
 		&PackageStoreOptimizer,
 		&Packages](int32 Index)
 	{
-		FLegacyCookedPackage* Package = Packages[Index];
+		FCookedPackage* Package = Packages[Index];
 		TIoStatusOr<FIoBuffer> HeaderBuffer = PackageStore.ReadPackageHeaderFromZen(Package->GlobalPackageId, 0);
 		Package->UExpSize = Package->UAssetSize - HeaderBuffer.ValueOrDie().DataSize();
 		Package->UAssetSize = HeaderBuffer.ValueOrDie().DataSize();
-		const FPackageStoreEntryResource* PackageStoreEntry = Package->PackageStoreEntry;
-		Package->OptimizedPackage = PackageStoreOptimizer.CreatePackageFromZenPackageHeader(Package->PackageName, HeaderBuffer.ValueOrDie(), PackageStoreEntry->ExportInfo.ExportBundleCount, PackageStoreEntry->ImportedPackageIds);
+		const FPackageStoreEntryResource& PackageStoreEntry = Package->PackageStoreEntry;
+		Package->OptimizedPackage = PackageStoreOptimizer.CreatePackageFromZenPackageHeader(Package->PackageName, HeaderBuffer.ValueOrDie(), PackageStoreEntry.ExportInfo.ExportBundleCount, PackageStoreEntry.ImportedPackageIds);
 		check(Package->OptimizedPackage->GetId() == Package->GlobalPackageId);
 		if (Package->OptionalSegmentUAssetSize)
 		{
 			TIoStatusOr<FIoBuffer> OptionalSegmentHeaderBuffer = PackageStore.ReadPackageHeaderFromZen(Package->GlobalPackageId, 1);
-			Package->OptionalSegmentUExpSize = Package->OptionalSegmentUAssetSize - OptionalSegmentHeaderBuffer.ValueOrDie().DataSize();
 			Package->OptionalSegmentUAssetSize = OptionalSegmentHeaderBuffer.ValueOrDie().DataSize();
-			Package->OptimizedOptionalSegmentPackage = PackageStoreOptimizer.CreatePackageFromZenPackageHeader(Package->PackageName, OptionalSegmentHeaderBuffer.ValueOrDie(), PackageStoreEntry->OptionalSegmentExportInfo.ExportBundleCount, PackageStoreEntry->OptionalSegmentImportedPackageIds);
+			Package->OptimizedOptionalSegmentPackage = PackageStoreOptimizer.CreatePackageFromZenPackageHeader(Package->PackageName, OptionalSegmentHeaderBuffer.ValueOrDie(), PackageStoreEntry.OptionalSegmentExportInfo.ExportBundleCount, PackageStoreEntry.OptionalSegmentImportedPackageIds);
 			check(Package->OptimizedOptionalSegmentPackage->GetId() == Package->GlobalPackageId);
 		}
 	}, EParallelForFlags::Unbalanced);
@@ -2057,7 +2187,7 @@ static void ProcessShaderLibraries(const FIoStoreArguments& Arguments, TArray<FC
 		IOSTORE_CPU_SCOPE(UpdatePackageStoreShaders);
 		for (FContainerTargetSpec* ContainerTarget : ContainerTargets)
 		{
-			for (FLegacyCookedPackage* Package : ContainerTarget->Packages)
+			for (FCookedPackage* Package : ContainerTarget->Packages)
 			{
 				TSet<FSHAHash>* FindShaderMapHashes = PackageNameToShaderMaps.Find(Package->PackageName);
 				if (FindShaderMapHashes)
@@ -2070,7 +2200,7 @@ static void ProcessShaderLibraries(const FIoStoreArguments& Arguments, TArray<FC
 							UE_LOG(LogIoStore, Warning, TEXT("Package '%s' in '%s' referencing missing shader map '%s'"), *Package->PackageName.ToString(), *ContainerTarget->Name.ToString(), *ShaderMapHash.ToString());
 							continue;
 						}
-						Package->OptimizedPackage->AddShaderMapHash(ShaderMapHash);
+						Package->ShaderMapHashes.Add(ShaderMapHash);
 						for (const FIoChunkId& ShaderChunkId : *FindChunkIds)
 						{
 							FShaderInfo* ShaderInfo = ChunkIdToShaderInfoMap.FindRef(ShaderChunkId);
@@ -2086,6 +2216,7 @@ static void ProcessShaderLibraries(const FIoStoreArguments& Arguments, TArray<FC
 						}
 					}
 				}
+				Algo::Sort(Package->ShaderMapHashes);
 			}
 		}
 	}
@@ -2108,7 +2239,7 @@ static void ProcessShaderLibraries(const FIoStoreArguments& Arguments, TArray<FC
 			const TSet<FShaderInfo*>* FindContainerShaderLibraryShaders = AllContainerShaderLibraryShadersMap.Find(ContainerTarget);
 			if (FindContainerShaderLibraryShaders)
 			{
-				for (FLegacyCookedPackage* Package : ContainerTarget->Packages)
+				for (FCookedPackage* Package : ContainerTarget->Packages)
 				{
 					for (FShaderInfo* ShaderInfo : Package->Shaders)
 					{
@@ -2155,7 +2286,7 @@ static void ProcessShaderLibraries(const FIoStoreArguments& Arguments, TArray<FC
 
 void InitializeContainerTargetsAndPackages(
 	const FIoStoreArguments& Arguments,
-	TArray<FLegacyCookedPackage*>& Packages,
+	TArray<FCookedPackage*>& Packages,
 	FPackageNameMap& PackageNameMap,
 	FPackageIdMap& PackageIdMap,
 	TArray<FContainerTargetSpec*>& ContainerTargets,
@@ -2217,7 +2348,8 @@ void InitializeContainerTargetsAndPackages(
 		}
 		else
 		{
-			OutTargetFile.Package = FindOrAddPackageFromFileName(Arguments, *SourceFile.NormalizedPath, Packages, PackageNameMap, PackageIdMap);
+			FLegacyCookedPackage* Package = FindOrAddLegacyPackage(Arguments, *SourceFile.NormalizedPath, Packages, PackageNameMap, PackageIdMap);
+			OutTargetFile.Package = Package;
 			if (!OutTargetFile.Package)
 			{
 				UE_LOG(LogIoStore, Warning, TEXT("Failed to obtain package name from file name '%s'"), *SourceFile.NormalizedPath);
@@ -2229,9 +2361,9 @@ void InitializeContainerTargetsAndPackages(
 			case FCookedFileStatData::PackageData:
 				OutTargetFile.ChunkType = EContainerChunkType::PackageData;
 				OutTargetFile.ChunkId = CreateIoChunkId(OutTargetFile.Package->GlobalPackageId.Value(), 0, EIoChunkType::ExportBundleData);
-				OutTargetFile.Package->FileName = SourceFile.NormalizedPath; // .uasset path
-				OutTargetFile.Package->UAssetSize = OriginalCookedFileStatData->FileSize;
-				OutTargetFile.Package->UExpSize = CookedFileStatData->FileSize;
+				Package->FileName = SourceFile.NormalizedPath; // .uasset path
+				Package->UAssetSize = OriginalCookedFileStatData->FileSize;
+				Package->UExpSize = CookedFileStatData->FileSize;
 				break;
 			case FCookedFileStatData::BulkData:
 				OutTargetFile.ChunkType = EContainerChunkType::BulkData;
@@ -2241,19 +2373,18 @@ void InitializeContainerTargetsAndPackages(
 			case FCookedFileStatData::OptionalBulkData:
 				OutTargetFile.ChunkType = EContainerChunkType::OptionalBulkData;
 				OutTargetFile.ChunkId = CreateIoChunkId(OutTargetFile.Package->GlobalPackageId.Value(), 0, EIoChunkType::OptionalBulkData);
-				OutTargetFile.Package->TotalBulkDataSize += CookedFileStatData->FileSize;
+				Package->TotalBulkDataSize += CookedFileStatData->FileSize;
 				break;
 			case FCookedFileStatData::MemoryMappedBulkData:
 				OutTargetFile.ChunkType = EContainerChunkType::MemoryMappedBulkData;
 				OutTargetFile.ChunkId = CreateIoChunkId(OutTargetFile.Package->GlobalPackageId.Value(), 0, EIoChunkType::MemoryMappedBulkData);
-				OutTargetFile.Package->TotalBulkDataSize += CookedFileStatData->FileSize;
+				Package->TotalBulkDataSize += CookedFileStatData->FileSize;
 				break;
 			case FCookedFileStatData::OptionalSegmentPackageData:
 				OutTargetFile.ChunkType = EContainerChunkType::OptionalSegmentPackageData;
 				OutTargetFile.ChunkId = CreateIoChunkId(OutTargetFile.Package->GlobalPackageId.Value(), 1, EIoChunkType::ExportBundleData);
-				OutTargetFile.Package->OptionalSegmentFileName = SourceFile.NormalizedPath; // .o.uasset path
-				OutTargetFile.Package->OptionalSegmentUAssetSize = OriginalCookedFileStatData->FileSize;
-				OutTargetFile.Package->OptionalSegmentUExpSize = CookedFileStatData->FileSize;
+				Package->OptionalSegmentFileName = SourceFile.NormalizedPath; // .o.uasset path
+				Package->OptionalSegmentUAssetSize = OriginalCookedFileStatData->FileSize;
 				break;
 			case FCookedFileStatData::OptionalSegmentBulkData:
 				OutTargetFile.ChunkType = EContainerChunkType::OptionalSegmentBulkData;
@@ -2337,12 +2468,13 @@ void InitializeContainerTargetsAndPackages(
 		}
 
 		OutTargetFile.Package = &FindOrAddPackage(Arguments, PackageName, Packages, PackageNameMap, PackageIdMap);
-		OutTargetFile.Package->PackageStoreEntry = PackageStore.GetPackageStoreEntry(OutTargetFile.Package->GlobalPackageId);
-		if (!OutTargetFile.Package->PackageStoreEntry)
+		const FPackageStoreEntryResource* PackageStoreEntry = PackageStore.GetPackageStoreEntry(OutTargetFile.Package->GlobalPackageId);
+		if (!PackageStoreEntry)
 		{
 			UE_LOG(LogIoStore, Warning, TEXT("Failed to find package store entry for package: '%s'"), *PackageName.ToString());
 			return false;
 		}
+		OutTargetFile.Package->PackageStoreEntry = *PackageStoreEntry;
 
 		if (Extension == TEXT(".m.ubulk"))
 		{
@@ -2362,13 +2494,11 @@ void InitializeContainerTargetsAndPackages(
 		else if (Extension == TEXT(".uasset") || Extension == TEXT(".umap"))
 		{
 			OutTargetFile.ChunkType = EContainerChunkType::PackageData;
-			OutTargetFile.Package->FileName = SourceFile.NormalizedPath;
 			OutTargetFile.Package->UAssetSize = OutTargetFile.SourceSize;
 		}
 		else if (Extension == TEXT(".o.uasset") || Extension == TEXT(".o.umap"))
 		{
 			OutTargetFile.ChunkType = EContainerChunkType::OptionalSegmentPackageData;
-			OutTargetFile.Package->OptionalSegmentFileName = SourceFile.NormalizedPath;
 			OutTargetFile.Package->OptionalSegmentUAssetSize = OutTargetFile.SourceSize;
 		}
 		else if (Extension == TEXT(".o.ubulk"))
@@ -2475,7 +2605,7 @@ void InitializeContainerTargetsAndPackages(
 		}
 	}
 
-	Algo::Sort(Packages, [](const FLegacyCookedPackage* A, const FLegacyCookedPackage* B)
+	Algo::Sort(Packages, [](const FCookedPackage* A, const FCookedPackage* B)
 	{
 		return A->GlobalPackageId < B->GlobalPackageId;
 	});
@@ -2595,72 +2725,6 @@ void LogContainerPackageInfo(const TArray<FContainerTargetSpec*>& ContainerTarge
 		TotalPackageCount,
 		TotalLocalizedPackageCount);
 
-
-	uint64 TotalHeaderSize = 0;
-	uint64 TotalGraphSize = 0;
-	uint64 TotalExportBundleEntriesSize = 0;
-	uint64 TotalImportMapSize = 0;
-	uint64 TotalExportMapSize = 0;
-	uint64 TotalNameMapSize = 0;
-
-	UE_LOG(LogIoStore, Display, TEXT(""));
-	UE_LOG(LogIoStore, Display, TEXT(""));
-	UE_LOG(LogIoStore, Display, TEXT("--------------------------------------------------- PackageHeader (KB) --------------------------------------------------"));
-	UE_LOG(LogIoStore, Display, TEXT(""));
-	UE_LOG(LogIoStore, Display, TEXT("%-30s %13s %13s %13s %13s %13s %13s"),
-		TEXT("Container"),
-		TEXT("Header"),
-		TEXT("Graph"),
-		TEXT("ExportBundleEntries"),
-		TEXT("ImportMap"),
-		TEXT("ExportMap"),
-		TEXT("NameMap"));
-	UE_LOG(LogIoStore, Display, TEXT("-------------------------------------------------------------------------------------------------------------------------"));
-	for (const FContainerTargetSpec* ContainerTarget : ContainerTargets)
-	{
-		uint64 HeaderSize = 0;
-		uint64 GraphSize = 0;
-		uint64 ExportBundleEntriesSize = 0;
-		uint64 ImportMapSize = 0;
-		uint64 ExportMapSize = 0;
-		uint64 NameMapSize = 0;
-
-		for (const FLegacyCookedPackage* Package : ContainerTarget->Packages)
-		{
-			HeaderSize += Package->OptimizedPackage->GetHeaderSize();
-			GraphSize += Package->OptimizedPackage->GetGraphDataSize();
-			ExportBundleEntriesSize += Package->OptimizedPackage->GetExportBundleEntriesSize();
-			ImportMapSize += Package->OptimizedPackage->GetImportMapSize();
-			ExportMapSize += Package->OptimizedPackage->GetExportMapSize();
-			NameMapSize += Package->OptimizedPackage->GetNameMapSize();
-		}
-
-		UE_LOG(LogIoStore, Display, TEXT("%-30s %13.0lf %13.0lf %13.0lf %13.0lf %13.0lf %13.0lf"),
-			*ContainerTarget->Name.ToString(),
-			(double)HeaderSize / 1024.0,
-			(double)GraphSize / 1024.0,
-			(double)ExportBundleEntriesSize / 1024.0,
-			(double)ImportMapSize / 1024.0,
-			(double)ExportMapSize / 1024.0,
-			(double)NameMapSize / 1024.0);
-
-		TotalHeaderSize += HeaderSize;
-		TotalGraphSize += GraphSize;
-		TotalExportBundleEntriesSize += ExportBundleEntriesSize;
-		TotalImportMapSize += ImportMapSize;
-		TotalExportMapSize += ExportMapSize;
-		TotalNameMapSize += NameMapSize;
-	}
-
-	UE_LOG(LogIoStore, Display, TEXT("%-30s %13.0lf %13.0lf %13.0lf %13.0lf %13.0lf %13.0lf"),
-		TEXT("TOTAL"),
-		(double)TotalHeaderSize / 1024.0,
-		(double)TotalGraphSize / 1024.0,
-		(double)TotalExportBundleEntriesSize / 1024.0,
-		(double)TotalImportMapSize / 1024.0,
-		(double)TotalExportMapSize / 1024.0,
-		(double)TotalNameMapSize / 1024.0);
-
 	UE_LOG(LogIoStore, Display, TEXT(""));
 	UE_LOG(LogIoStore, Display, TEXT(""));
 }
@@ -3405,6 +3469,230 @@ public:
 	}
 };
 
+// modified copy from PakFileUtilities
+static FName RemapLocalizationPathIfNeeded(const FString& Path)
+{
+	static constexpr TCHAR L10NString[] = TEXT("/L10N/");
+	static constexpr int32 L10NPrefixLength = sizeof(L10NString) / sizeof(TCHAR) - 1;
+
+	int32 BeginL10NOffset = Path.Find(L10NString, ESearchCase::IgnoreCase);
+	if (BeginL10NOffset >= 0)
+	{
+		int32 EndL10NOffset = BeginL10NOffset + L10NPrefixLength;
+		int32 NextSlashIndex = Path.Find(TEXT("/"), ESearchCase::IgnoreCase, ESearchDir::FromStart, EndL10NOffset);
+		int32 RegionLength = NextSlashIndex - EndL10NOffset;
+		if (RegionLength >= 2)
+		{
+			FString NonLocalizedPath = Path.Mid(0, BeginL10NOffset) + Path.Mid(NextSlashIndex);
+			return FName(NonLocalizedPath);
+		}
+	}
+	return NAME_None;
+}
+
+void ProcessRedirects(const FIoStoreArguments& Arguments, const TMap<FPackageId, FCookedPackage*>& PackagesMap)
+{
+	TRACE_CPUPROFILER_EVENT_SCOPE(ProcessRedirects);
+
+	for (const auto& KV : PackagesMap)
+	{
+		FCookedPackage* Package = KV.Value;
+		FName LocalizedSourcePackageName = RemapLocalizationPathIfNeeded(Package->PackageName.ToString());
+		if (!LocalizedSourcePackageName.IsNone())
+		{
+			Package->SourcePackageName = LocalizedSourcePackageName;
+			Package->bIsLocalized = true;
+		}
+	}
+
+	const bool bIsBuildingDLC = Arguments.IsDLC();
+	if (bIsBuildingDLC && Arguments.bRemapPluginContentToGame)
+	{
+		for (const auto& KV : PackagesMap)
+		{
+			FCookedPackage* Package = KV.Value;
+			const int32 DLCNameLen = Arguments.DLCName.Len() + 1;
+			FString PackageNameStr = Package->PackageName.ToString();
+			FString RedirectedPackageNameStr = TEXT("/Game");
+			RedirectedPackageNameStr.AppendChars(*PackageNameStr + DLCNameLen, PackageNameStr.Len() - DLCNameLen);
+			FName RedirectedPackageName = FName(*RedirectedPackageNameStr);
+			Package->SourcePackageName = RedirectedPackageName;
+		}
+	}
+}
+
+void CreateContainerHeader(FContainerTargetSpec& ContainerTarget, bool bIsOptional)
+{
+	TRACE_CPUPROFILER_EVENT_SCOPE(CreateContainerHeader);
+	FIoContainerHeader& Header = bIsOptional ? ContainerTarget.OptionalSegmentHeader : ContainerTarget.Header;
+	Header.ContainerId = ContainerTarget.ContainerId;
+
+	int32 NonOptionalSegmentStoreEntriesCount = 0;
+	int32 OptionalSegmentStoreEntriesCount = 0;
+	if (bIsOptional)
+	{
+		for (const FCookedPackage* Package : ContainerTarget.Packages)
+		{
+			if (Package->PackageStoreEntry.OptionalSegmentExportInfo.ExportCount)
+			{
+				if (Package->PackageStoreEntry.IsAutoOptional())
+				{
+					// Auto optional packages fully replace the non-optional segment
+					++NonOptionalSegmentStoreEntriesCount;
+				}
+				else
+				{
+					++OptionalSegmentStoreEntriesCount;
+				}
+			}
+		}
+	}
+	else
+	{
+		NonOptionalSegmentStoreEntriesCount = ContainerTarget.Packages.Num();
+	}
+
+	struct FStoreEntriesWriter
+	{
+		const int32 StoreTocSize;
+		FLargeMemoryWriter StoreTocArchive = FLargeMemoryWriter(0, true);
+		FLargeMemoryWriter StoreDataArchive = FLargeMemoryWriter(0, true);
+
+		void Flush(TArray<uint8>& OutputBuffer)
+		{
+			check(StoreTocArchive.TotalSize() == StoreTocSize);
+			if (StoreTocSize)
+			{
+				const int32 StoreByteCount = StoreTocArchive.TotalSize() + StoreDataArchive.TotalSize();
+				OutputBuffer.AddUninitialized(StoreByteCount);
+				FBufferWriter PackageStoreArchive(OutputBuffer.GetData(), StoreByteCount);
+				PackageStoreArchive.Serialize(StoreTocArchive.GetData(), StoreTocArchive.TotalSize());
+				PackageStoreArchive.Serialize(StoreDataArchive.GetData(), StoreDataArchive.TotalSize());
+			}
+		}
+	};
+
+	FStoreEntriesWriter StoreEntriesWriter
+	{
+		static_cast<int32>(NonOptionalSegmentStoreEntriesCount * sizeof(FFilePackageStoreEntry))
+	};
+
+	FStoreEntriesWriter OptionalSegmentStoreEntriesWriter
+	{
+		static_cast<int32>(OptionalSegmentStoreEntriesCount * sizeof(FFilePackageStoreEntry))
+	};
+
+	auto SerializePackageEntryCArrayHeader = [](FStoreEntriesWriter& Writer, int32 Count)
+	{
+		const int32 RemainingTocSize = Writer.StoreTocSize - Writer.StoreTocArchive.Tell();
+		const int32 OffsetFromThis = RemainingTocSize + Writer.StoreDataArchive.Tell();
+		uint32 ArrayNum = Count > 0 ? Count : 0;
+		uint32 OffsetToDataFromThis = ArrayNum > 0 ? OffsetFromThis : 0;
+
+		Writer.StoreTocArchive << ArrayNum;
+		Writer.StoreTocArchive << OffsetToDataFromThis;
+	};
+
+	TArray<const FCookedPackage*> SortedPackages(ContainerTarget.Packages);
+	Algo::Sort(SortedPackages, [](const FCookedPackage* A, const FCookedPackage* B)
+		{
+			return A->GlobalPackageId < B->GlobalPackageId;
+		});
+
+	Header.PackageIds.Reserve(NonOptionalSegmentStoreEntriesCount);
+	Header.OptionalSegmentPackageIds.Reserve(OptionalSegmentStoreEntriesCount);
+	FPackageStoreNameMapBuilder RedirectsNameMapBuilder;
+	RedirectsNameMapBuilder.SetNameMapType(FMappedName::EType::Container);
+	TSet<FName> AllLocalizedPackages;
+	if (bIsOptional)
+	{
+		for (const FCookedPackage* Package : SortedPackages)
+		{
+			const FPackageStoreEntryResource& Entry = Package->PackageStoreEntry;
+			if (Entry.OptionalSegmentExportInfo.ExportCount)
+			{
+				FStoreEntriesWriter* TargetEntriesWriter;
+				if (Entry.IsAutoOptional())
+				{
+					Header.PackageIds.Add(Package->GlobalPackageId);
+					TargetEntriesWriter = &StoreEntriesWriter;
+				}
+				else
+				{
+					Header.OptionalSegmentPackageIds.Add(Package->GlobalPackageId);
+					TargetEntriesWriter = &OptionalSegmentStoreEntriesWriter;
+				}
+
+				// OptionalStoreEntry
+				FPackageStoreExportInfo OptionalSegmentExportInfo = Entry.OptionalSegmentExportInfo;
+				TargetEntriesWriter->StoreTocArchive << OptionalSegmentExportInfo;
+
+				// OptionalImportedPackages
+				const TArray<FPackageId>& OptionalSegmentImportedPackageIds = Entry.OptionalSegmentImportedPackageIds;
+				SerializePackageEntryCArrayHeader(*TargetEntriesWriter, OptionalSegmentImportedPackageIds.Num());
+				for (FPackageId OptionalSegmentImportedPackageId : OptionalSegmentImportedPackageIds)
+				{
+					check(OptionalSegmentImportedPackageId.IsValid());
+					TargetEntriesWriter->StoreDataArchive << OptionalSegmentImportedPackageId;
+				}
+
+				// ShaderMapHashes is N/A for optional segments
+				SerializePackageEntryCArrayHeader(*TargetEntriesWriter, 0);
+			}
+		}
+	}
+	else
+	{
+		for (const FCookedPackage* Package : SortedPackages)
+		{
+			const FPackageStoreEntryResource& Entry = Package->PackageStoreEntry;
+			Header.PackageIds.Add(Package->GlobalPackageId);
+			if (!Package->SourcePackageName.IsNone())
+			{
+				RedirectsNameMapBuilder.MarkNameAsReferenced(Package->SourcePackageName);
+				FMappedName MappedSourcePackageName = RedirectsNameMapBuilder.MapName(Package->SourcePackageName);
+				if (Package->bIsLocalized)
+				{
+					if (!AllLocalizedPackages.Contains(Package->SourcePackageName))
+					{
+						Header.LocalizedPackages.Add({ FPackageId::FromName(Package->SourcePackageName), MappedSourcePackageName });
+						AllLocalizedPackages.Add(Package->SourcePackageName);
+					}
+				}
+				else
+				{
+					Header.PackageRedirects.Add({ FPackageId::FromName(Package->SourcePackageName), Package->GlobalPackageId, MappedSourcePackageName });
+				}
+			}
+
+			// StoreEntries
+			FPackageStoreExportInfo ExportInfo = Entry.ExportInfo;
+			StoreEntriesWriter.StoreTocArchive << ExportInfo;
+
+			// ImportedPackages
+			const TArray<FPackageId>& ImportedPackageIds = Entry.ImportedPackageIds;
+			SerializePackageEntryCArrayHeader(StoreEntriesWriter, ImportedPackageIds.Num());
+			for (FPackageId ImportedPackageId : ImportedPackageIds)
+			{
+				check(ImportedPackageId.IsValid());
+				StoreEntriesWriter.StoreDataArchive << ImportedPackageId;
+			}
+
+			// ShaderMapHashes
+			const TArray<FSHAHash>& ShaderMapHashes = Package->ShaderMapHashes;
+			SerializePackageEntryCArrayHeader(StoreEntriesWriter, ShaderMapHashes.Num());
+			for (const FSHAHash& ShaderMapHash : ShaderMapHashes)
+			{
+				StoreEntriesWriter.StoreDataArchive << const_cast<FSHAHash&>(ShaderMapHash);
+			}
+		}
+	}
+	Header.RedirectsNameMap = RedirectsNameMapBuilder.GetNameMap();
+
+	StoreEntriesWriter.Flush(Header.StoreEntries);
+	OptionalSegmentStoreEntriesWriter.Flush(Header.OptionalSegmentStoreEntries);
+}
+
 int32 CreateTarget(const FIoStoreArguments& Arguments, const FIoStoreWriterSettings& GeneralIoWriterSettings)
 {
 	IOSTORE_CPU_SCOPE(CreateTarget);
@@ -3430,7 +3718,7 @@ int32 CreateTarget(const FIoStoreArguments& Arguments, const FIoStoreWriterSetti
 		UE_LOG(LogIoStore, Display, TEXT("Hash database verification on: hashes will be checked for accuracy during this run."));
 	}
 
-	TArray<FLegacyCookedPackage*> Packages;
+	TArray<FCookedPackage*> Packages;
 	FPackageNameMap PackageNameMap;
 	FPackageIdMap PackageIdMap;
 
@@ -3533,20 +3821,6 @@ int32 CreateTarget(const FIoStoreArguments& Arguments, const FIoStoreWriterSetti
 	TArray<FShaderInfo*> Shaders;
 	ProcessShaderLibraries(Arguments, ContainerTargets, Shaders);
 
-	if (Arguments.IsDLC() && Arguments.bRemapPluginContentToGame)
-	{
-		IOSTORE_CPU_SCOPE(RemapPluginContent);
-		for (FLegacyCookedPackage* Package : Packages)
-		{
-			const int32 DLCNameLen = Arguments.DLCName.Len() + 1;
-			FString PackageNameStr = Package->PackageName.ToString();
-			FString RedirectedPackageNameStr = TEXT("/Game");
-			RedirectedPackageNameStr.AppendChars(*PackageNameStr + DLCNameLen, PackageNameStr.Len() - DLCNameLen);
-			FName RedirectedPackageName = FName(*RedirectedPackageNameStr);
-			Package->OptimizedPackage->RedirectFrom(RedirectedPackageName);
-		}
-	}
-
 	{
 		IOSTORE_CPU_SCOPE(AppendNonPackageChunks);
 		for (FContainerTargetSpec* ContainerTarget : ContainerTargets)
@@ -3565,7 +3839,7 @@ int32 CreateTarget(const FIoStoreArguments& Arguments, const FIoStoreWriterSetti
 						if (TargetFile.ChunkType == EContainerChunkType::OptionalSegmentBulkData)
 						{
 							FIoChunkId ChunkId = TargetFile.ChunkId;
-							if (TargetFile.Package->OptimizedOptionalSegmentPackage && TargetFile.Package->OptimizedOptionalSegmentPackage->HasEditorData())
+							if (TargetFile.Package->PackageStoreEntry.IsAutoOptional())
 							{
 								// Auto optional packages replace the non-optional part when the container is mounted
 								ChunkId = CreateIoChunkId(TargetFile.Package->GlobalPackageId.Value(), 0, EIoChunkType::BulkData);
@@ -3582,26 +3856,24 @@ int32 CreateTarget(const FIoStoreArguments& Arguments, const FIoStoreWriterSetti
 		}
 	}
 
-	TMap<FPackageId, FPackageStorePackage*> OptimizedPackagesMap;
+	UE_LOG(LogIoStore, Display, TEXT("Processing redirects..."));
+	ProcessRedirects(Arguments, PackageIdMap);
+
+	UE_LOG(LogIoStore, Display, TEXT("Optimizing packages..."));
 	{
-		IOSTORE_CPU_SCOPE(OptimizedPackagesMap);
-		for (FLegacyCookedPackage* Package : Packages)
+		TMap<FPackageId, FPackageStorePackage*> OptimizedPackagesMap;
+		for (FCookedPackage* Package : Packages)
 		{
 			check(Package->OptimizedPackage);
 			OptimizedPackagesMap.Add(Package->OptimizedPackage->GetId(), Package->OptimizedPackage);
 		}
+		PackageStoreOptimizer.OptimizeExportBundles(OptimizedPackagesMap);
 	}
 
-	UE_LOG(LogIoStore, Display, TEXT("Processing redirects..."));
-	PackageStoreOptimizer.ProcessRedirects(OptimizedPackagesMap, Arguments.IsDLC());
-
-	UE_LOG(LogIoStore, Display, TEXT("Optimizing packages..."));
-	PackageStoreOptimizer.OptimizeExportBundles(OptimizedPackagesMap);
-
 	UE_LOG(LogIoStore, Display, TEXT("Finalizing packages..."));
 	{
 		IOSTORE_CPU_SCOPE(FinalizingPackages);
-		for (FLegacyCookedPackage* Package : Packages)
+		for (FCookedPackage* Package : Packages)
 		{
 			check(Package->OptimizedPackage);
 			PackageStoreOptimizer.FinalizePackage(Package->OptimizedPackage);
@@ -3609,6 +3881,8 @@ int32 CreateTarget(const FIoStoreArguments& Arguments, const FIoStoreWriterSetti
 			{
 				PackageStoreOptimizer.FinalizePackage(Package->OptimizedOptionalSegmentPackage);
 			}
+			// Update package store entry after optimizing export bundles
+			Package->PackageStoreEntry = PackageStoreOptimizer.CreatePackageStoreEntry(Package->OptimizedPackage, Package->OptimizedOptionalSegmentPackage);
 		}
 	}
 
@@ -3619,6 +3893,7 @@ int32 CreateTarget(const FIoStoreArguments& Arguments, const FIoStoreWriterSetti
 		IOSTORE_CPU_SCOPE(CreateClusterCSV);
 		ClusterStatsCsv.CreateOutputFile(ClusterCSVPath);
 	}
+	SortPackagesInLoadOrder(Packages, PackageIdMap);
 	CreateDiskLayout(ContainerTargets, Packages, Arguments.OrderMaps, PackageIdMap, Arguments.bClusterByOrderFilePriority);
 
 	{
@@ -3627,7 +3902,6 @@ int32 CreateTarget(const FIoStoreArguments& Arguments, const FIoStoreWriterSetti
 		{
 			if (ContainerTarget->IoStoreWriter)
 			{
-				TArray<FPackageStoreEntryResource> PackageStoreEntries;
 				for (FContainerTargetFile& TargetFile : ContainerTarget->TargetFiles)
 				{
 					if (TargetFile.ChunkType == EContainerChunkType::PackageData || TargetFile.ChunkType == EContainerChunkType::OptionalSegmentPackageData)
@@ -3652,7 +3926,6 @@ int32 CreateTarget(const FIoStoreArguments& Arguments, const FIoStoreWriterSetti
 						else
 						{
 							ContainerTarget->IoStoreWriter->Append(TargetFile.ChunkId, WriteRequestManager.Read(TargetFile), WriteOptions);
-							PackageStoreEntries.Add(PackageStoreOptimizer.CreatePackageStoreEntry(TargetFile.Package->OptimizedPackage, TargetFile.Package->OptimizedOptionalSegmentPackage));
 						}
 					}
 				}
@@ -3673,12 +3946,12 @@ int32 CreateTarget(const FIoStoreArguments& Arguments, const FIoStoreWriterSetti
 						WriteOptions);
 				};
 
-				ContainerTarget->Header = PackageStoreOptimizer.CreateContainerHeader(ContainerTarget->ContainerId, PackageStoreEntries);
+				CreateContainerHeader(*ContainerTarget, false);
 				WriteContainerHeaderChunk(ContainerTarget->Header, ContainerTarget->IoStoreWriter.Get());
 
 				if (ContainerTarget->OptionalSegmentIoStoreWriter)
 				{
-					ContainerTarget->OptionalSegmentHeader = PackageStoreOptimizer.CreateOptionalContainerHeader(ContainerTarget->ContainerId, PackageStoreEntries);
+					CreateContainerHeader(*ContainerTarget, true);
 					WriteContainerHeaderChunk(ContainerTarget->OptionalSegmentHeader, ContainerTarget->OptionalSegmentIoStoreWriter.Get());
 				}
 			}
@@ -3926,16 +4199,20 @@ int32 CreateTarget(const FIoStoreArguments& Arguments, const FIoStoreWriterSetti
 	UE_LOG(LogIoStore, Display, TEXT("Calculating stats..."));
 	uint64 UExpSize = 0;
 	uint64 UAssetSize = 0;
+	uint64 UBulkSize = 0;
+	uint64 HeaderSize = 0;
 	uint64 ImportedPackagesCount = 0;
 	uint64 NoImportedPackagesCount = 0;
 	uint64 NameMapCount = 0;
 	
-	for (const FLegacyCookedPackage* Package : Packages)
+	for (const FCookedPackage* Package : Packages)
 	{
 		UExpSize += Package->UExpSize;
 		UAssetSize += Package->UAssetSize;
+		UBulkSize += Package->TotalBulkDataSize;
 		NameMapCount += Package->OptimizedPackage->GetNameCount();
-		int32 PackageImportedPackagesCount = Package->OptimizedPackage->GetImportedPackageIds().Num();
+		HeaderSize += Package->OptimizedPackage->GetHeaderSize();
+		int32 PackageImportedPackagesCount = Package->PackageStoreEntry.ImportedPackageIds.Num();
 		ImportedPackagesCount += PackageImportedPackagesCount;
 		NoImportedPackagesCount += PackageImportedPackagesCount == 0;
 	}
@@ -3975,22 +4252,20 @@ int32 CreateTarget(const FIoStoreArguments& Arguments, const FIoStoreWriterSetti
 	LogWriterResults(IoStoreWriterResults);
 	LogContainerPackageInfo(ContainerTargets);
 	
-	UE_LOG(LogIoStore, Display, TEXT("Input:  %8.2lf MB UExp"), (double)UExpSize / 1024.0 / 1024.0);
-	UE_LOG(LogIoStore, Display, TEXT("Input:  %8.2lf MB UAsset"), (double)UAssetSize / 1024.0 / 1024.0);
 	UE_LOG(LogIoStore, Display, TEXT("Input:  %8d Packages"), Packages.Num());
+	UE_LOG(LogIoStore, Display, TEXT("Input:  %8.2lf GB UExp"), (double)UExpSize / 1024.0 / 1024.0 / 1024.0);
+	UE_LOG(LogIoStore, Display, TEXT("Input:  %8.2lf GB UAsset"), (double)UAssetSize / 1024.0 / 1024.0 / 1024.0);
+	UE_LOG(LogIoStore, Display, TEXT("Input:  %8.2lf GB UBulk"), (double)UBulkSize / 1024.0 / 1024.0 / 1024.0);
 	UE_LOG(LogIoStore, Display, TEXT("Input:  %8.2f MB for %d Global shaders"), (double)GlobalShaderSize / 1024.0 / 1024.0, GlobalShaderCount);
 	UE_LOG(LogIoStore, Display, TEXT("Input:  %8.2f MB for %d Shared shaders"), (double)SharedShaderSize / 1024.0 / 1024.0, SharedShaderCount);
 	UE_LOG(LogIoStore, Display, TEXT("Input:  %8.2f MB for %d Unique shaders"), (double)UniqueShaderSize / 1024.0 / 1024.0, UniqueShaderCount);
 	UE_LOG(LogIoStore, Display, TEXT("Input:  %8.2f MB for %d Inline shaders"), (double)InlineShaderSize / 1024.0 / 1024.0, InlineShaderCount);
 	UE_LOG(LogIoStore, Display, TEXT(""));
-	UE_LOG(LogIoStore, Display, TEXT("Output: %8llu Export bundle entries"), PackageStoreOptimizer.GetTotalExportBundleEntryCount());
-	UE_LOG(LogIoStore, Display, TEXT("Output: %8llu Export bundles"), PackageStoreOptimizer.GetTotalExportBundleCount());
-	UE_LOG(LogIoStore, Display, TEXT("Output: %8llu Internal export bundle arcs"), PackageStoreOptimizer.GetTotalInternalBundleArcsCount());
-	UE_LOG(LogIoStore, Display, TEXT("Output: %8llu External export bundle arcs"), PackageStoreOptimizer.GetTotalExternalBundleArcsCount());
 	UE_LOG(LogIoStore, Display, TEXT("Output: %8llu Name map entries"), NameMapCount);
 	UE_LOG(LogIoStore, Display, TEXT("Output: %8llu Imported package entries"), ImportedPackagesCount);
 	UE_LOG(LogIoStore, Display, TEXT("Output: %8llu Packages without imports"), NoImportedPackagesCount);
 	UE_LOG(LogIoStore, Display, TEXT("Output: %8d Public runtime script objects"), PackageStoreOptimizer.GetTotalScriptObjectCount());
+	UE_LOG(LogIoStore, Display, TEXT("Output: %8.2lf GB HeaderData"), (double)HeaderSize / 1024.0 / 1024.0 / 1024.0);
 	UE_LOG(LogIoStore, Display, TEXT("Output: %8.2lf MB InitialLoadData"), (double)InitialLoadSize / 1024.0 / 1024.0);
 
 	if (ChunkDatabase.IsValid())
diff --git a/Engine/Source/Developer/IoStoreUtilities/Private/PackageStoreOptimizer.cpp b/Engine/Source/Developer/IoStoreUtilities/Private/PackageStoreOptimizer.cpp
--- a/Engine/Source/Developer/IoStoreUtilities/Private/PackageStoreOptimizer.cpp
+++ b/Engine/Source/Developer/IoStoreUtilities/Private/PackageStoreOptimizer.cpp
@@ -71,32 +71,6 @@ EObjectMark GetExcludedObjectMarksForObject(const UObject* Object, const ITarget
 	return Marks;
 }
 
-// modified copy from PakFileUtilities
-static FString RemapLocalizationPathIfNeeded(const FString& Path, FString* OutRegion)
-{
-	static constexpr TCHAR L10NString[] = TEXT("/L10N/");
-	static constexpr int32 L10NPrefixLength = sizeof(L10NString) / sizeof(TCHAR) - 1;
-
-	int32 BeginL10NOffset = Path.Find(L10NString, ESearchCase::IgnoreCase);
-	if (BeginL10NOffset >= 0)
-	{
-		int32 EndL10NOffset = BeginL10NOffset + L10NPrefixLength;
-		int32 NextSlashIndex = Path.Find(TEXT("/"), ESearchCase::IgnoreCase, ESearchDir::FromStart, EndL10NOffset);
-		int32 RegionLength = NextSlashIndex - EndL10NOffset;
-		if (RegionLength >= 2)
-		{
-			FString NonLocalizedPath = Path.Mid(0, BeginL10NOffset) + Path.Mid(NextSlashIndex);
-			if (OutRegion)
-			{
-				*OutRegion = Path.Mid(EndL10NOffset, RegionLength);
-				OutRegion->ToLowerInline();
-			}
-			return NonLocalizedPath;
-		}
-	}
-	return Path;
-}
-
 void FPackageStoreOptimizer::Initialize()
 {
 	FindScriptObjects();
@@ -112,7 +86,6 @@ FPackageStorePackage* FPackageStoreOptimizer::CreateMissingPackage(const FName&
 	FPackageStorePackage* Package = new FPackageStorePackage();
 	Package->Name = Name;
 	Package->Id = FPackageId::FromName(Name);
-	Package->SourceName = *RemapLocalizationPathIfNeeded(Name.ToString(), &Package->Region);
 	return Package;
 }
 
@@ -121,9 +94,7 @@ FPackageStorePackage* FPackageStoreOptimizer::CreatePackageFromCookedHeader(cons
 	FPackageStorePackage* Package = new FPackageStorePackage();
 	Package->Id = FPackageId::FromName(Name);
 	Package->Name = Name;
-	FString NameStr = Name.ToString();
-	Package->SourceName = *RemapLocalizationPathIfNeeded(NameStr, &Package->Region);
-
+	
 	FCookedHeaderData CookedHeaderData = LoadCookedHeader(CookedHeaderBuffer);
 	if (!CookedHeaderData.Summary.bUnversioned)
 	{
@@ -166,8 +137,7 @@ FPackageStorePackage* FPackageStoreOptimizer::CreatePackageFromZenPackageHeader(
 	}
 
 	Package->Name = Name;
-	Package->SourceName = *RemapLocalizationPathIfNeeded(NameStr, &Package->Region);
-
+	
 	FZenPackageHeaderData ZenHeaderData = LoadZenPackageHeader(Buffer, ExportBundleCount, ImportedPackageIds);
 	if (ZenHeaderData.VersioningInfo.IsSet())
 	{
@@ -887,10 +857,6 @@ void FPackageStoreOptimizer::OptimizeExportBundles(const TMap<FPackageId, FPacka
 		check(Package);
 		if (Package->CurrentBundle == nullptr || Package != PreviousPackage)
 		{
-			if (Package->GraphData.ExportBundles.IsEmpty())
-			{
-				Package->LoadOrder = NextLoadOrder++;
-			}
 			Package->CurrentBundle = &Package->GraphData.ExportBundles.AddDefaulted_GetRef();
 		}
 		for (FPackageStorePackage::FExportGraphNode* ExportGraphNode : Node->ExportGraphNodes)
@@ -1134,159 +1100,6 @@ void FPackageStoreOptimizer::CreateExportBundles(FPackageStorePackage* Package)
 	}
 }
 
-bool FPackageStoreOptimizer::VerifyRedirect(
-	const FPackageStorePackage* SourcePackage,
-	FPackageStorePackage& TargetPackage,
-	bool bIsBuildingDLC) const
-{
-	if (!SourcePackage)
-	{
-		if (bIsBuildingDLC)
-		{
-			// We can't verify against the source package but check for presence of UStructs
-			for (const FPackageStorePackage::FExport& Export : TargetPackage.Exports)
-			{
-				if (!Export.SuperIndex.IsNull() && Export.OuterIndex.IsNull())
-				{
-					UE_LOG(LogPackageStoreOptimizer, Warning, TEXT("Skipping redirect to package '%s' due to presence of UStruct '%s'"), *TargetPackage.Name.ToString(), *Export.ObjectName.ToString());
-					return false;
-				}
-			}
-			return true;
-		}
-		return false;
-	}
-
-	const int32 ExportCount =
-		SourcePackage->Exports.Num() < TargetPackage.Exports.Num() ?
-		SourcePackage->Exports.Num() :
-		TargetPackage.Exports.Num();
-
-	UE_CLOG(SourcePackage->Exports.Num() != TargetPackage.Exports.Num(), LogPackageStoreOptimizer, Verbose,
-		TEXT("Redirection target package '%s' (0x%llX) for source package '%s' (0x%llX)  - Has ExportCount %d vs. %d"),
-		*TargetPackage.Name.ToString(),
-		TargetPackage.Id.ValueForDebugging(),
-		*SourcePackage->Name.ToString(),
-		SourcePackage->Id.ValueForDebugging(),
-		TargetPackage.Exports.Num(),
-		SourcePackage->Exports.Num());
-
-	auto AppendMismatchMessage = [&TargetPackage, SourcePackage](
-		const TCHAR* Text, FName ExportName, FPackageObjectIndex TargetIndex, FPackageObjectIndex SourceIndex, FString& FailReason)
-	{
-		FailReason.Appendf(TEXT("Public export '%s' has %s %s vs. %s"),
-			*ExportName.ToString(),
-			Text,
-			*TargetPackage.Exports[TargetIndex.ToExport()].ObjectName.ToString(),
-			*SourcePackage->Exports[SourceIndex.ToExport()].ObjectName.ToString());
-	};
-
-	TMap<uint64, const FPackageStorePackage::FExport*> SourceExports;
-	for (const FPackageStorePackage::FExport& SourceExport : SourcePackage->Exports)
-	{
-		if (SourceExport.bIsPublic)
-		{
-			SourceExports.Add(SourceExport.PublicExportHash, &SourceExport);
-		}
-	}
-	TArray<FString> FailReasons;
-	for (const FPackageStorePackage::FExport& TargetExport : TargetPackage.Exports)
-	{
-		if (TargetExport.bIsPublic)
-		{
-			const FPackageStorePackage::FExport* SourceExport;
-			if (SourceExports.RemoveAndCopyValue(TargetExport.PublicExportHash, SourceExport))
-			{
-				if (TargetExport.ClassIndex != SourceExport->ClassIndex)
-				{
-					AppendMismatchMessage(TEXT("class"), TargetExport.ObjectName,
-						TargetExport.ClassIndex, SourceExport->ClassIndex, FailReasons.AddDefaulted_GetRef());
-				}
-				if (TargetExport.TemplateIndex != SourceExport->TemplateIndex)
-				{
-					AppendMismatchMessage(TEXT("template"), TargetExport.ObjectName,
-						TargetExport.TemplateIndex, SourceExport->TemplateIndex, FailReasons.AddDefaulted_GetRef());
-				}
-				if (TargetExport.SuperIndex != SourceExport->SuperIndex)
-				{
-					AppendMismatchMessage(TEXT("super"), TargetExport.ObjectName,
-						TargetExport.SuperIndex, SourceExport->SuperIndex, FailReasons.AddDefaulted_GetRef());
-				}
-			}
-		}
-	}
-	for (const auto& KV : SourceExports)
-	{
-		const FPackageStorePackage::FExport* SourceExport = KV.Value;
-		FailReasons.AddDefaulted_GetRef().Appendf(TEXT("Public source export '%s' is missing or not public in the target package"),
-			*SourceExport->ObjectName.ToString());
-	}
-
-	if (!FailReasons.IsEmpty())
-	{
-		for (const FString& FailReason : FailReasons)
-		{
-			UE_LOG(LogPackageStoreOptimizer, Warning,
-				TEXT("Redirection target package '%s' (0x%llX) for '%s' (0x%llX) - %s"),
-				*TargetPackage.Name.ToString(),
-				TargetPackage.Id.ValueForDebugging(),
-				*SourcePackage->Name.ToString(),
-				SourcePackage->Id.ValueForDebugging(),
-				*FailReason);
-		}
-		return false;
-	}
-	return true;
-}
-
-void FPackageStoreOptimizer::ProcessRedirects(const TMap<FPackageId, FPackageStorePackage*>& PackagesMap, bool bIsBuildingDLC) const
-{
-	TRACE_CPUPROFILER_EVENT_SCOPE(ProcessRedirects);
-
-	for (const auto& KV : PackagesMap)
-	{
-		FPackageStorePackage* Package = KV.Value;
-		if (Package->SourceName.IsNone() || Package->Name == Package->SourceName)
-		{
-			continue;
-		}
-		
-		FPackageId SourcePackageId = FPackageId::FromName(Package->SourceName);
-		FPackageStorePackage* SourcePackage = PackagesMap.FindRef(SourcePackageId);
-		Package->bIsRedirected = VerifyRedirect(SourcePackage, *Package, bIsBuildingDLC);
-		
-		if (Package->bIsRedirected)
-		{
-			UE_LOG(LogPackageStoreOptimizer, Verbose, TEXT("Adding package redirect from '%s' (0x%llX) to '%s' (0x%llX)."),
-				*Package->SourceName.ToString(),
-				SourcePackageId.ValueForDebugging(),
-				*Package->Name.ToString(),
-				Package->Id.ValueForDebugging());
-		}
-		else
-		{
-			if (Package->Region.Len() > 0 && !SourcePackage)
-			{
-				// no update or verification required
-				UE_LOG(LogPackageStoreOptimizer, Verbose,
-					TEXT("For culture '%s': Localized package '%s' (0x%llX) is unique and does not override a source package."),
-					*Package->Region,
-					*Package->Name.ToString(),
-					Package->Id.ValueForDebugging());
-			}
-			else
-			{
-				UE_LOG(LogPackageStoreOptimizer, Display,
-					TEXT("Skipping package redirect from '%s' (0x%llX) to '%s' (0x%llX) due to mismatching public exports."),
-					*Package->SourceName.ToString(),
-					SourcePackageId.ValueForDebugging(),
-					*Package->Name.ToString(),
-					Package->Id.ValueForDebugging());
-			}
-		}
-	}
-}
-
 void FPackageStoreOptimizer::SerializeGraphData(const TArray<FPackageId>& ImportedPackageIds, FPackageStorePackage::FGraphData& GraphData, FBufferWriter& GraphArchive) const
 {
 	uint32 ExportBundleEntryIndex = 0;
@@ -1355,14 +1168,14 @@ void FPackageStoreOptimizer::FinalizePackageHeader(FPackageStorePackage* Package
 	{
 		ImportedPublicExportHashesArchive << ImportedPublicExportHash;
 	}
-	Package->ImportedPublicExportHashesSize = ImportedPublicExportHashesArchive.Tell();
+	uint64 ImportedPublicExportHashesSize = ImportedPublicExportHashesArchive.Tell();
 
 	FBufferWriter ImportMapArchive(nullptr, 0, EBufferWriterFlags::AllowResize | EBufferWriterFlags::TakeOwnership);
 	for (FPackageObjectIndex Import : Package->Imports)
 	{
 		ImportMapArchive << Import;
 	}
-	Package->ImportMapSize = ImportMapArchive.Tell();
+	uint64 ImportMapSize = ImportMapArchive.Tell();
 
 	FBufferWriter ExportMapArchive(nullptr, 0, EBufferWriterFlags::AllowResize | EBufferWriterFlags::TakeOwnership);
 	for (const FPackageStorePackage::FExport& Export : Package->Exports)
@@ -1390,7 +1203,7 @@ void FPackageStoreOptimizer::FinalizePackageHeader(FPackageStorePackage* Package
 
 		ExportMapArchive << ExportMapEntry;
 	}
-	Package->ExportMapSize = ExportMapArchive.Tell();
+	uint64 ExportMapSize = ExportMapArchive.Tell();
 
 	FBufferWriter ExportBundleEntriesArchive(nullptr, 0, EBufferWriterFlags::AllowResize | EBufferWriterFlags::TakeOwnership);
 	for (const FPackageStorePackage::FExportBundle& ExportBundle : Package->GraphData.ExportBundles)
@@ -1400,43 +1213,44 @@ void FPackageStoreOptimizer::FinalizePackageHeader(FPackageStorePackage* Package
 			ExportBundleEntriesArchive << BundleEntry;
 		}
 	}
-	Package->ExportBundleEntriesSize = ExportBundleEntriesArchive.Tell();
+	uint64 ExportBundleEntriesSize = ExportBundleEntriesArchive.Tell();
 
 	FBufferWriter GraphArchive(nullptr, 0, EBufferWriterFlags::AllowResize | EBufferWriterFlags::TakeOwnership);
 	SerializeGraphData(Package->ImportedPackageIds, Package->GraphData, GraphArchive);
-	Package->GraphDataSize = GraphArchive.Tell();
+	uint64 GraphDataSize = GraphArchive.Tell();
 
 	Package->NameMapBuilder.MarkNameAsReferenced(Package->Name);
 	FMappedName MappedPackageName = Package->NameMapBuilder.MapName(Package->Name);
 
 	FBufferWriter NameMapArchive(nullptr, 0, EBufferWriterFlags::AllowResize | EBufferWriterFlags::TakeOwnership);
 	SaveNameBatch(Package->NameMapBuilder.GetNameMap(), NameMapArchive);
-	Package->NameMapSize = NameMapArchive.Tell();
+	uint64 NameMapSize = NameMapArchive.Tell();
 
 	FBufferWriter VersioningInfoArchive(nullptr, 0, EBufferWriterFlags::AllowResize | EBufferWriterFlags::TakeOwnership);
 	if (Package->VersioningInfo.IsSet())
 	{
 		VersioningInfoArchive << Package->VersioningInfo.GetValue();
-		Package->VersioningInfoSize = VersioningInfoArchive.Tell();
 	}
+	uint64 VersioningInfoSize = VersioningInfoArchive.Tell();
+
 	
 	FBufferWriter BulkDataMapAr(nullptr, 0, EBufferWriterFlags::AllowResize | EBufferWriterFlags::TakeOwnership);
 	for (FBulkDataMapEntry& Entry : Package->BulkDataEntries)
 	{
 		BulkDataMapAr << Entry;
 	}
-	Package->BulkDataMapSize = BulkDataMapAr.Tell();
+	uint64 BulkDataMapSize = BulkDataMapAr.Tell();
 
 	Package->HeaderSize =
 		sizeof(FZenPackageSummary)
-		+ Package->VersioningInfoSize
-		+ Package->NameMapSize
-		+ Package->ImportedPublicExportHashesSize
-		+ Package->ImportMapSize
-		+ Package->ExportMapSize
-		+ Package->ExportBundleEntriesSize
-		+ Package->GraphDataSize
-		+ Package->BulkDataMapSize + sizeof(int64);
+		+ VersioningInfoSize
+		+ NameMapSize
+		+ ImportedPublicExportHashesSize
+		+ ImportMapSize
+		+ ExportMapSize
+		+ ExportBundleEntriesSize
+		+ GraphDataSize
+		+ BulkDataMapSize + sizeof(int64);
 
 	Package->HeaderBuffer = FIoBuffer(Package->HeaderSize);
 	uint8* HeaderData = Package->HeaderBuffer.Data();
@@ -1461,8 +1275,8 @@ void FPackageStoreOptimizer::FinalizePackageHeader(FPackageStorePackage* Package
 
 	HeaderArchive.Serialize(NameMapArchive.GetWriterData(), NameMapArchive.Tell());
 
-	HeaderArchive << Package->BulkDataMapSize;
-	HeaderArchive.Serialize(BulkDataMapAr.GetWriterData(), Package->BulkDataMapSize);
+	HeaderArchive << BulkDataMapSize;
+	HeaderArchive.Serialize(BulkDataMapAr.GetWriterData(), BulkDataMapSize);
 
 	PackageSummary->ImportedPublicExportHashesOffset = HeaderArchive.Tell();
 	HeaderArchive.Serialize(ImportedPublicExportHashesArchive.GetWriterData(), ImportedPublicExportHashesArchive.Tell());
@@ -1477,18 +1291,8 @@ void FPackageStoreOptimizer::FinalizePackageHeader(FPackageStorePackage* Package
 	check(HeaderArchive.Tell() == PackageSummary->HeaderSize)
 }
 
-void FPackageStoreOptimizer::FinalizePackage(FPackageStorePackage* Package)
+void FPackageStoreOptimizer::FinalizePackage(FPackageStorePackage* Package) const
 {
-	++TotalPackageCount;
-	TotalExportBundleCount += Package->GraphData.ExportBundles.Num();
-	uint64 PackageExportBundleEntryCount = 0;
-	for (const FPackageStorePackage::FExportBundle& ExportBundle : Package->GraphData.ExportBundles)
-	{
-		PackageExportBundleEntryCount += ExportBundle.Entries.Num();
-	}
-	check(PackageExportBundleEntryCount == FExportBundleEntry::ExportCommandType_Count * Package->Exports.Num());
-	TotalExportBundleEntryCount += PackageExportBundleEntryCount;
-
 	for (FPackageStorePackage::FExportGraphNode& Node : Package->ExportGraphNodes)
 	{
 		check(Node.ExportBundleIndex >= 0);
@@ -1506,7 +1310,6 @@ void FPackageStoreOptimizer::FinalizePackage(FPackageStorePackage* Package)
 				if (!bIsAlreadyInSet)
 				{
 					Package->GraphData.InternalArcs.Add(Arc);
-					++TotalInternalBundleArcsCount;
 				}
 			}
 		}
@@ -1525,7 +1328,6 @@ void FPackageStoreOptimizer::FinalizePackage(FPackageStorePackage* Package)
 			if (!ArcsFromImportedPackage.Contains(Arc))
 			{
 				ArcsFromImportedPackage.Add(Arc);
-				++TotalExternalBundleArcsCount;
 			}
 		}
 	}
@@ -1752,7 +1554,7 @@ void FPackageStoreOptimizer::LoadScriptObjectsBuffer(const FIoBuffer& ScriptObje
 FPackageStoreEntryResource FPackageStoreOptimizer::CreatePackageStoreEntry(const FPackageStorePackage* Package, const FPackageStorePackage* OptionalSegmentPackage) const
 {
 	FPackageStoreEntryResource Result;
-	Result.Flags = Package->bIsRedirected ? EPackageStoreEntryFlags::Redirected : EPackageStoreEntryFlags::None;
+	Result.Flags = EPackageStoreEntryFlags::None;
 
 	if (OptionalSegmentPackage && OptionalSegmentPackage->HasEditorData())
 	{
@@ -1762,15 +1564,10 @@ FPackageStoreEntryResource FPackageStoreOptimizer::CreatePackageStoreEntry(const
 	
 	Result.PackageName = Package->Name;
 	Result.PackageId = FPackageId::FromName(Package->Name);
-	Result.SourcePackageName = Package->SourceName;
-	Result.Region = FName(*Package->Region);
 	Result.ExportInfo.ExportCount = Package->Exports.Num();
 	Result.ExportInfo.ExportBundleCount = Package->GraphData.ExportBundles.Num();
 	Result.ImportedPackageIds = Package->ImportedPackageIds;
-	Result.ShaderMapHashes = Package->ShaderMapHashes.Array();
-	// Package->ShaderMapHashes is a TSet and is unsorted - we want this sorted as it's serialized
-	// in to the ContainerHeader, and if not sorted leads to staging non-determinism.
-	Algo::Sort(Result.ShaderMapHashes);
+	
 	if (OptionalSegmentPackage)
 	{
 		Result.OptionalSegmentExportInfo.ExportCount = OptionalSegmentPackage->Exports.Num();
@@ -1779,190 +1576,3 @@ FPackageStoreEntryResource FPackageStoreOptimizer::CreatePackageStoreEntry(const
 	}
 	return Result;
 }
-
-FIoContainerHeader FPackageStoreOptimizer::CreateContainerHeaderInternal(const FIoContainerId& ContainerId, TArrayView<const FPackageStoreEntryResource> PackageStoreEntries, bool bIsOptional) const
-{
-	TRACE_CPUPROFILER_EVENT_SCOPE(CreateContainerHeaderInternal);
-	FIoContainerHeader Header;
-	Header.ContainerId = ContainerId;
-	
-	int32 NonOptionalSegmentStoreEntriesCount = 0;
-	int32 OptionalSegmentStoreEntriesCount = 0;
-	if (bIsOptional)
-	{
-		for (const FPackageStoreEntryResource& Entry : PackageStoreEntries)
-		{
-			if (Entry.OptionalSegmentExportInfo.ExportCount)
-			{
-				if (Entry.IsAutoOptional())
-				{
-					// Auto optional packages fully replace the non-optional segment
-					++NonOptionalSegmentStoreEntriesCount;
-				}
-				else
-				{
-					++OptionalSegmentStoreEntriesCount;
-				}
-			}
-		}
-	}
-	else
-	{
-		NonOptionalSegmentStoreEntriesCount = PackageStoreEntries.Num();
-	}
-
-	struct FStoreEntriesWriter
-	{
-		const int32 StoreTocSize;
-		FLargeMemoryWriter StoreTocArchive = FLargeMemoryWriter(0, true);
-		FLargeMemoryWriter StoreDataArchive = FLargeMemoryWriter(0, true);
-
-		void Flush(TArray<uint8>& OutputBuffer)
-		{
-			check(StoreTocArchive.TotalSize() == StoreTocSize);
-			if (StoreTocSize)
-			{
-				const int32 StoreByteCount = StoreTocArchive.TotalSize() + StoreDataArchive.TotalSize();
-				OutputBuffer.AddUninitialized(StoreByteCount);
-				FBufferWriter PackageStoreArchive(OutputBuffer.GetData(), StoreByteCount);
-				PackageStoreArchive.Serialize(StoreTocArchive.GetData(), StoreTocArchive.TotalSize());
-				PackageStoreArchive.Serialize(StoreDataArchive.GetData(), StoreDataArchive.TotalSize());
-			}
-		}
-	};
-
-	FStoreEntriesWriter StoreEntriesWriter
-	{
-		static_cast<int32>(NonOptionalSegmentStoreEntriesCount * sizeof(FFilePackageStoreEntry))
-	};
-
-	FStoreEntriesWriter OptionalSegmentStoreEntriesWriter
-	{
-		static_cast<int32>(OptionalSegmentStoreEntriesCount * sizeof(FFilePackageStoreEntry))
-	};
-
-	auto SerializePackageEntryCArrayHeader = [](FStoreEntriesWriter& Writer, int32 Count)
-	{
-		const int32 RemainingTocSize = Writer.StoreTocSize - Writer.StoreTocArchive.Tell();
-		const int32 OffsetFromThis = RemainingTocSize + Writer.StoreDataArchive.Tell();
-		uint32 ArrayNum = Count > 0 ? Count : 0;
-		uint32 OffsetToDataFromThis = ArrayNum > 0 ? OffsetFromThis : 0;
-
-		Writer.StoreTocArchive << ArrayNum;
-		Writer.StoreTocArchive << OffsetToDataFromThis;
-	};
-
-	TArray<const FPackageStoreEntryResource*> SortedPackageStoreEntries;
-	SortedPackageStoreEntries.Reserve(PackageStoreEntries.Num());
-	for (const FPackageStoreEntryResource& Entry : PackageStoreEntries)
-	{
-		SortedPackageStoreEntries.Add(&Entry);
-	}
-	Algo::Sort(SortedPackageStoreEntries, [](const FPackageStoreEntryResource* A, const FPackageStoreEntryResource* B)
-	{
-		return A->GetPackageId() < B->GetPackageId();
-	});
-
-	Header.PackageIds.Reserve(NonOptionalSegmentStoreEntriesCount);
-	Header.OptionalSegmentPackageIds.Reserve(OptionalSegmentStoreEntriesCount);
-	FPackageStoreNameMapBuilder RedirectsNameMapBuilder;
-	RedirectsNameMapBuilder.SetNameMapType(FMappedName::EType::Container);
-	TSet<FPackageId> AllLocalizedPackages;
-	if (bIsOptional)
-	{
-		for (const FPackageStoreEntryResource* Entry : SortedPackageStoreEntries)
-		{
-			if (Entry->OptionalSegmentExportInfo.ExportCount)
-			{
-				FStoreEntriesWriter* TargetEntriesWriter;
-				if (Entry->IsAutoOptional())
-				{
-					Header.PackageIds.Add(Entry->GetPackageId());
-					TargetEntriesWriter = &StoreEntriesWriter;
-				}
-				else
-				{
-					Header.OptionalSegmentPackageIds.Add(Entry->GetPackageId());
-					TargetEntriesWriter = &OptionalSegmentStoreEntriesWriter;
-				}
-
-				// OptionalStoreEntry
-				FPackageStoreExportInfo OptionalSegmentExportInfo = Entry->OptionalSegmentExportInfo;
-				TargetEntriesWriter->StoreTocArchive << OptionalSegmentExportInfo;
-
-				// OptionalImportedPackages
-				const TArray<FPackageId>& OptionalSegmentImportedPackageIds = Entry->OptionalSegmentImportedPackageIds;
-				SerializePackageEntryCArrayHeader(*TargetEntriesWriter, OptionalSegmentImportedPackageIds.Num());
-				for (FPackageId OptionalSegmentImportedPackageId : OptionalSegmentImportedPackageIds)
-				{
-					check(OptionalSegmentImportedPackageId.IsValid());
-					TargetEntriesWriter->StoreDataArchive << OptionalSegmentImportedPackageId;
-				}
-
-				// ShaderMapHashes is N/A for optional segments
-				SerializePackageEntryCArrayHeader(*TargetEntriesWriter, 0);
-			}
-		}
-	}
-	else
-	{
-		for (const FPackageStoreEntryResource* Entry : SortedPackageStoreEntries)
-		{
-			Header.PackageIds.Add(Entry->GetPackageId());
-			if (Entry->IsRedirected())
-			{
-				RedirectsNameMapBuilder.MarkNameAsReferenced(Entry->GetSourcePackageName());
-				FMappedName MappedSourcePackageName = RedirectsNameMapBuilder.MapName(Entry->GetSourcePackageName());
-				if (!Entry->Region.IsNone())
-				{
-					if (!AllLocalizedPackages.Contains(Entry->GetSourcePackageId()))
-					{
-						Header.LocalizedPackages.Add({ Entry->GetSourcePackageId(), MappedSourcePackageName });
-						AllLocalizedPackages.Add(Entry->GetSourcePackageId());
-					}
-				}
-				else
-				{
-					Header.PackageRedirects.Add({ Entry->GetSourcePackageId(), Entry->GetPackageId(), MappedSourcePackageName });
-				}
-			}
-
-			// StoreEntries
-			FPackageStoreExportInfo ExportInfo = Entry->ExportInfo;
-			StoreEntriesWriter.StoreTocArchive << ExportInfo;
-
-			// ImportedPackages
-			const TArray<FPackageId>& ImportedPackageIds = Entry->ImportedPackageIds;
-			SerializePackageEntryCArrayHeader(StoreEntriesWriter, ImportedPackageIds.Num());
-			for (FPackageId ImportedPackageId : ImportedPackageIds)
-			{
-				check(ImportedPackageId.IsValid());
-				StoreEntriesWriter.StoreDataArchive << ImportedPackageId;
-			}
-
-			// ShaderMapHashes
-			const TArray<FSHAHash>& ShaderMapHashes = Entry->ShaderMapHashes;
-			SerializePackageEntryCArrayHeader(StoreEntriesWriter, ShaderMapHashes.Num());
-			for (const FSHAHash& ShaderMapHash : ShaderMapHashes)
-			{
-				StoreEntriesWriter.StoreDataArchive << const_cast<FSHAHash&>(ShaderMapHash);
-			}
-		}
-	}
-	Header.RedirectsNameMap = RedirectsNameMapBuilder.GetNameMap();
-
-	StoreEntriesWriter.Flush(Header.StoreEntries);
-	OptionalSegmentStoreEntriesWriter.Flush(Header.OptionalSegmentStoreEntries);
-
-	return Header;
-}
-
-FIoContainerHeader FPackageStoreOptimizer::CreateContainerHeader(const FIoContainerId& ContainerId, TArrayView<const FPackageStoreEntryResource> PackageStoreEntries) const
-{
-	return CreateContainerHeaderInternal(ContainerId, PackageStoreEntries, false);
-}
-
-FIoContainerHeader FPackageStoreOptimizer::CreateOptionalContainerHeader(const FIoContainerId& ContainerId, TArrayView<const FPackageStoreEntryResource> PackageStoreEntries) const
-{
-	return CreateContainerHeaderInternal(ContainerId, PackageStoreEntries, true);
-}
diff --git a/Engine/Source/Developer/IoStoreUtilities/Private/ZenStoreWriter.cpp b/Engine/Source/Developer/IoStoreUtilities/Private/ZenStoreWriter.cpp
--- a/Engine/Source/Developer/IoStoreUtilities/Private/ZenStoreWriter.cpp
+++ b/Engine/Source/Developer/IoStoreUtilities/Private/ZenStoreWriter.cpp
@@ -637,11 +637,6 @@ void FZenStoreWriter::EndCook(const FCookInfo& Info)
 		}
 	}
 
-	UE_LOG(LogZenStoreWriter, Display, TEXT("Input:\t%d Packages"), PackageStoreOptimizer->GetTotalPackageCount());
-	UE_LOG(LogZenStoreWriter, Display, TEXT("Output:\t%d Export bundles"), PackageStoreOptimizer->GetTotalExportBundleCount());
-	UE_LOG(LogZenStoreWriter, Display, TEXT("Output:\t%d Export bundle entries"), PackageStoreOptimizer->GetTotalExportBundleEntryCount());
-	UE_LOG(LogZenStoreWriter, Display, TEXT("Output:\t%d Internal export bundle arcs"), PackageStoreOptimizer->GetTotalInternalBundleArcsCount());
-	UE_LOG(LogZenStoreWriter, Display, TEXT("Output:\t%d External export bundle arcs"), PackageStoreOptimizer->GetTotalExternalBundleArcsCount());
 	UE_LOG(LogZenStoreWriter, Display, TEXT("Output:\t%d Public runtime script objects"), PackageStoreOptimizer->GetTotalScriptObjectCount());
 }
 
diff --git a/Engine/Source/Runtime/Core/Private/IO/PackageStore.cpp b/Engine/Source/Runtime/Core/Private/IO/PackageStore.cpp
--- a/Engine/Source/Runtime/Core/Private/IO/PackageStore.cpp
+++ b/Engine/Source/Runtime/Core/Private/IO/PackageStore.cpp
@@ -48,8 +48,6 @@ FArchive& operator<<(FArchive& Ar, FPackageStoreEntryResource& PackageStoreEntry
 
 	Ar << Flags;
 	Ar << PackageStoreEntry.PackageName;
-	Ar << PackageStoreEntry.SourcePackageName;
-	Ar << PackageStoreEntry.Region;
 	Ar << PackageStoreEntry.ExportInfo;
 	Ar << PackageStoreEntry.ImportedPackageIds;
 	Ar << PackageStoreEntry.OptionalSegmentExportInfo;
@@ -70,8 +68,6 @@ FCbWriter& operator<<(FCbWriter& Writer, const FPackageStoreEntryResource& Packa
 
 	Writer << "flags" << static_cast<uint32>(PackageStoreEntry.Flags);
 	Writer << "packagename" << PackageStoreEntry.PackageName.ToString();
-	Writer << "sourcepackagename" << PackageStoreEntry.SourcePackageName.ToString();
-	Writer << "region" << PackageStoreEntry.Region.ToString();
 	Writer << "exportinfo" << PackageStoreEntry.ExportInfo;
 
 	if (PackageStoreEntry.ImportedPackageIds.Num())
@@ -118,8 +114,6 @@ FPackageStoreEntryResource FPackageStoreEntryResource::FromCbObject(const FCbObj
 	Entry.Flags				= static_cast<EPackageStoreEntryFlags>(Obj["flags"].AsUInt32());
 	Entry.PackageName		= FName(Obj["packagename"].AsString());
 	Entry.PackageId			= FPackageId::FromName(Entry.PackageName);
-	Entry.SourcePackageName	= FName(Obj["sourcepackagename"].AsString());
-	Entry.Region			= FName(Obj["region"].AsString());
 	Entry.ExportInfo		= FPackageStoreExportInfo::FromCbObject(Obj["exportinfo"].AsObject());
 	
 	if (Obj["importedpackageids"])
diff --git a/Engine/Source/Runtime/Core/Public/IO/PackageStore.h b/Engine/Source/Runtime/Core/Public/IO/PackageStore.h
--- a/Engine/Source/Runtime/Core/Public/IO/PackageStore.h
+++ b/Engine/Source/Runtime/Core/Public/IO/PackageStore.h
@@ -67,7 +67,6 @@ struct FPackageStoreEntry
 enum class EPackageStoreEntryFlags : uint32
 {
 	None		= 0,
-	Redirected	= 0x01,
 	AutoOptional= 0x02,
 };
 ENUM_CLASS_FLAGS(EPackageStoreEntryFlags);
@@ -86,15 +85,11 @@ struct FPackageStoreEntryResource
 	/** The package name. */
 	FName PackageName;
 	FPackageId PackageId;
-	/** Used for localized and redirected packages. */
-	FName SourcePackageName;
-	/** Region name for localized packages. */
-	FName Region;
 	/** The package export information. */
 	FPackageStoreExportInfo ExportInfo;
 	/** Imported package IDs. */
 	TArray<FPackageId> ImportedPackageIds;
-	/** Referenced shader map hashes - must be sorted in order to preserve determinism across builds. */
+	/** Referenced shader map hashes. */
 	TArray<FSHAHash> ShaderMapHashes;
 	/** The editor data package export information. */
 	FPackageStoreExportInfo OptionalSegmentExportInfo;
@@ -107,23 +102,6 @@ struct FPackageStoreEntryResource
 		return PackageId;
 	}
 
-	/** Returns the source package ID. */
-	FPackageId GetSourcePackageId() const
-	{
-		return SourcePackageName.IsNone() ? FPackageId() : FPackageId::FromName(SourcePackageName);
-	}
-
-	FName GetSourcePackageName() const
-	{
-		return SourcePackageName;
-	}
-
-	/** Returns whether this package is redirected. */
-	bool IsRedirected() const
-	{
-		return EnumHasAnyFlags(Flags, EPackageStoreEntryFlags::Redirected); 
-	}
-
 	/** Returns whether this package was saved as auto optional */
 	bool IsAutoOptional() const
 	{
