ZenCook: Added support for file regions
#rnx

diff --git a/Engine/Source/Developer/IoStoreUtilities/Private/IoStoreUtilities.cpp b/Engine/Source/Developer/IoStoreUtilities/Private/IoStoreUtilities.cpp
--- a/Engine/Source/Developer/IoStoreUtilities/Private/IoStoreUtilities.cpp
+++ b/Engine/Source/Developer/IoStoreUtilities/Private/IoStoreUtilities.cpp
@@ -655,6 +655,14 @@ struct FContainerTargetFile
 class FCookedPackageStore
 {
 public:
+	struct FChunkInfo
+	{
+		FIoChunkId ChunkId;
+		FName PackageName;
+		FString RelativeFileName;
+		TArray<FFileRegion> FileRegions;
+	};
+
 	FCookedPackageStore(const FString& InCookedDir)
 		: CookedDir(InCookedDir)
 	{
@@ -706,19 +714,26 @@ public:
 					FCbObject ChunkObj = ChunkEntry.AsObject();
 					FIoChunkId ChunkId;
 					ChunkId.Set(ChunkObj["id"].AsObjectId().GetView());
-					ChunkIdToPackageNameMap.Add(ChunkId, PackageStoreEntry.PackageName);
+					FChunkInfo& ChunkInfo = ChunkInfoMap.Add(ChunkId);
+					ChunkInfo.ChunkId = ChunkId;
+					ChunkInfo.PackageName = PackageStoreEntry.PackageName;
 					if (ChunkObj["filename"])
 					{
 						TStringBuilder<1024> RelativeFilename;
 						RelativeFilename.Append(ChunkObj["filename"].AsString());
-						ChunkIdToRelativeFileNameMap.Add(ChunkId, *RelativeFilename);
+						ChunkInfo.RelativeFileName = RelativeFilename;
 						TStringBuilder<1024> PathBuilder;
 						FPathViews::AppendPath(PathBuilder, CookedDir);
 						FPathViews::AppendPath(PathBuilder, RelativeFilename);
 						FPathViews::NormalizeFilename(PathBuilder);
 						FilenameToChunkIdMap.Add(*PathBuilder, ChunkId);
-						FString FilenameWithoutExtension = FPathViews::ChangeExtension(PathBuilder, FStringView());
-						FilenameToPackageNameMap.Add(MoveTemp(FilenameWithoutExtension), PackageStoreEntry.PackageName);
+					}
+					const FCbArrayView RegionsArray = ChunkObj["fileregions"].AsArrayView();
+					ChunkInfo.FileRegions.Reserve(RegionsArray.Num());
+					for (FCbFieldView RegionObj : RegionsArray)
+					{
+						FFileRegion& Region = ChunkInfo.FileRegions.AddDefaulted_GetRef();
+						FFileRegion::LoadFromCompactBinary(RegionObj, Region);
 					}
 				}
 			};
@@ -736,20 +751,36 @@ public:
 		return FilenameToChunkIdMap.FindRef(*Filename);
 	}
 
+	const FChunkInfo* GetChunkInfoFromFileName(const FString& Filename) const
+	{
+		FIoChunkId ChunkId = GetChunkIdFromFileName(Filename);
+		return ChunkInfoMap.Find(ChunkId);
+	}
+
 	FString GetRelativeFilenameFromChunkId(const FIoChunkId& ChunkId) const
 	{
-		return ChunkIdToRelativeFileNameMap.FindRef(ChunkId);
+		const FChunkInfo* FindChunkInfo = ChunkInfoMap.Find(ChunkId);
+		if (!FindChunkInfo)
+		{
+			return FString();
+		}
+		return FindChunkInfo->RelativeFileName;
 	}
 
 	FName GetPackageNameFromChunkId(const FIoChunkId& ChunkId) const
 	{
-		return ChunkIdToPackageNameMap.FindRef(ChunkId);
+		const FChunkInfo* FindChunkInfo = ChunkInfoMap.Find(ChunkId);
+		if (!FindChunkInfo)
+		{
+			return NAME_None;
+		}
+		return FindChunkInfo->PackageName;
 	}
 
 	FName GetPackageNameFromFileName(const FString& Filename) const
 	{
-		FString FilenameWithoutExtension = FPaths::ChangeExtension(Filename, FString());
-		return FilenameToPackageNameMap.FindRef(FilenameWithoutExtension);
+		FIoChunkId ChunkId = GetChunkIdFromFileName(Filename);
+		return GetPackageNameFromChunkId(ChunkId);
 	}
 
 	const FPackageStoreEntryResource* GetPackageStoreEntry(FPackageId PackageId) const
@@ -844,9 +875,7 @@ private:
 	FString CookedDir;
 	TMap<FPackageId, FPackageStoreEntryResource> PackageIdToEntry;
 	TMap<FString, FIoChunkId> FilenameToChunkIdMap;
-	TMap<FIoChunkId, FString> ChunkIdToRelativeFileNameMap;
-	TMap<FIoChunkId, FName> ChunkIdToPackageNameMap;
-	TMap<FString, FName> FilenameToPackageNameMap;
+	TMap<FIoChunkId, FChunkInfo> ChunkInfoMap;
 };
 
 struct FFileOrderMap
@@ -2444,13 +2473,13 @@ void InitializeContainerTargetsAndPackages(
 			return true;
 		}
 
-		OutTargetFile.ChunkId = PackageStore.GetChunkIdFromFileName(SourceFile.NormalizedPath);
-		if (!OutTargetFile.ChunkId.IsValid())
+		const FCookedPackageStore::FChunkInfo* ChunkInfo = PackageStore.GetChunkInfoFromFileName(SourceFile.NormalizedPath);
+		if (!ChunkInfo)
 		{
 			UE_LOG(LogIoStore, Warning, TEXT("File not found in manifest: '%s'"), *SourceFile.NormalizedPath);
 			return false;
 		}
-
+		OutTargetFile.ChunkId = ChunkInfo->ChunkId;
 		TIoStatusOr<uint64> ChunkSize = PackageStore.GetChunkSize(OutTargetFile.ChunkId);
 		if (!ChunkSize.IsOk())
 		{
@@ -2459,18 +2488,17 @@ void InitializeContainerTargetsAndPackages(
 		}
 		OutTargetFile.SourceSize = ChunkSize.ValueOrDie();
 
-		FName PackageName = PackageStore.GetPackageNameFromChunkId(OutTargetFile.ChunkId);
-		if (PackageName.IsNone())
+		if (ChunkInfo->PackageName.IsNone())
 		{
 			UE_LOG(LogIoStore, Warning, TEXT("Package name not found for: '%s'"), *SourceFile.NormalizedPath);
 			return false;
 		}
 
-		OutTargetFile.Package = &FindOrAddPackage(Arguments, PackageName, Packages, PackageNameMap, PackageIdMap);
+		OutTargetFile.Package = &FindOrAddPackage(Arguments, ChunkInfo->PackageName, Packages, PackageNameMap, PackageIdMap);
 		const FPackageStoreEntryResource* PackageStoreEntry = PackageStore.GetPackageStoreEntry(OutTargetFile.Package->GlobalPackageId);
 		if (!PackageStoreEntry)
 		{
-			UE_LOG(LogIoStore, Warning, TEXT("Failed to find package store entry for package: '%s'"), *PackageName.ToString());
+			UE_LOG(LogIoStore, Warning, TEXT("Failed to find package store entry for package: '%s'"), *ChunkInfo->PackageName.ToString());
 			return false;
 		}
 		OutTargetFile.Package->PackageStoreEntry = *PackageStoreEntry;
@@ -2511,18 +2539,10 @@ void InitializeContainerTargetsAndPackages(
 		}
 
 		// Only keep the regions for the file if neither compression nor encryption are enabled, otherwise the regions will be meaningless.
-		// TODO: There are no region files when cooking for Zen
-		//if (Arguments.bFileRegions && !SourceFile.bNeedsCompression && !SourceFile.bNeedsEncryption)
-		//{
-		//	FString RegionsFilename = OutTargetFile.ChunkType == EContainerChunkType::PackageData
-		//		? FPaths::ChangeExtension(SourceFile.NormalizedPath, FString(TEXT(".uexp")) + FFileRegion::RegionsFileExtension)
-		//		: SourceFile.NormalizedPath + FFileRegion::RegionsFileExtension;
-		//	TUniquePtr<FArchive> RegionsFile(IFileManager::Get().CreateFileReader(*RegionsFilename));
-		//	if (RegionsFile.IsValid())
-		//	{
-		//		FFileRegion::SerializeFileRegions(*RegionsFile.Get(), OutTargetFile.FileRegions);
-		//	}
-		//}
+		if (Arguments.bFileRegions && !SourceFile.bNeedsCompression && !SourceFile.bNeedsEncryption)
+		{
+			OutTargetFile.FileRegions = ChunkInfo->FileRegions;
+		}
 		
 		return true;
 	};
@@ -2913,8 +2933,16 @@ private:
 						check(TargetFile.Package->UAssetSize > 0);
 						const uint64 HeaderSize = TargetFile.Package->UAssetSize;
 						FIoBuffer ExportsBuffer(SourceBuffer.Data() + HeaderSize, SourceBuffer.DataSize() - HeaderSize, SourceBuffer);
-						SourceBuffer = Manager.PackageStoreOptimizer.CreatePackageBuffer(TargetFile.Package->OptimizedPackage, ExportsBuffer, bHasUpdatedExportBundleRegions ? nullptr : &FileRegions);
-						bHasUpdatedExportBundleRegions = true;
+						SourceBuffer = Manager.PackageStoreOptimizer.CreatePackageBuffer(TargetFile.Package->OptimizedPackage, ExportsBuffer, nullptr);
+						if (!bHasUpdatedExportBundleRegions)
+						{
+							// The regions are relative to the start of the export data, adjust to include the the new header size
+							for (FFileRegion& Region : FileRegions)
+							{
+								Region.Offset += TargetFile.Package->OptimizedPackage->GetHeaderSize();
+							}
+							bHasUpdatedExportBundleRegions = true;
+						}
 					}
 					OnSourceBufferLoaded();
 				});
diff --git a/Engine/Source/Developer/IoStoreUtilities/Private/ZenStoreWriter.cpp b/Engine/Source/Developer/IoStoreUtilities/Private/ZenStoreWriter.cpp
--- a/Engine/Source/Developer/IoStoreUtilities/Private/ZenStoreWriter.cpp
+++ b/Engine/Source/Developer/IoStoreUtilities/Private/ZenStoreWriter.cpp
@@ -229,19 +229,18 @@ void FZenStoreWriter::WritePackageData(const FPackageInfo& Info, FLargeMemoryWri
 	FIoBuffer CookedExportsBuffer = FIoBuffer(PackageData.Data() + Info.HeaderSize, PackageData.DataSize() - Info.HeaderSize, PackageData);
 	Entry.OptimizedPackage.Reset(PackageStoreOptimizer->CreatePackageFromCookedHeader(Info.PackageName, CookedHeaderBuffer));
 	PackageStoreOptimizer->FinalizePackage(Entry.OptimizedPackage.Get());
-	TArray<FFileRegion> FileRegionsCopy(FileRegions);
-	for (FFileRegion& Region : FileRegionsCopy)
+	Entry.FileRegions = FileRegions;
+	for (FFileRegion& Region : Entry.FileRegions)
 	{
 		// Adjust regions so they are relative to the start of the exports buffer
 		Region.Offset -= Info.HeaderSize;
 	}
-	FIoBuffer PackageBuffer = PackageStoreOptimizer->CreatePackageBuffer(Entry.OptimizedPackage.Get(), CookedExportsBuffer, &FileRegionsCopy);
-	for (FFileRegion& Region : FileRegionsCopy)
+	FIoBuffer PackageBuffer = PackageStoreOptimizer->CreatePackageBuffer(Entry.OptimizedPackage.Get(), CookedExportsBuffer, &Entry.FileRegions);
+	for (FFileRegion& Region : Entry.FileRegions)
 	{
-		// Adjust regions once more so they are relative to the exports bundle buffer
+		// Adjust regions once more so they are relative to the start of the exports data again
 		Region.Offset -= Entry.OptimizedPackage->GetHeaderSize();
 	}
-	//WriteFileRegions(*FPaths::ChangeExtension(Info.LooseFilePath, FString(".uexp") + FFileRegion::RegionsFileExtension), FileRegionsCopy);
 
 	// Commit to Zen build store
 
@@ -303,8 +302,7 @@ void FZenStoreWriter::WriteBulkData(const FBulkDataInfo& Info, const FIoBuffer&
 	BulkEntry.Info		= Info;
 	BulkEntry.ChunkId	= ChunkOid;
 	BulkEntry.IsValid	= true;
-
-	//	WriteFileRegions(*(Info.LooseFilePath + FFileRegion::RegionsFileExtension), FileRegions);
+	BulkEntry.FileRegions = FileRegions;
 }
 
 void FZenStoreWriter::WriteAdditionalFile(const FAdditionalFileInfo& Info, const FIoBuffer& FileData)
@@ -818,6 +816,24 @@ void FZenStoreWriter::CommitPackageInternal(FZenCommitInfo&& ZenCommitInfo)
 		OplogEntryDesc << "key" << PackageNameKey;
 		OplogEntryDesc << "packagestoreentry" << PackageStoreEntry;
 		
+		auto AppendFileNameAndRegionsToOplog = [this, &OplogEntryDesc](const FString& LooseFilePath, const TArray<FFileRegion>& FileRegions)
+		{
+			FStringView RelativePathView;
+			if (FPathViews::TryMakeChildPathRelativeTo(LooseFilePath, OutputPath, RelativePathView))
+			{
+				OplogEntryDesc << "filename" << RelativePathView;
+			}
+			if (!FileRegions.IsEmpty())
+			{
+				OplogEntryDesc.BeginArray("fileregions");
+				for (const FFileRegion& FileRegion : FileRegions)
+				{
+					OplogEntryDesc << FileRegion;
+				}
+				OplogEntryDesc.EndArray();
+			}
+		};
+
 		OplogEntryDesc.BeginArray("packagedata");
 		
 		for (FPackageDataEntry& PkgData : PackageState->PackageData)
@@ -834,11 +850,7 @@ void FZenStoreWriter::CommitPackageInternal(FZenCommitInfo&& ZenCommitInfo)
 			OplogEntryDesc.BeginObject();
 			OplogEntryDesc << "id" << PkgData.ChunkId;
 			OplogEntryDesc << "data" << PkgDataAttachment;
-			FStringView RelativePathView;
-			if (FPathViews::TryMakeChildPathRelativeTo(PkgData.Info.LooseFilePath, OutputPath, RelativePathView))
-			{
-				OplogEntryDesc << "filename" << RelativePathView;
-			}
+			AppendFileNameAndRegionsToOplog(PkgData.Info.LooseFilePath, PkgData.FileRegions);
 			OplogEntryDesc.EndObject();
 		}
 
@@ -863,11 +875,7 @@ void FZenStoreWriter::CommitPackageInternal(FZenCommitInfo&& ZenCommitInfo)
 				OplogEntryDesc << "id" << Bulk.ChunkId;
 				OplogEntryDesc << "type" << LexToString(Bulk.Info.BulkDataType);
 				OplogEntryDesc << "data" << BulkAttachment;
-				FStringView RelativePathView;
-				if (FPathViews::TryMakeChildPathRelativeTo(Bulk.Info.LooseFilePath, OutputPath, RelativePathView))
-				{
-					OplogEntryDesc << "filename" << RelativePathView;
-				}
+				AppendFileNameAndRegionsToOplog(Bulk.Info.LooseFilePath, Bulk.FileRegions);
 				OplogEntryDesc.EndObject();
 			}
 
diff --git a/Engine/Source/Developer/IoStoreUtilities/Public/ZenStoreWriter.h b/Engine/Source/Developer/IoStoreUtilities/Public/ZenStoreWriter.h
--- a/Engine/Source/Developer/IoStoreUtilities/Public/ZenStoreWriter.h
+++ b/Engine/Source/Developer/IoStoreUtilities/Public/ZenStoreWriter.h
@@ -113,6 +113,7 @@ private:
 		TFuture<FCompressedBuffer> CompressedPayload;
 		FBulkDataInfo Info;
 		FCbObjectId ChunkId;
+		TArray<FFileRegion> FileRegions;
 		bool IsValid = false;
 	};
 
@@ -124,6 +125,7 @@ private:
 		FPackageInfo Info;
 		FCbObjectId ChunkId;
 		TUniquePtr<class FPackageStorePackage> OptimizedPackage;
+		TArray<FFileRegion> FileRegions;
 		bool IsValid = false;
 	};
 
diff --git a/Engine/Source/Runtime/Core/Private/Serialization/FileRegions.cpp b/Engine/Source/Runtime/Core/Private/Serialization/FileRegions.cpp
--- a/Engine/Source/Runtime/Core/Private/Serialization/FileRegions.cpp
+++ b/Engine/Source/Runtime/Core/Private/Serialization/FileRegions.cpp
@@ -4,6 +4,8 @@
 
 #include "Misc/AssertionMacros.h"
 #include "Serialization/Archive.h"
+#include "Serialization/CompactBinary.h"
+#include "Serialization/CompactBinaryWriter.h"
 
 const TCHAR* FFileRegion::RegionsFileExtension = TEXT(".uregs");
 
@@ -16,6 +18,28 @@ inline FArchive& operator<<(FArchive& Ar, FFileRegion& Region)
 	return Ar;
 }
 
+FCbWriter& operator<<(FCbWriter& Writer, const FFileRegion& Region)
+{
+	Writer.BeginObject();
+
+	Writer << "offset" << Region.Offset;
+	Writer << "length" << Region.Length;
+	Writer << "type" << static_cast<uint8>(Region.Type);
+
+	Writer.EndObject();
+
+	return Writer;
+}
+
+bool FFileRegion::LoadFromCompactBinary(FCbFieldView Obj, FFileRegion& OutRegion)
+{
+	OutRegion.Offset = Obj["offset"].AsUInt64();
+	OutRegion.Length = Obj["length"].AsUInt64();
+	OutRegion.Type = static_cast<EFileRegionType>(Obj["type"].AsUInt8());
+
+	return true;
+}
+
 // NOTE: This serialization function must match FileRegion::ReadRegionsFromFile in AutomationUtils/FileRegions.cs
 void FFileRegion::SerializeFileRegions(class FArchive& Ar, TArray<FFileRegion>& Regions)
 {
diff --git a/Engine/Source/Runtime/Core/Public/Serialization/FileRegions.h b/Engine/Source/Runtime/Core/Public/Serialization/FileRegions.h
--- a/Engine/Source/Runtime/Core/Public/Serialization/FileRegions.h
+++ b/Engine/Source/Runtime/Core/Public/Serialization/FileRegions.h
@@ -7,6 +7,9 @@
 #include "HAL/Platform.h"
 #include "PixelFormat.h"
 
+class FCbFieldView;
+class FCbWriter;
+
 // Describes the type of data found in a given file region.
 enum class EFileRegionType : uint8
 {
@@ -47,6 +50,9 @@ struct FFileRegion
 	static CORE_API void AccumulateFileRegions(TArray<FFileRegion>& InOutRegions, int64 EntryOffset, int64 PayloadOffset, int64 EndOffset, TArrayView<const FFileRegion> InnerFileRegions);
 	static CORE_API void SerializeFileRegions(class FArchive& Ar, TArray<FFileRegion>& Regions);
 
+	CORE_API friend FCbWriter& operator<<(FCbWriter& Writer, const FFileRegion& Region);
+	static CORE_API bool LoadFromCompactBinary(FCbFieldView Obj, FFileRegion& OutRegion);
+
 	static inline EFileRegionType SelectType(EPixelFormat Format)
 	{
 		switch (Format)
