Always emit text to verify cross container refs output even if there are no problems.

diff --git a/Engine/Source/Developer/IoStoreUtilities/Private/IoStoreUtilities.cpp b/Engine/Source/Developer/IoStoreUtilities/Private/IoStoreUtilities.cpp
--- a/Engine/Source/Developer/IoStoreUtilities/Private/IoStoreUtilities.cpp
+++ b/Engine/Source/Developer/IoStoreUtilities/Private/IoStoreUtilities.cpp
@@ -6003,12 +6003,6 @@ int32 ValidateCrossContainerRefs(
 		}
 	}
 
-	UE_LOG(LogIoStore, Display, TEXT("Final valid edges: %d"), FinalValidEdges.Num());
-	for (const TPair<const FContainerDesc*, const FContainerDesc*>& Pair : FinalValidEdges)
-	{
-		UE_LOG(LogIoStore, Display, TEXT("%s -> %s"), *Pair.Key->Name.ToString(), *Pair.Value->Name.ToString());
-	}
-
 	TMap<TTuple<const FContainerDesc*, const FContainerDesc*>, TSet<TTuple<const FPackageDesc*, const FPackageDesc*>>> Errors;
 	Algo::SortBy(Packages, [](FPackageDesc* Desc) { return Desc->PackageName; }, FNameLexicalLess());
 	for (const FPackageDesc* Package : Packages)
@@ -6089,6 +6083,20 @@ int32 ValidateCrossContainerRefs(
 		OutputBuffer->SetSuppressEventTag(true);
 		OutputOverride = OutputBuffer.Get();
 	}
+
+	OutputOverride->Logf(ELogVerbosity::Display, TEXT("Invalid cross-container reference report"));
+	OutputOverride->Logf(ELogVerbosity::Display, TEXT("Final valid edges: %d"), FinalValidEdges.Num());
+	for (const TPair<const FContainerDesc*, const FContainerDesc*>& Pair : FinalValidEdges)
+	{
+		OutputOverride->Logf(ELogVerbosity::Display, TEXT("\t%s -> %s"), *Pair.Key->Name.ToString(), *Pair.Value->Name.ToString());
+	}
+
+	if (Errors.Num() == 0)
+	{
+		OutputOverride->Logf(ELogVerbosity::Display, TEXT("No errors."));
+		return 0;
+	}
+
 	for (const TPair<TTuple<const FContainerDesc*, const FContainerDesc*>, TSet<TTuple<const FPackageDesc*, const FPackageDesc*>>>& Pair : Errors)
 	{
 		OutputOverride->Logf(ELogVerbosity::Display, TEXT("%s -> %s"), *Pair.Key.Key->Name.ToString(), *Pair.Key.Value->Name.ToString());
