I/O - extended ondemand I/O store writer to output compressed and encrypted chunks

- renabled the ondemand PAK rule that was previously removed
- enabled handling more than one ondemand containe
- added ondemand I/O store TOC
- added support for serializing ondemand TOC's to CbObjet

diff --git a/Engine/Source/Developer/IoStoreUtilities/Private/IoStoreOnDemand.cpp b/Engine/Source/Developer/IoStoreUtilities/Private/IoStoreOnDemand.cpp
--- a/Engine/Source/Developer/IoStoreUtilities/Private/IoStoreOnDemand.cpp
+++ b/Engine/Source/Developer/IoStoreUtilities/Private/IoStoreOnDemand.cpp
@@ -2,250 +2,384 @@
 
 #include "IoStoreOnDemand.h"
 
+#include "Containers/ChunkedArray.h"
 #include "HAL/FileManager.h"
+#include "IO/IoChunkEncoding.h"
+#include "IO/IoDispatcher.h"
 #include "IO/IoStore.h"
-#include "Misc/FileHelper.h"
+#include "Memory/MemoryView.h"
 #include "Misc/Paths.h"
-#include "Serialization/JsonReader.h"
-#include "Serialization/JsonSerializer.h"
-#include "Serialization/JsonWriter.h"
 #include "Tasks/Task.h"
 #include "Tasks/Pipe.h"
+#include "Templates/SharedPointer.h"
+#include <atomic>
 
-namespace
+namespace UE::IO::Private
 {
 
-using FJsonWriter = TSharedPtr<TJsonWriter<TCHAR, TPrettyJsonPrintPolicy<TCHAR>>>; 
-using FJsonWriterFactory = TJsonWriterFactory<TCHAR, TPrettyJsonPrintPolicy<TCHAR>>;
+struct FPendingWrite
+{
+	FIoChunkId ChunkId;
+	TUniquePtr<IIoStoreWriteRequest> WriteRequest;
+	FIoWriteOptions WriteOptions;
+	FIoBuffer ChunkBuffer;
+	FIoBuffer ChunkHeader;
+	FIoHash ChunkHash;
+	FString ErrorText;
+	uint64 RawSize{0};
+	uint64 EncodedSize{0};
+};
+
+using FContainerEntries = TChunkedArray<FPendingWrite>;
+
+class FIoStoreOnDemandWriter;
 
-class FOnDemandIoStoreWriter final
+///////////////////////////////////////////////////////////////////////////////////////////////////////////////////////
+
+struct FOnDemandContainer
 	: public IIoStoreWriter
 {
-	struct FWriteStats
-	{
-		std::atomic<uint32> PendingCount{0};
-		std::atomic<uint32> TotalCount{0};
-		std::atomic<uint64> TotalByteCount{0};
-	};
-
-	struct FPendingWrite
-	{
-		FPendingWrite(const FIoChunkId& Id, IIoStoreWriteRequest* Request, const FIoWriteOptions& Options)
-			: ChunkId(Id)
-			, WriteRequest(Request)
-			, WriteOptions(Options)
-			, TaskPipe(UE_SOURCE_LOCATION)
-		{ }
-		
-		FIoChunkId ChunkId;
-		TUniquePtr<IIoStoreWriteRequest> WriteRequest;
-		FIoWriteOptions WriteOptions;
-		UE::Tasks::FPipe TaskPipe;
-		FString ErrorText;
-		uint64 UncompressedSize{0};
-		uint64 CompressedSize{0};
-	};
+	FOnDemandContainer(FIoStoreOnDemandWriter& Writer, const FString& ContainerName, const FIoContainerSettings& ContainerSettings);
+
+	virtual void SetReferenceChunkDatabase(TSharedPtr<IIoStoreWriterReferenceChunkDatabase> ReferenceChunkDatabase) override;
+	virtual void SetHashDatabase(TSharedPtr<IIoStoreWriterHashDatabase> HashDatabase, bool bVerifyHashDatabase) override;
+	virtual void EnableDiskLayoutOrdering(const TArray<TUniquePtr<FIoStoreReader>>& PatchSourceReaders = TArray<TUniquePtr<FIoStoreReader>>()) override;
+	virtual void EnumerateChunks(TFunction<bool(FIoStoreTocChunkInfo&&)>&& Callback) const override;
+	virtual void Append(const FIoChunkId& ChunkId, FIoBuffer Chunk, const FIoWriteOptions& WriteOptions, uint64 OrderHint = MAX_uint64) override;
+	virtual void Append(const FIoChunkId& ChunkId, IIoStoreWriteRequest* Request, const FIoWriteOptions& WriteOptions) override;
+	virtual TIoStatusOr<FIoStoreWriterResult> GetResult() override;
+
+	FIoStoreOnDemandWriter& Writer;
+	FString Name;
+	FIoContainerSettings Settings;
+	FContainerEntries Entries;
+	TIoStatusOr<FIoStoreWriterResult> WriteResult;
+};
 
-	static constexpr int32 MaxConcurrentWrites{64};
+///////////////////////////////////////////////////////////////////////////////////////////////////////////////////////
 
+class FIoStoreOnDemandWriter
+	: public IIoStoreOnDemandWriter
+{
 public:
-	FOnDemandIoStoreWriter(const FOnDemandWriterSettings& WriterSettings)
-		: Settings(WriterSettings)
-	{
-		Settings.TocFilePath = FPaths::SetExtension(Settings.TocFilePath, TEXT("utocmanifest"));
-	}
+	FIoStoreOnDemandWriter(const FIoStoreWriterSettings& WriterSettings, const FString& OutputDirectory, uint32 MaxConcurrentWrites);
 
-	virtual ~FOnDemandIoStoreWriter()
-	{ }
+	virtual TSharedPtr<IIoStoreWriter> CreateContainer(const FString& ContainerName, const FIoContainerSettings& ContainerSettings) override;
+	virtual void Flush() override;
 
-	virtual void SetReferenceChunkDatabase(TSharedPtr<IIoStoreWriterReferenceChunkDatabase> ReferenceChunkDatabase) override
-	{ }
+	void Append(
+		const FString& ContainerName,
+		const FIoContainerSettings& Settings,
+		const FIoChunkId& ChunkId,
+		IIoStoreWriteRequest* Request,
+		const FIoWriteOptions& WriteOptions);
 
-	virtual void SetHashDatabase(TSharedPtr<IIoStoreWriterHashDatabase> HashDatabase, bool bVerifyHashDatabase) override
-	{ }
+private:
+	FIoStoreWriterSettings WriterSettings;
+	FString OutputDirectory;
+	FString ChunksDirectory;
+	TMap<FString, TSharedPtr<FOnDemandContainer>> Containers;
+	FCriticalSection ContainersCS;
+	FEventRef WriteCompletedEvent;
+	UE::Tasks::FPipe TaskPipe;
+	std::atomic<uint32> PendingCount{0};
+	std::atomic<uint32> TotalCount{0};
+	uint32 MaxConcurrentWrites;
+};
 
-	virtual void EnableDiskLayoutOrdering(const TArray<TUniquePtr<FIoStoreReader>>& PatchSourceReaders = TArray<TUniquePtr<FIoStoreReader>>())
-	{ }
+FIoStoreOnDemandWriter::FIoStoreOnDemandWriter(const FIoStoreWriterSettings& Settings, const FString& InOutputDirectory, uint32 InMaxConcurrentWrites)
+	: WriterSettings(Settings)
+	, OutputDirectory(InOutputDirectory)
+	, TaskPipe(UE_SOURCE_LOCATION)
+	, MaxConcurrentWrites(InMaxConcurrentWrites)
+{
+	ChunksDirectory = FString::Printf(TEXT("%s/IoChunksV%u"), *OutputDirectory, EIoOnDemandTocVersion::Latest);
+}
 
-	virtual void Append(const FIoChunkId& ChunkId, FIoBuffer Chunk, const FIoWriteOptions& WriteOptions, uint64 OrderHint = MAX_uint64) override
-	{
-		struct FWriteRequest
-			: IIoStoreWriteRequest
-		{
-			FWriteRequest(FIoBuffer InSourceBuffer, uint64 InOrderHint)
-				: OrderHint(InOrderHint)
-			{
-				SourceBuffer = InSourceBuffer;
-				SourceBuffer.MakeOwned();
-			}
+TSharedPtr<IIoStoreWriter> FIoStoreOnDemandWriter::CreateContainer(const FString& ContainerName, const FIoContainerSettings& ContainerSettings)
+{
+	TSharedPtr<FOnDemandContainer> Container = MakeShareable(new FOnDemandContainer(*this, ContainerName, ContainerSettings));
+	Containers.Add(ContainerName, Container);
 
-			virtual ~FWriteRequest() = default;
+	return Container;
+}
 
-			void PrepareSourceBufferAsync(FGraphEventRef CompletionEvent) override
-			{
-				CompletionEvent->DispatchSubsequents();
-			}
+void FIoStoreOnDemandWriter::Flush() 
+{
+	TaskPipe.WaitUntilEmpty();
 
-			const FIoBuffer* GetSourceBuffer() override
-			{
-				return &SourceBuffer;
-			}
+	FScopeLock _(&ContainersCS);
 
-			void FreeSourceBuffer() override
-			{
-			}
+	FIoStoreOndemandTocResource TocResource;
+	TocResource.Header.Magic = FIoStoreOnDemandTocHeader::ExpectedMagic;
+	TocResource.Header.Version = static_cast<uint32>(EIoOnDemandTocVersion::Latest);
+	TocResource.Header.BlockSize = WriterSettings.CompressionBlockSize;
+	TocResource.Header.CompressionFormat = WriterSettings.CompressionMethod.ToString();
+	TocResource.Header.ChunksDirectory = FString::Printf(TEXT("IoChunksV%u"), EIoOnDemandTocVersion ::Latest);
 
-			uint64 GetOrderHint() override
-			{
-				return OrderHint;
-			}
+	for (auto& KV : Containers)
+	{
+		TSharedPtr<FOnDemandContainer>& Container = KV.Value;
+		FContainerEntries& WriteEntries = Container->Entries;
+		FIoStoreOnDemandContainerEntry& ContainerEntry = TocResource.Containers.AddDefaulted_GetRef();
+		FIoStoreWriterResult WriteResult; 
+		
+		ContainerEntry.ContainerName = Container->Name;
+		WriteResult.ContainerId = FIoContainerId::FromName(FName(Container->Name));
+		WriteResult.ContainerName = Container->Name;
 
-			TArrayView<const FFileRegion> GetRegions()
+		if (Container->Settings.IsCompressed())
+		{
+			WriteResult.CompressionMethod = WriterSettings.CompressionMethod;
+			WriteResult.ContainerFlags = EIoContainerFlags::OnDemand | EIoContainerFlags::Compressed;
+		}
+
+		if (Container->Settings.IsEncrypted())
+		{
+			check(Container->Settings.EncryptionKey.IsValid());
+			ContainerEntry.EncryptionKeyGuid = Container->Settings.EncryptionKeyGuid.ToString();
+			WriteResult.ContainerFlags |= EIoContainerFlags::Encrypted;
+		}
+
+		FString Error;
+		for (FPendingWrite& WriteEntry : WriteEntries)
+		{
+			if (!WriteEntry.ErrorText.IsEmpty())
 			{
-				return TArrayView<const FFileRegion>();
+				UE_LOG(LogIoStore, Error, TEXT("%s"), *WriteEntry.ErrorText);
+				Error = MoveTemp(WriteEntry.ErrorText);
+				break;
 			}
 
-			FIoBuffer SourceBuffer;
-			uint64 OrderHint;
-		};
+			FIoStoreOnDemandTocEntry& TocEntry = ContainerEntry.Entries.AddDefaulted_GetRef();
+			TocEntry.Hash = WriteEntry.ChunkHash;
+			TocEntry.ChunkId = WriteEntry.ChunkId;
+			TocEntry.RawSize = WriteEntry.RawSize;
+			TocEntry.EncodedSize = WriteEntry.EncodedSize;
+			TocEntry.BlockOffset = ContainerEntry.BlockSizes.Num();
+
+			const FIoChunkEncoding::FHeader* Header = FIoChunkEncoding::FHeader::Decode(WriteEntry.ChunkHeader.GetView());
+			check(Header);
+			TConstArrayView<uint32> Blocks = Header->GetBlocks();
+			TocEntry.BlockCount = Blocks.Num();
+			ContainerEntry.BlockSizes.Append(Blocks);
+
+			WriteResult.UncompressedContainerSize += WriteEntry.RawSize;
+			WriteResult.CompressedContainerSize += WriteEntry.EncodedSize;
+			WriteResult.TocEntryCount++;
+			WriteResult.TocSize += sizeof(FIoStoreOnDemandTocEntry) + (sizeof(uint32) + Blocks.Num());
+		}
 
-		Append(ChunkId, new FWriteRequest(Chunk, OrderHint), WriteOptions);
+		if (Error.Len())
+		{
+			Container->WriteResult = FIoStatus(EIoErrorCode::WriteError, Error);
+		}
+		else
+		{
+			Container->WriteResult = WriteResult;
+		}
 	}
 
-	virtual void Append(const FIoChunkId& ChunkId, IIoStoreWriteRequest* Request, const FIoWriteOptions& WriteOptions) override
+	TIoStatusOr<FString> Status = FIoStoreOndemandTocResource::Save(*OutputDirectory, TocResource);
+	if (Status.IsOk())
+	{
+		UE_LOG(LogIoStore, Display, TEXT("Saved ondemand TOC '%s'"), *Status.ConsumeValueOrDie());
+	}
+	else
 	{
-		check(Request);
+		UE_LOG(LogIoStore, Error, TEXT("Failed writing ondemand TOC, reason '%s'"), *Status.Status().ToString());
+	}
+}
 
-		for(;;)
+void FIoStoreOnDemandWriter::Append(
+	const FString& ContainerName,
+	const FIoContainerSettings& ContainerSettings,
+	const FIoChunkId& ChunkId,
+	IIoStoreWriteRequest* Request,
+	const FIoWriteOptions& WriteOptions)
+{
+	for(;;)
+	{
 		{
+			FScopeLock _(&ContainersCS);
+
+			if (PendingCount.load(std::memory_order_relaxed) < MaxConcurrentWrites)
 			{
-				FScopeLock _(&PendingLock);
-				if (WriteStats.PendingCount.load() < MaxConcurrentWrites)
-				{
-					WriteStats.PendingCount++;
-					FPendingWrite* PendingWrite = PendingWrites.Add_GetRef(MakeUnique<FPendingWrite>(ChunkId, Request, WriteOptions)).Get();
-					
-					// Setup the task pipe when holding the lock to make it easy to flush all pending task pipe(s)
-
-					FGraphEventRef Event = FGraphEvent::CreateGraphEvent();
-					PendingWrite->WriteRequest->PrepareSourceBufferAsync(Event);
-					UE::Tasks::FTask ReadChunkTask = PendingWrite->TaskPipe.Launch(TEXT("ReadChunk"), [PendingWrite, Event]() mutable
+				PendingCount.fetch_add(1, std::memory_order_relaxed);
+				
+				TSharedPtr<FOnDemandContainer>& Container = Containers.FindChecked(ContainerName);
+				FPendingWrite* PendingWrite = new(Container->Entries) FPendingWrite();
+				PendingWrite->ChunkId = ChunkId;
+				PendingWrite->WriteRequest.Reset(Request);
+				PendingWrite->WriteOptions = WriteOptions;
+
+				FGraphEventRef Event = FGraphEvent::CreateGraphEvent();
+				PendingWrite->WriteRequest->PrepareSourceBufferAsync(Event);
+
+				UE::Tasks::FTask ReadChunkTask = TaskPipe.Launch(TEXT("ReadChunk"),
+					[PendingWrite, Event]() mutable
 					{
 						Event->Wait();
 					});
 
-					// Add compression/encryption tasks if needed
-					
-					UE::Tasks::FTask WriteChunkTask = PendingWrite->TaskPipe.Launch(TEXT("WriteChunk"), [this, PendingWrite]() mutable
+				UE::Tasks::FTask EncodeChunkTask = TaskPipe.Launch(TEXT("EncodeChunk"),
+					[this, PendingWrite, &ContainerSettings]() mutable
 					{
 						if (const FIoBuffer* SourceBuffer = PendingWrite->WriteRequest->GetSourceBuffer())
 						{
-							FString FileName = LexToString(PendingWrite->ChunkId) + TEXT(".iochunk");
-							FString FilePath = Settings.OutputDirectory / FileName; 
+							PendingWrite->ChunkBuffer = *SourceBuffer;
+							PendingWrite->ChunkBuffer.EnsureOwned();
+							PendingWrite->ChunkHash = FIoHash::HashBuffer(PendingWrite->ChunkBuffer.GetView());
+							PendingWrite->RawSize = PendingWrite->ChunkBuffer.GetSize();
+
+							PendingWrite->WriteRequest->FreeSourceBuffer();
+							PendingWrite->WriteRequest.Reset();
+
+							const FAES::FAESKey& Key = ContainerSettings.EncryptionKey;
+							FMemoryView KeyView = Key.IsValid() ? MakeMemoryView(Key.Key, Key.KeySize) : FMemoryView();
+							FIoChunkEncodingParams Params {WriterSettings.CompressionMethod, KeyView, uint32(WriterSettings.CompressionBlockSize)};
 
-							if (TUniquePtr<FArchive> Ar(IFileManager::Get().CreateFileWriter(*FilePath)); Ar.IsValid())
+							const bool bEncoded = FIoChunkEncoding::Encode(Params, PendingWrite->ChunkBuffer.GetView(), PendingWrite->ChunkHeader, PendingWrite->ChunkBuffer);
+							if (bEncoded)
 							{
-								const uint64 ChunkSize = SourceBuffer->GetSize();
-								const uint32 CurrentCount = WriteStats.TotalCount.fetch_add(1) + 1; 
-								UE_CLOG((CurrentCount % 128 == 0), LogIoStore, Display,
-									TEXT("Writing ondemand chunk #%u '%s' -> '%s' (%llu Bytes)"),
-									CurrentCount, *PendingWrite->WriteOptions.FileName, *FileName, ChunkSize);
-
-								Ar->Serialize((void*)SourceBuffer->GetData(), ChunkSize);
-								PendingWrite->UncompressedSize = PendingWrite->CompressedSize = ChunkSize;
-								PendingWrite->WriteRequest->FreeSourceBuffer();
-								WriteStats.TotalByteCount += ChunkSize;
+								PendingWrite->EncodedSize = PendingWrite->ChunkBuffer.GetSize();
+								PendingWrite->ChunkHash = FIoHash::HashBuffer(PendingWrite->ChunkBuffer.GetView());
 							}
 							else
 							{
-								PendingWrite->ErrorText = FString::Printf(TEXT("Failed to create file '%s'"), *FileName);
+								PendingWrite->ErrorText = FString::Printf(TEXT("Failed to compress '%s'"), *PendingWrite->WriteOptions.FileName);
 							}
 						}
 						else
 						{
-							PendingWrite->ErrorText = TEXT("Invalid source buffer");
+							PendingWrite->ErrorText = FString::Printf(TEXT("Failed to read source buffer '%s'"), *PendingWrite->WriteOptions.FileName);
 						}
+					}, ReadChunkTask);
+
+					UE::Tasks::FTask WriteChunkTask = TaskPipe.Launch(TEXT("WriteChunk"),
+						[this, PendingWrite]() mutable
+						{
+							FIoBuffer ChunkBuffer = MoveTemp(PendingWrite->ChunkBuffer);
+							if (ChunkBuffer.GetSize() > 0)
+							{
+								const FString HashString = LexToString(PendingWrite->ChunkHash);
+								const FString FileName = HashString + TEXT(".iochunk");
+								const FString FilePath = ChunksDirectory / HashString.Left(2) / FileName;
+
+								if (TUniquePtr<FArchive> Ar(IFileManager::Get().CreateFileWriter(*FilePath)); Ar.IsValid())
+								{
+									const uint64 ChunkSize = ChunkBuffer.GetSize();
+									const uint32 CurrentCount = TotalCount.fetch_add(1) + 1; 
+									UE_CLOG((CurrentCount % 128 == 0), LogIoStore, Display,
+										TEXT("Writing ondemand chunk #%u '%s' -> '%s' (%llu Bytes)"),
+										CurrentCount, *PendingWrite->WriteOptions.FileName, *FileName, ChunkSize);
+
+									Ar->Serialize((void*)ChunkBuffer.GetData(), ChunkSize);
+								}
+								else
+								{
+									PendingWrite->ErrorText = FString::Printf(TEXT("Failed to create file '%s'"), *FileName);
+								}
+							}
+							else
+							{
+								PendingWrite->ErrorText = TEXT("Invalid source buffer");
+							}
 
-						PendingWrite->WriteRequest.Reset();
+							PendingCount--;
+							WriteCompletedEvent->Trigger();
+						}, EncodeChunkTask);
 
-						WriteStats.PendingCount--;
-						WriteCompletedEvent->Trigger();
-					}, ReadChunkTask);
-					
-					break;
-				}
+				break;
 			}
-
-			WriteCompletedEvent->Wait();
 		}
+
+		WriteCompletedEvent->Wait();
 	}
+}
+
+///////////////////////////////////////////////////////////////////////////////////////////////////////////////////////
+
+FOnDemandContainer::FOnDemandContainer(FIoStoreOnDemandWriter& OnDemandWriter, const FString& ContainerName, const FIoContainerSettings& ContainerSettings)
+	: Writer(OnDemandWriter)
+	, Name(ContainerName)
+	, Settings(ContainerSettings)
+{
+}
+
+void FOnDemandContainer::SetReferenceChunkDatabase(TSharedPtr<IIoStoreWriterReferenceChunkDatabase> ReferenceChunkDatabase)
+{ }
 
-	virtual TIoStatusOr<FIoStoreWriterResult> GetResult() override
+void FOnDemandContainer::SetHashDatabase(TSharedPtr<IIoStoreWriterHashDatabase> HashDatabase, bool bVerifyHashDatabase)
+{ }
+
+void FOnDemandContainer::EnableDiskLayoutOrdering(const TArray<TUniquePtr<FIoStoreReader>>& PatchSourceReaders)
+{ }
+
+void FOnDemandContainer::EnumerateChunks(TFunction<bool(FIoStoreTocChunkInfo&&)>&& Callback) const
+{ }
+
+void FOnDemandContainer::Append(const FIoChunkId& ChunkId, FIoBuffer Chunk, const FIoWriteOptions& WriteOptions, uint64 OrderHint)
+{
+	struct FWriteRequest
+		: IIoStoreWriteRequest
 	{
-		FScopeLock _(&PendingLock);
+		FWriteRequest(FIoBuffer InSourceBuffer, uint64 InOrderHint)
+			: OrderHint(InOrderHint)
+		{
+			SourceBuffer = InSourceBuffer;
+			SourceBuffer.MakeOwned();
+		}
 
-		for (TUniquePtr<FPendingWrite>& PendingWrite : PendingWrites)
+		virtual ~FWriteRequest() = default;
+
+		void PrepareSourceBufferAsync(FGraphEventRef CompletionEvent) override
 		{
-			PendingWrite->TaskPipe.WaitUntilEmpty();
+			CompletionEvent->DispatchSubsequents();
 		}
-		
-		FIoStoreWriterResult WriteResult;
-		WriteResult.ContainerId = FIoContainerId::FromName(Settings.ContainerName);
-		WriteResult.ContainerName = Settings.ContainerName.ToString();
-		
-		FString Json;
-		FJsonWriter JsonWriter = FJsonWriterFactory::Create(&Json);
-		JsonWriter->WriteObjectStart();
-		JsonWriter->WriteValue(TEXT("ContainerName"), Settings.ContainerName.ToString());
-		JsonWriter->WriteArrayStart(TEXT("Chunks"));
-		
-		for (TUniquePtr<FPendingWrite>& PendingWrite : PendingWrites)
+
+		const FIoBuffer* GetSourceBuffer() override
 		{
-			if (PendingWrite->ErrorText.IsEmpty() == false)
-			{
-				UE_LOG(LogIoStore, Display, TEXT("Failed to write on demand chunk '%s', reason '%s'"), *LexToString(PendingWrite->ChunkId), *PendingWrite->ErrorText);
-				continue;
-			}
-			JsonWriter->WriteObjectStart();
-			JsonWriter->WriteValue(TEXT("IoChunkId"), LexToString(PendingWrite->ChunkId));
-			JsonWriter->WriteValue(TEXT("FileName"), *PendingWrite->WriteOptions.FileName);
-			JsonWriter->WriteObjectEnd();
-			
-			WriteResult.UncompressedContainerSize += PendingWrite->UncompressedSize;
-			WriteResult.CompressedContainerSize += PendingWrite->CompressedSize;
-			WriteResult.TocEntryCount++;
+			return &SourceBuffer;
 		}
 
-		JsonWriter->WriteArrayEnd();
-		JsonWriter->WriteObjectEnd();
-		JsonWriter->Close();
+		void FreeSourceBuffer() override
+		{
+		}
 
-		UE_LOG(LogIoStore, Display, TEXT("Saving ondemand JSON manifest '%s'"), *Settings.TocFilePath);
-		if (!FFileHelper::SaveStringToFile(Json, *Settings.TocFilePath))
+		uint64 GetOrderHint() override
 		{
-			return FIoStatus(EIoErrorCode::CorruptToc, TEXTVIEW("Failed writing ondemand JSON manifest")); 
+			return OrderHint;
 		}
-		
-		WriteResult.TocSize = sizeof(TCHAR) + Json.Len();
 
-		return WriteResult;
-	}
+		TArrayView<const FFileRegion> GetRegions()
+		{
+			return TArrayView<const FFileRegion>();
+		}
 
-	virtual void EnumerateChunks(TFunction<bool(FIoStoreTocChunkInfo&&)>&& Callback) const override
-	{
-	}
+		FIoBuffer SourceBuffer;
+		uint64 OrderHint;
+	};
 
-private:
-	FOnDemandWriterSettings Settings;
-	FWriteStats WriteStats;
-	FCriticalSection PendingLock;
-	TArray<TUniquePtr<FPendingWrite>> PendingWrites;
-	FEventRef WriteCompletedEvent;
-};
+	Append(ChunkId, new FWriteRequest(Chunk, OrderHint), WriteOptions);
+}
+
+void FOnDemandContainer::Append(const FIoChunkId& ChunkId, IIoStoreWriteRequest* Request, const FIoWriteOptions& WriteOptions)
+{
+	Writer.Append(Name, Settings, ChunkId, Request, WriteOptions);
+}
+
+TIoStatusOr<FIoStoreWriterResult> FOnDemandContainer::GetResult() 
+{
+	return WriteResult;
+}
 
-} // namespace
+} // namespace UE::IO::Private 
 
-TUniquePtr<IIoStoreWriter> MakeOnDemandIoStoreWriter(const FOnDemandWriterSettings& WriterSettings)
+TUniquePtr<IIoStoreOnDemandWriter> MakeIoStoreOnDemandWriter(
+	const FIoStoreWriterSettings& WriterSettings,
+	const FString& OutputDirectory,
+	uint32 MaxConcurrentWrites)
 {
-	return MakeUnique<FOnDemandIoStoreWriter>(WriterSettings);
+	return MakeUnique<UE::IO::Private::FIoStoreOnDemandWriter>(WriterSettings, OutputDirectory, MaxConcurrentWrites);
 }
diff --git a/Engine/Source/Developer/IoStoreUtilities/Private/IoStoreOnDemand.h b/Engine/Source/Developer/IoStoreUtilities/Private/IoStoreOnDemand.h
--- a/Engine/Source/Developer/IoStoreUtilities/Private/IoStoreOnDemand.h
+++ b/Engine/Source/Developer/IoStoreUtilities/Private/IoStoreOnDemand.h
@@ -3,16 +3,21 @@
 #pragma once
 
 #include "Containers/UnrealString.h"
+#include "IO/IoDispatcher.h"
 #include "Templates/UniquePtr.h"
 #include "UObject/NameTypes.h"
 
 class IIoStoreWriter;
 
-struct FOnDemandWriterSettings
+class IIoStoreOnDemandWriter
 {
-	FString TocFilePath;
-	FString OutputDirectory;
-	FName ContainerName;
+public:
+	virtual ~IIoStoreOnDemandWriter() = default;
+	virtual TSharedPtr<IIoStoreWriter> CreateContainer(const FString& ContainerName, const FIoContainerSettings& ContainerSettings) = 0;
+	virtual void Flush() = 0;
 };
 
-TUniquePtr<IIoStoreWriter> MakeOnDemandIoStoreWriter(const FOnDemandWriterSettings& WriterSettings);
+TUniquePtr<IIoStoreOnDemandWriter> MakeIoStoreOnDemandWriter(
+	const FIoStoreWriterSettings& WriterSettings,
+	const FString& OutputDirectory,
+	uint32 MaxConcurrentWrites = 64);
diff --git a/Engine/Source/Developer/IoStoreUtilities/Private/IoStoreUtilities.cpp b/Engine/Source/Developer/IoStoreUtilities/Private/IoStoreUtilities.cpp
--- a/Engine/Source/Developer/IoStoreUtilities/Private/IoStoreUtilities.cpp
+++ b/Engine/Source/Developer/IoStoreUtilities/Private/IoStoreUtilities.cpp
@@ -2625,11 +2625,12 @@ void LogWriterResults(const TArray<FIoStoreWriterResult>& Results)
 				*Result.CompressionMethod.ToString());
 		}
 
-		FString ContainerSettings = FString::Printf(TEXT("%s/%s/%s/%s"),
+		FString ContainerSettings = FString::Printf(TEXT("%s/%s/%s/%s/%s"),
 			EnumHasAnyFlags(Result.ContainerFlags, EIoContainerFlags::Compressed) ? TEXT("C") : TEXT("-"),
 			EnumHasAnyFlags(Result.ContainerFlags, EIoContainerFlags::Encrypted) ? TEXT("E") : TEXT("-"),
 			EnumHasAnyFlags(Result.ContainerFlags, EIoContainerFlags::Signed) ? TEXT("S") : TEXT("-"),
-			EnumHasAnyFlags(Result.ContainerFlags, EIoContainerFlags::Indexed) ? TEXT("I") : TEXT("-"));
+			EnumHasAnyFlags(Result.ContainerFlags, EIoContainerFlags::Indexed) ? TEXT("I") : TEXT("-"),
+			EnumHasAnyFlags(Result.ContainerFlags, EIoContainerFlags::OnDemand) ? TEXT("O") : TEXT("-"));
 
 		UE_LOG(LogIoStore, Display, TEXT("%-30s %10s %15.2lf %15llu %15.2lf %25s"),
 			*Result.ContainerName,
@@ -2655,7 +2656,7 @@ void LogWriterResults(const TArray<FIoStoreWriterResult>& Results)
 		TEXT("-"));
 
 	UE_LOG(LogIoStore, Display, TEXT(""));
-	UE_LOG(LogIoStore, Display, TEXT("** Flags: (C)ompressed / (E)ncrypted / (S)igned) / (I)ndexed) **"));
+	UE_LOG(LogIoStore, Display, TEXT("** Flags: (C)ompressed / (E)ncrypted / (S)igned) / (I)ndexed) / (O)nDemand **"));
 	UE_LOG(LogIoStore, Display, TEXT(""));
 	UE_LOG(LogIoStore, Display, TEXT("Compression block padding: %8.2lf MB"), (double)TotalPaddingSize / 1024.0 / 1024.0);
 	UE_LOG(LogIoStore, Display, TEXT(""));
@@ -3719,6 +3720,7 @@ int32 CreateTarget(const FIoStoreArguments& Arguments, const FIoStoreWriterSetti
 	}
 	TArray<TSharedPtr<IIoStoreWriter>> IoStoreWriters;
 	TSharedPtr<IIoStoreWriter> GlobalIoStoreWriter;
+	TUniquePtr<IIoStoreOnDemandWriter> IoStoreOnDemandWriter;
 	{
 		IOSTORE_CPU_SCOPE(InitializeWriters);
 		if (!Arguments.IsDLC())
@@ -3737,50 +3739,72 @@ int32 CreateTarget(const FIoStoreArguments& Arguments, const FIoStoreWriterSetti
 		{
 			IOSTORE_CPU_SCOPE(InitializeWriter);
 			check(ContainerTarget->ContainerId.IsValid());
-			if (!ContainerTarget->OutputPath.IsEmpty())
+
+			if (ContainerTarget->OutputPath.IsEmpty() && ContainerTarget->OnDemandOutputPath.IsEmpty())
+			{
+				continue;
+			}
+
+			if (!ContainerTarget->OnDemandOutputPath.IsEmpty())
 			{
-				if (!ContainerTarget->OnDemandOutputPath.IsEmpty())
+				if (!IoStoreOnDemandWriter.IsValid())
 				{
-					FOnDemandWriterSettings WriterSettings;
-					WriterSettings.ContainerName = ContainerTarget->Name;
-					WriterSettings.TocFilePath = ContainerTarget->OutputPath;
-					WriterSettings.OutputDirectory = ContainerTarget->OnDemandOutputPath;
-					ContainerTarget->IoStoreWriter = TSharedPtr<IIoStoreWriter>(MakeOnDemandIoStoreWriter(WriterSettings).Release());
-					IoStoreWriters.Add(ContainerTarget->IoStoreWriter);
+					FIoStoreWriterSettings WriterSettings;
+					WriterSettings.CompressionMethod = NAME_Oodle;
+					IoStoreOnDemandWriter = MakeIoStoreOnDemandWriter(WriterSettings, ContainerTarget->OnDemandOutputPath);
 				}
-				else
+
+				FIoContainerSettings ContainerSettings;
+				ContainerSettings.ContainerFlags = EIoContainerFlags::OnDemand | EIoContainerFlags::Compressed;
+
+				if (EnumHasAnyFlags(ContainerTarget->ContainerFlags, EIoContainerFlags::Encrypted))
 				{
-					FIoContainerSettings ContainerSettings;
-					ContainerSettings.ContainerId = ContainerTarget->ContainerId;
-					if (Arguments.bCreateDirectoryIndex)
+					if (const FNamedAESKey* Key = Arguments.KeyChain.GetEncryptionKeys().Find(ContainerTarget->EncryptionKeyGuid))
 					{
-						ContainerSettings.ContainerFlags = ContainerTarget->ContainerFlags | EIoContainerFlags::Indexed;
-					}
-					if (EnumHasAnyFlags(ContainerTarget->ContainerFlags, EIoContainerFlags::Encrypted))
-					{
-						const FNamedAESKey* Key = Arguments.KeyChain.GetEncryptionKeys().Find(ContainerTarget->EncryptionKeyGuid);
-						check(Key);
 						ContainerSettings.EncryptionKeyGuid = ContainerTarget->EncryptionKeyGuid;
 						ContainerSettings.EncryptionKey = Key->Key;
+						ContainerSettings.ContainerFlags |= EIoContainerFlags::Encrypted;
 					}
-					if (EnumHasAnyFlags(ContainerTarget->ContainerFlags, EIoContainerFlags::Signed))
-					{
-						ContainerSettings.SigningKey = Arguments.KeyChain.GetSigningKey();
-						ContainerSettings.ContainerFlags |= EIoContainerFlags::Signed;
-					}
-					ContainerSettings.bGenerateDiffPatch = ContainerTarget->bGenerateDiffPatch;
-					ContainerTarget->IoStoreWriter = IoStoreWriterContext->CreateContainer(*ContainerTarget->OutputPath, ContainerSettings);
-					ContainerTarget->IoStoreWriter->EnableDiskLayoutOrdering(ContainerTarget->PatchSourceReaders);
-					ContainerTarget->IoStoreWriter->SetReferenceChunkDatabase(ChunkDatabase);
-					ContainerTarget->IoStoreWriter->SetHashDatabase(HashDatabase, Arguments.bVerifyHashDatabase);
-					IoStoreWriters.Add(ContainerTarget->IoStoreWriter);
-					if (!ContainerTarget->OptionalSegmentOutputPath.IsEmpty())
+					else
 					{
-						ContainerTarget->OptionalSegmentIoStoreWriter = IoStoreWriterContext->CreateContainer(*ContainerTarget->OptionalSegmentOutputPath, ContainerSettings);
-					    ContainerTarget->OptionalSegmentIoStoreWriter->SetReferenceChunkDatabase(ChunkDatabase);
-						IoStoreWriters.Add(ContainerTarget->OptionalSegmentIoStoreWriter);
+						UE_LOG(LogIoStore, Warning, TEXT("Failed to find encryption key '%s' for container '%s'"), *ContainerTarget->EncryptionKeyGuid.ToString(), *ContainerTarget->Name.ToString());
 					}
 				}
+				ContainerTarget->IoStoreWriter = IoStoreOnDemandWriter->CreateContainer(ContainerTarget->Name.ToString(), ContainerSettings);
+				IoStoreWriters.Add(ContainerTarget->IoStoreWriter);
+			}
+			else
+			{
+				FIoContainerSettings ContainerSettings;
+				ContainerSettings.ContainerId = ContainerTarget->ContainerId;
+				if (Arguments.bCreateDirectoryIndex)
+				{
+					ContainerSettings.ContainerFlags = ContainerTarget->ContainerFlags | EIoContainerFlags::Indexed;
+				}
+				if (EnumHasAnyFlags(ContainerTarget->ContainerFlags, EIoContainerFlags::Encrypted))
+				{
+					const FNamedAESKey* Key = Arguments.KeyChain.GetEncryptionKeys().Find(ContainerTarget->EncryptionKeyGuid);
+					check(Key);
+					ContainerSettings.EncryptionKeyGuid = ContainerTarget->EncryptionKeyGuid;
+					ContainerSettings.EncryptionKey = Key->Key;
+				}
+				if (EnumHasAnyFlags(ContainerTarget->ContainerFlags, EIoContainerFlags::Signed))
+				{
+					ContainerSettings.SigningKey = Arguments.KeyChain.GetSigningKey();
+					ContainerSettings.ContainerFlags |= EIoContainerFlags::Signed;
+				}
+				ContainerSettings.bGenerateDiffPatch = ContainerTarget->bGenerateDiffPatch;
+				ContainerTarget->IoStoreWriter = IoStoreWriterContext->CreateContainer(*ContainerTarget->OutputPath, ContainerSettings);
+				ContainerTarget->IoStoreWriter->EnableDiskLayoutOrdering(ContainerTarget->PatchSourceReaders);
+				ContainerTarget->IoStoreWriter->SetReferenceChunkDatabase(ChunkDatabase);
+				ContainerTarget->IoStoreWriter->SetHashDatabase(HashDatabase, Arguments.bVerifyHashDatabase);
+				IoStoreWriters.Add(ContainerTarget->IoStoreWriter);
+				if (!ContainerTarget->OptionalSegmentOutputPath.IsEmpty())
+				{
+					ContainerTarget->OptionalSegmentIoStoreWriter = IoStoreWriterContext->CreateContainer(*ContainerTarget->OptionalSegmentOutputPath, ContainerSettings);
+					ContainerTarget->OptionalSegmentIoStoreWriter->SetReferenceChunkDatabase(ChunkDatabase);
+					IoStoreWriters.Add(ContainerTarget->OptionalSegmentIoStoreWriter);
+				}
 			}
 		}
 	}
@@ -3937,6 +3961,13 @@ int32 CreateTarget(const FIoStoreArguments& Arguments, const FIoStoreWriterSetti
 		GlobalIoStoreWriter->Append(CreateIoChunkId(0, 0, EIoChunkType::ScriptObjects), ScriptObjectsBuffer, WriteOptions);
 	}
 
+	if (IoStoreOnDemandWriter.IsValid())
+	{
+		UE_LOG(LogIoStore, Display, TEXT("Serializing ondemand content..."));
+		IOSTORE_CPU_SCOPE(SerializingContentOnDemand);
+		IoStoreOnDemandWriter->Flush();
+	}
+
 	UE_LOG(LogIoStore, Display, TEXT("Serializing container(s)..."));
 	{
 		IOSTORE_CPU_SCOPE(Serializing);
@@ -7366,15 +7397,17 @@ bool ParseContainerGenerationArguments(FIoStoreArguments& Arguments, FIoStoreWri
 		{
 			FContainerSourceSpec& ContainerSpec = Arguments.Containers.AddDefaulted_GetRef();
 
-			if (!FParse::Value(*Command, TEXT("Output="), ContainerSpec.OutputPath))
+			if (FParse::Value(*Command, TEXT("Output="), ContainerSpec.OutputPath))
+			{
+				ContainerSpec.OutputPath = FPaths::ChangeExtension(ContainerSpec.OutputPath, TEXT(""));
+			}
+			else if (!FParse::Value(*Command, TEXT("OnDemandOutput="), ContainerSpec.OnDemandOutputPath))
 			{
 				UE_LOG(LogIoStore, Error, TEXT("Output argument missing from command '%s'"), *Command);
 				return false;
 			}
-			ContainerSpec.OutputPath = FPaths::ChangeExtension(ContainerSpec.OutputPath, TEXT(""));
 
 			FParse::Value(*Command, TEXT("OptionalOutput="), ContainerSpec.OptionalOutputPath);
-			FParse::Value(*Command, TEXT("OnDemandOutput="), ContainerSpec.OnDemandOutputPath);
 
 			FString ContainerName;
 			if (FParse::Value(*Command, TEXT("ContainerName="), ContainerName))
diff --git a/Engine/Source/Programs/AutomationTool/Scripts/CopyBuildToStagingDirectory.Automation.cs b/Engine/Source/Programs/AutomationTool/Scripts/CopyBuildToStagingDirectory.Automation.cs
--- a/Engine/Source/Programs/AutomationTool/Scripts/CopyBuildToStagingDirectory.Automation.cs
+++ b/Engine/Source/Programs/AutomationTool/Scripts/CopyBuildToStagingDirectory.Automation.cs
@@ -281,6 +281,7 @@ namespace AutomationScripts
 			string ContainerName,
 			FileReference PakOutputLocation,
 			DirectoryReference OptionalOutputLocation,
+			DirectoryReference OnDemandOutputLocation,
 			bool bCompressed,
 			EncryptionAndSigning.CryptoSettings CryptoSettings,
 			string EncryptionKeyGuid,
@@ -289,7 +290,17 @@ namespace AutomationScripts
 			bool bIsDLC)
 		{
 			StringBuilder CmdLine = new StringBuilder();
-			CmdLine.AppendFormat("-Output={0}", MakePathSafeToUseWithCommandLine(Path.ChangeExtension(PakOutputLocation.FullName, ".utoc")));
+
+			if (OnDemandOutputLocation != null)
+			{
+				FileReference OnDemandPakOutPutLocation = FileReference.Combine(OnDemandOutputLocation, ContainerName);
+				CmdLine.AppendFormat(" -OnDemandOutput={0}", MakePathSafeToUseWithCommandLine(OnDemandOutputLocation.FullName));
+			}
+			else
+			{
+				CmdLine.AppendFormat("-Output={0}", MakePathSafeToUseWithCommandLine(Path.ChangeExtension(PakOutputLocation.FullName, ".utoc")));
+			}
+
 			if (OptionalOutputLocation != null)
 			{
 				CmdLine.AppendFormat(" -OptionalOutput={0}", MakePathSafeToUseWithCommandLine(OptionalOutputLocation.FullName));
@@ -2150,7 +2161,22 @@ namespace AutomationScripts
 			// List of pak files to use instead of manifest
 			public List<string> OverridePaks;
 
+			// Whether this rule is defined for content-on-demand
+			public bool bOnDemand;
+
 			public bool bStageLoose;
+
+			// Encryption key
+			public string EncryptionKeyGuid;
+
+			public static bool IsMatch(PakFileRules PakRules, KeyValuePair<string, string> StagingFile)
+			{
+				bool bMatched = !PakRules.bDisabled &&
+					((!PakRules.bUseExactFilePathMatch && PakRules.Filter != null && PakRules.Filter.Matches(StagingFile.Key)) ||
+					(PakRules.bUseExactFilePathMatch && PakRules.ExactFileMatches != null && PakRules.ExactFileMatches.Contains(StagingFile.Value)));
+				
+				return bMatched;
+			}
 		};
 
 		/// <summary>
@@ -2284,6 +2310,8 @@ namespace AutomationScripts
 				PakRulesConfig.TryGetValue(SectionName, "bStageLoose", out PakRules.bStageLoose);
 				PakRulesConfig.TryGetValue(SectionName, "bOverrideChunkManifest", out PakRules.bOverrideChunkManifest);
 				PakRulesConfig.TryGetValue(SectionName, "bDisabled", out PakRules.bDisabled);
+				PakRulesConfig.TryGetValue(SectionName, "bOnDemand", out PakRules.bOnDemand);
+				PakRulesConfig.TryGetValue(SectionName, "EncryptionKeyGuid", out PakRules.EncryptionKeyGuid);
 				string PakString;
 				PakRulesConfig.TryGetValue(SectionName, "OverridePaks", out PakString);
 
@@ -2296,12 +2324,20 @@ namespace AutomationScripts
 					PakRules.OverridePaks = null;
 				}
 
-				if (PakRules.bExcludeFromPaks && PakRules.OverridePaks != null)
+				if (PakRules.bExcludeFromPaks)
 				{
-					Logger.LogWarning("Error in PakFileRules {Arg0}, set to exclude but also sets override!", PakRules.Name);
-					continue;
+					if (PakRules.OverridePaks != null)
+					{
+						LogWarning("Error in PakFileRules {Arg0}, set to exclude but also sets override!", PakRules.Name);
+						continue;
+					}
+					if (PakRules.bOnDemand)
+					{
+						LogWarning("Error in PakFileRules {Arg0}, set to exclude but also sets to ondemand!", PakRules.Name);
+						continue;
+					}
 				}
-				else if (!PakRules.bExcludeFromPaks && PakRules.OverridePaks == null)
+				else if (PakRules.OverridePaks == null && !PakRules.bOnDemand)
 				{
 					Logger.LogWarning("Error in PakFileRules {Arg0}, set to include but did not specify paks!", PakRules.Name);
 					continue;
@@ -2348,14 +2384,48 @@ namespace AutomationScripts
 				return;
 			}
 
+			// Process content-on-demand rules first
+			if (ModifyPakList != null)
+			{
+				foreach (PakFileRules Rule in RulesList.Where(Rule => Rule.bOnDemand))
+				{
+					if (PakFileRules.IsMatch(Rule, StagingFile))
+					{
+						ModifyPakList.Clear();
+
+						string ChunkName = Rule.Name;
+						if (Rule.OverridePaks != null && Rule.OverridePaks.Count > 0)
+						{
+							string OverrideChunkName = Rule.OverridePaks[0];
+							if (ChunkNameToDefinition.ContainsKey(OverrideChunkName))
+							{
+								ChunkName = OverrideChunkName;
+							}
+							else
+							{
+								LogWarning("Undefined PAK override '{0}' in PAK rule '{1}'", OverrideChunkName, Rule.Name);
+							}
+						}
+
+						ChunkDefinition Chunk = null;
+						if (ChunkNameToDefinition.TryGetValue(ChunkName, out Chunk))
+						{
+							ModifyPakList.Add(Chunk);
+						}
+						else
+						{
+							LogWarning("Undefined chunk name '{0}' in PAK rule '{1}'", ChunkName, Rule.Name);
+						}
+						
+						return;
+					}
+				}
+			}
+
 			// Search in order, return on first match
 			foreach (var PakRules in RulesList)
 			{
-				bool bMatched = !PakRules.bDisabled &&
-					((!PakRules.bUseExactFilePathMatch && PakRules.Filter != null && PakRules.Filter.Matches(StagingFile.Key)) ||
-					(PakRules.bUseExactFilePathMatch && PakRules.ExactFileMatches != null && PakRules.ExactFileMatches.Contains(StagingFile.Value)));
-
-				if (bMatched)
+				if (PakFileRules.IsMatch(PakRules, StagingFile))
 				{
 					if (PakRules.bStageLoose)
 					{
@@ -2415,7 +2485,6 @@ namespace AutomationScripts
 			}
 		}
 
-
 		/// <summary>
 		/// Creates a pak file using staging context (single manifest)
 		/// </summary>
@@ -2595,6 +2664,11 @@ namespace AutomationScripts
 			/// Whether to rehydrate the assets when creating the pak file or not
 			/// </summary>
 			public bool bRehydrateAssets;
+			
+			/// <summary>
+			// Whether this pak file is for content-on-demand content delivery or not
+			/// </summary>
+			public bool bOnDemand;
 
 			/// <summary>
 			/// GUID of the encryption key for this pak file
@@ -2607,13 +2681,14 @@ namespace AutomationScripts
 			/// <param name="PakName">Path to the base output file for this pak file</param>
 			/// <param name="UnrealPakResponseFile">Map of files within the pak file to their source file on disk</param>
 			/// <param name="bCompressed">Whether to enable compression</param>
-			public CreatePakParams(string PakName, Dictionary<string, string> UnrealPakResponseFile, bool bCompressed, bool RehydrateAssets, string EncryptionKeyGuid)
+			public CreatePakParams(string PakName, Dictionary<string, string> UnrealPakResponseFile, bool bCompressed, bool RehydrateAssets, string EncryptionKeyGuid, bool bOnDemand = false)
 			{
 				this.PakName = PakName;
 				this.UnrealPakResponseFile = UnrealPakResponseFile;
 				this.bCompressed = bCompressed;
 				this.bRehydrateAssets = RehydrateAssets;
 				this.EncryptionKeyGuid = EncryptionKeyGuid;
+				this.bOnDemand = bOnDemand;
 			}
 		}
 
@@ -3141,11 +3216,13 @@ namespace AutomationScripts
 				OutputFilename = OutputFilename + PostFix;
 
 				StagedFileReference OutputRelativeLocation;
+				StagedDirectoryReference OnDemandOutputRelativeLocation;
 				if (Params.HasDLCName)
 				{
 					if (Params.DLCOverrideStagedSubDir != null)
 					{
 						OutputRelativeLocation = StagedFileReference.Combine(SC.RelativeProjectRootForStage, Params.DLCOverrideStagedSubDir, "Content", "Paks", Params.DLCFile.GetFileNameWithoutExtension() + OutputFilename + ".pak");
+						OnDemandOutputRelativeLocation = StagedDirectoryReference.Combine(SC.RelativeProjectRootForStage, Params.DLCOverrideStagedSubDir, "ContentOnDemand");
 					}
 					else
 					{
@@ -3160,23 +3237,29 @@ namespace AutomationScripts
 							PluginSubdirectory = Params.DLCFile.GetFileNameWithoutAnyExtensions();
 						}
 						OutputRelativeLocation = StagedFileReference.Combine(SC.RelativeProjectRootForStage, PluginSubdirectory, "Content", "Paks", SC.FinalCookPlatform, Params.DLCFile.GetFileNameWithoutExtension() + OutputFilename + ".pak");
+						OnDemandOutputRelativeLocation = StagedDirectoryReference.Combine(SC.RelativeProjectRootForStage, PluginSubdirectory, "ContentOnDemand");
 					}
 				}
 				else
 				{
 					OutputRelativeLocation = StagedFileReference.Combine(SC.RelativeProjectRootForStage, "Content", "Paks", OutputFilename + OutputFilenameExtension);
+					OnDemandOutputRelativeLocation = StagedDirectoryReference.Combine(SC.RelativeProjectRootForStage, "ContentOnDemand");
 				}
 				if (SC.StageTargetPlatform.DeployLowerCaseFilenames(StagedFileType.UFS))
 				{
 					OutputRelativeLocation = OutputRelativeLocation.ToLowerInvariant();
 				}
 				OutputRelativeLocation = SC.StageTargetPlatform.Remap(OutputRelativeLocation);
+				
+				//TODO: Do we need to remap the ondemand output location?
+				//OnDemandOutputRelativeLocation = SC.StageTargetPlatform.Remap(OnDemandOutputRelativeLocation);
 
 				FileReference OutputLocation = FileReference.Combine(SC.RuntimeRootDir, OutputRelativeLocation.Name);
+				DirectoryReference OnDemandOutputLocation = DirectoryReference.Combine(SC.RuntimeRootDir, OnDemandOutputRelativeLocation.Name);
 
 				bool bCopiedExistingPak = false;
 
-				if (SC.StageTargetPlatform != SC.CookSourcePlatform && !Params.IgnorePaksFromDifferentCookSource)
+				if (!PakParams.bOnDemand && SC.StageTargetPlatform != SC.CookSourcePlatform && !Params.IgnorePaksFromDifferentCookSource)
 				{
 					// Check to see if we have an existing pak file we can use
 
@@ -3313,6 +3396,7 @@ namespace AutomationScripts
 								OutputLocation.GetFileNameWithoutAnyExtensions(),
 								OutputLocation,
 								SC.OptionalFileStageDirectory,
+								PakParams.bOnDemand ? OnDemandOutputLocation : null,
 								bCompressContainers,
 								CryptoSettings,
 								PakParams.EncryptionKeyGuid,
@@ -3843,6 +3927,7 @@ namespace AutomationScripts
 				ResponseFile = new ConcurrentDictionary<string, string>();
 				Manifest = null;
 				bCompressed = false;
+				bOnDemand = false;
 			}
 
 			// Name of pak file without extension, ie pakchunk0
@@ -3858,6 +3943,9 @@ namespace AutomationScripts
 			public bool bCompressed;
 			public string EncryptionKeyGuid;
 			public string RequestedEncryptionKeyGuid;
+
+			// Whether the chunk is used for content-on-demand content delivery or not
+			public bool bOnDemand;
 		}
 
 		/// <summary>
@@ -3958,6 +4046,54 @@ namespace AutomationScripts
 					}
 				}
 
+				// Add ondemand chunk definitions
+				if (PakRulesList != null)
+				{
+					foreach (PakFileRules Rule in PakRulesList.Where(R => R.bOnDemand))
+					{
+						string ChunkName = Rule.Name;
+						if (Rule.OverridePaks != null && Rule.OverridePaks.Count > 0)
+						{
+							ChunkName = Rule.OverridePaks[0];
+						}
+
+						if (ChunkDefinitions.Any(D => D.ChunkName == ChunkName))
+						{
+							continue;
+						}
+
+						Log.TraceInformation("Creating content-on-demand chunk definition '{0}' for PAK rule '{1}'", ChunkName, Rule.Name);
+
+						bool bIsEncryptionKeyValid = true;
+						if (string.IsNullOrEmpty(Rule.EncryptionKeyGuid))
+						{
+							bIsEncryptionKeyValid = false;
+							Log.TraceWarning("No encryption key specified for content-on-demand PAK rule '{0}'", Rule.Name);
+						}
+						else
+						{
+							Guid KeyGuid;
+							if (!Guid.TryParse(Rule.EncryptionKeyGuid, out KeyGuid))
+							{
+								bIsEncryptionKeyValid = false;
+								Log.TraceWarning("Invalid encryption key GUID specified for content-on-demand PAK rule '{0}'", Rule.Name);
+							}
+
+							if (!PakCryptoSettings.ContainsEncryptionKey(Rule.EncryptionKeyGuid))
+							{
+								bIsEncryptionKeyValid = false;
+								Log.TraceWarning("No matching encryption key found for PAK rule '{0}', GUID '{1}'", Rule.Name, Rule.EncryptionKeyGuid);
+							}
+						}
+
+						ChunkDefinition Chunk = new ChunkDefinition(ChunkName);
+						Chunk.bOnDemand = true;
+						Chunk.bCompressed = true;
+						Chunk.EncryptionKeyGuid = bIsEncryptionKeyValid ? Rule.EncryptionKeyGuid : null;
+						ChunkDefinitions.Add(Chunk);
+					}
+				}
+
 				ConcurrentDictionary<string, ChunkDefinition> ChunkNameToDefinition = new ConcurrentDictionary<string, ChunkDefinition>(StringComparer.InvariantCultureIgnoreCase);
 				foreach (ChunkDefinition Chunk in ChunkDefinitions)
 				{
@@ -4027,32 +4163,35 @@ namespace AutomationScripts
 					{
 						ChunkDefinition TargetChunk = Chunk;
 
-						string OrigExt = Path.GetExtension(OriginalFilename);
-						if (bStageLoose)
+						if (!TargetChunk.bOnDemand)
 						{
-							// make a streaming pak
-							string OnDemandChunkName = Chunk.ChunkName + "ondemand";
-							if (!OnDemandChunks.ContainsKey(OnDemandChunkName))
+							string OrigExt = Path.GetExtension(OriginalFilename);
+							if (bStageLoose)
 							{
-								ChunkDefinition OnDemandChunk = new ChunkDefinition(OnDemandChunkName);
-								OnDemandChunk.RequestedEncryptionKeyGuid = Chunk.RequestedEncryptionKeyGuid;
-								OnDemandChunk.EncryptionKeyGuid = Chunk.EncryptionKeyGuid;
-								OnDemandChunks.TryAdd(OnDemandChunkName, OnDemandChunk);
+								// make a streaming pak
+								string OnDemandChunkName = Chunk.ChunkName + "ondemand";
+								if (!OnDemandChunks.ContainsKey(OnDemandChunkName))
+								{
+									ChunkDefinition OnDemandChunk = new ChunkDefinition(OnDemandChunkName);
+									OnDemandChunk.RequestedEncryptionKeyGuid = Chunk.RequestedEncryptionKeyGuid;
+									OnDemandChunk.EncryptionKeyGuid = Chunk.EncryptionKeyGuid;
+									OnDemandChunks.TryAdd(OnDemandChunkName, OnDemandChunk);
+								}
+								TargetChunk = OnDemandChunks[OnDemandChunkName];
 							}
-							TargetChunk = OnDemandChunks[OnDemandChunkName];
-						}
-						else if (OrigExt.Equals(OptionalBulkDataFileExtension))
-						{
-							// any optional files encountered we want to put in a separate pak file
-							string OptionalChunkName = Chunk.ChunkName + "optional";
-							if (!OptionalChunks.ContainsKey(OptionalChunkName))
+							else if (OrigExt.Equals(OptionalBulkDataFileExtension))
 							{
-								ChunkDefinition OptionalChunk = new ChunkDefinition(OptionalChunkName);
-								OptionalChunk.RequestedEncryptionKeyGuid = Chunk.RequestedEncryptionKeyGuid;
-								OptionalChunk.EncryptionKeyGuid = Chunk.EncryptionKeyGuid;
-								OptionalChunks.TryAdd(OptionalChunkName, OptionalChunk);
+								// any optional files encountered we want to put in a separate pak file
+								string OptionalChunkName = Chunk.ChunkName + "optional";
+								if (!OptionalChunks.ContainsKey(OptionalChunkName))
+								{
+									ChunkDefinition OptionalChunk = new ChunkDefinition(OptionalChunkName);
+									OptionalChunk.RequestedEncryptionKeyGuid = Chunk.RequestedEncryptionKeyGuid;
+									OptionalChunk.EncryptionKeyGuid = Chunk.EncryptionKeyGuid;
+									OptionalChunks.TryAdd(OptionalChunkName, OptionalChunk);
+								}
+								TargetChunk = OptionalChunks[OptionalChunkName];
 							}
-							TargetChunk = OptionalChunks[OptionalChunkName];
 						}
 
 						TargetChunk.ResponseFile.TryAdd(StagingFile.Key, StagingFile.Value);
@@ -4170,7 +4309,7 @@ namespace AutomationScripts
 					string EncryptionKeyToUse = Params.SkipEncryption ? "" : Chunk.EncryptionKeyGuid;
 					PakInputs.Add(new CreatePakParams(Chunk.ChunkName,
 						Chunk.ResponseFile.ToDictionary(entry => entry.Key, entry => entry.Value),
-						Params.Compressed || Chunk.bCompressed, Params.RehydrateAssets, EncryptionKeyToUse));
+						Params.Compressed || Chunk.bCompressed, Params.RehydrateAssets, EncryptionKeyToUse, Chunk.bOnDemand));
 				}
 			}
 
diff --git a/Engine/Source/Programs/UnrealBuildTool/System/EncryptionAndSigning.cs b/Engine/Source/Programs/UnrealBuildTool/System/EncryptionAndSigning.cs
--- a/Engine/Source/Programs/UnrealBuildTool/System/EncryptionAndSigning.cs
+++ b/Engine/Source/Programs/UnrealBuildTool/System/EncryptionAndSigning.cs
@@ -188,6 +188,29 @@ namespace UnrealBuildTool
 				DirectoryReference.CreateDirectory(InFile.Directory);
 				FileReference.WriteAllTextIfDifferent(InFile, JsonSerializer.Serialize(this));
 			}
+
+			/// <summary>
+			/// Returns whether the specified encryption key exist.
+			/// </summary>
+			public bool ContainsEncryptionKey(string KeyGuid)
+			{
+				if (string.IsNullOrEmpty(KeyGuid))
+				{
+					return false;
+				}
+
+				if (EncryptionKey != null && EncryptionKey.Guid == KeyGuid)
+				{
+					return true;
+				}
+
+				if (SecondaryEncryptionKeys != null)
+				{
+					return SecondaryEncryptionKeys.Any(Key => Key.Guid == KeyGuid);
+				}
+
+				return false;
+			}
 		}
 
 		/// <summary>
diff --git a/Engine/Source/Runtime/Core/Internal/IO/IoStore.h b/Engine/Source/Runtime/Core/Internal/IO/IoStore.h
--- a/Engine/Source/Runtime/Core/Internal/IO/IoStore.h
+++ b/Engine/Source/Runtime/Core/Internal/IO/IoStore.h
@@ -15,6 +15,8 @@
 
 CORE_API DECLARE_LOG_CATEGORY_EXTERN(LogIoStore, Log, All);
 
+class FCbWriter;
+
 /**
  * I/O store container format version
  */
@@ -252,3 +254,57 @@ struct FIoStoreTocResource
 
 	CORE_API static uint64 HashChunkIdWithSeed(int32 Seed, const FIoChunkId& ChunkId);
 };
+
+enum class EIoOnDemandTocVersion : uint32
+{
+	Invalid			= 0,
+	Initial			= 1,
+
+	LatestPlusOne,
+	Latest			= (LatestPlusOne - 1)
+};
+
+struct FIoStoreOnDemandTocHeader
+{
+	static constexpr uint64 ExpectedMagic = 0x6f6e64656d616e64; // ondemand
+
+	uint64 Magic = ExpectedMagic;
+	uint32 Version = uint32(EIoOnDemandTocVersion::Latest);
+	uint32 BlockSize = 0;
+	FString CompressionFormat;
+	FString ChunksDirectory;
+	
+	CORE_API friend FCbWriter& operator<<(FCbWriter& Writer, const FIoStoreOnDemandTocHeader& Header);
+};
+
+struct FIoStoreOnDemandTocEntry
+{
+	FIoHash Hash = FIoHash::Zero;
+	FIoChunkId ChunkId = FIoChunkId::InvalidChunkId;
+	uint64 RawSize = 0;
+	uint64 EncodedSize = 0;
+	uint32 BlockOffset = ~uint32(0);
+	uint32 BlockCount = 0; 
+	
+	CORE_API friend FCbWriter& operator<<(FCbWriter& Writer, const FIoStoreOnDemandTocEntry& Entry);
+};
+
+struct FIoStoreOnDemandContainerEntry
+{
+	FString ContainerName;
+	FString EncryptionKeyGuid;
+	TArray<FIoStoreOnDemandTocEntry> Entries;
+	TArray<uint32> BlockSizes;
+
+	CORE_API friend FCbWriter& operator<<(FCbWriter& Writer, const FIoStoreOnDemandContainerEntry& ContainerEntry);
+};
+
+struct FIoStoreOndemandTocResource
+{
+	FIoStoreOnDemandTocHeader Header;
+	TArray<FIoStoreOnDemandContainerEntry> Containers;
+	
+	CORE_API friend FCbWriter& operator<<(FCbWriter& Writer, const FIoStoreOndemandTocResource& TocResource);
+
+	UE_NODISCARD CORE_API static TIoStatusOr<FString> Save(const TCHAR* Directory, const FIoStoreOndemandTocResource& TocResource);
+};
diff --git a/Engine/Source/Runtime/Core/Private/IO/IoStore.cpp b/Engine/Source/Runtime/Core/Private/IO/IoStore.cpp
--- a/Engine/Source/Runtime/Core/Private/IO/IoStore.cpp
+++ b/Engine/Source/Runtime/Core/Private/IO/IoStore.cpp
@@ -2,6 +2,7 @@
 
 #include "IO/IoStore.h"
 #include "Containers/Map.h"
+#include "Containers/StringView.h"
 #include "HAL/FileManager.h"
 #include "Templates/UniquePtr.h"
 #include "Misc/Paths.h"
@@ -21,7 +22,10 @@
 #include "Misc/ConfigCacheIni.h"
 #include "DerivedDataPluginInterface.h"
 #include "DerivedDataCacheInterface.h"
+#include "Serialization/CompactBinaryWriter.h"
+#include "Serialization/CompactBinarySerialization.h"
 #include "Serialization/LargeMemoryReader.h"
+#include "Serialization/LargeMemoryWriter.h"
 #include "Serialization/MemoryWriter.h"
 #include "ProfilingDebugging/CountersTrace.h"
 #include "Tasks/Task.h"
@@ -3239,3 +3243,107 @@ void FIoStoreReader::GetFilenamesByBlockIndex(const TArray<int32>& InBlockIndexL
 			return true;
 		});
 }
+
+//////////////////////////////////////////////////////////////////////////
+
+FCbWriter& operator<<(FCbWriter& Writer, const FIoStoreOnDemandTocHeader& Header)
+{
+	Writer.BeginObject();
+	Writer.AddInteger(UTF8TEXTVIEW("Magic"), Header.Magic);
+	Writer.AddInteger(UTF8TEXTVIEW("Version"), Header.Version);
+	Writer.AddInteger(UTF8TEXTVIEW("BlockSize"), Header.BlockSize);
+	Writer.AddString(UTF8TEXTVIEW("CompressionFormat"), Header.CompressionFormat);
+	Writer.AddString(UTF8TEXTVIEW("ChunksDirectory"), Header.ChunksDirectory);
+	Writer.EndObject();
+
+	return Writer;
+}
+
+FCbWriter& operator<<(FCbWriter& Writer, const FIoStoreOnDemandTocEntry& Entry)
+{
+	Writer.BeginObject();
+	Writer.AddHash(UTF8TEXTVIEW("Hash"), Entry.Hash);
+	Writer << UTF8TEXTVIEW("ChunkId") << Entry.ChunkId;
+	Writer.AddInteger(UTF8TEXTVIEW("RawSize"), Entry.RawSize);
+	Writer.AddInteger(UTF8TEXTVIEW("EncodedSize"), Entry.EncodedSize);
+	Writer.AddInteger(UTF8TEXTVIEW("BlockOffset"), Entry.BlockOffset);
+	Writer.AddInteger(UTF8TEXTVIEW("BlockCount"), Entry.BlockCount);
+	Writer.EndObject();
+
+	return Writer;
+}
+
+FCbWriter& operator<<(FCbWriter& Writer, const FIoStoreOnDemandContainerEntry& ContainerEntry)
+{
+	Writer.BeginObject();
+	Writer.AddString(UTF8TEXTVIEW("Name"), ContainerEntry.ContainerName);
+	Writer.AddString(UTF8TEXTVIEW("EncryptionKeyGuid"), ContainerEntry.EncryptionKeyGuid);
+
+	Writer.BeginArray(UTF8TEXTVIEW("Entries"));
+	for (const FIoStoreOnDemandTocEntry& Entry : ContainerEntry.Entries)
+	{
+		Writer << Entry;
+	}
+	Writer.EndArray();
+	
+	Writer.BeginArray(UTF8TEXTVIEW("BlockSizes"));
+	for (uint32 BlockSize : ContainerEntry.BlockSizes)
+	{
+		Writer << BlockSize;
+	}
+	Writer.EndArray();
+
+	Writer.EndObject();
+
+	return Writer;
+}
+
+FCbWriter& operator<<(FCbWriter& Writer, const FIoStoreOndemandTocResource& TocResource)
+{
+	Writer.BeginObject();
+	Writer << UTF8TEXTVIEW("Header") << TocResource.Header;
+
+	Writer.BeginArray(UTF8TEXTVIEW("Containers"));
+	for (const FIoStoreOnDemandContainerEntry& Container : TocResource.Containers)
+	{
+		Writer << Container;
+	}
+	Writer.EndArray();
+	Writer.EndObject();
+	
+	return Writer;
+}
+
+TIoStatusOr<FString> FIoStoreOndemandTocResource::Save(const TCHAR* Directory, const FIoStoreOndemandTocResource& TocResource)
+{
+	if (TocResource.Header.Magic != FIoStoreOnDemandTocHeader::ExpectedMagic)
+	{
+		return FIoStatus(EIoErrorCode::CorruptToc);
+	}
+
+	if (TocResource.Header.CompressionFormat.IsEmpty())
+	{
+		return FIoStatus(EIoErrorCode::CorruptToc);
+	}
+
+	FCbWriter Writer;
+	Writer << TocResource;
+
+	FLargeMemoryWriter Ar;
+	SaveCompactBinary(Ar, Writer.Save());
+
+	const FIoHash TocHash = FIoHash::HashBuffer(Ar.GetView());
+	const FString FilePath = FString(Directory) / LexToString(TocHash) + TEXT(".iochunktoc");
+
+	IFileManager& FileMgr = IFileManager::Get();
+	if (TUniquePtr<FArchive> FileAr(FileMgr.CreateFileWriter(*FilePath)); FileAr.IsValid())
+	{
+		FileAr->Serialize(Ar.GetData(), Ar.TotalSize());
+		FileAr->Flush();
+		FileAr->Close();
+
+		return FilePath;
+	}
+
+	return FIoStatus(EIoErrorCode::WriteError);
+}
diff --git a/Engine/Source/Runtime/Core/Public/IO/IoDispatcher.h b/Engine/Source/Runtime/Core/Public/IO/IoDispatcher.h
--- a/Engine/Source/Runtime/Core/Public/IO/IoDispatcher.h
+++ b/Engine/Source/Runtime/Core/Public/IO/IoDispatcher.h
@@ -978,6 +978,7 @@ enum class EIoContainerFlags : uint8
 	Encrypted	= (1 << 1),
 	Signed		= (1 << 2),
 	Indexed		= (1 << 3),
+	OnDemand	= (1 << 4),
 };
 ENUM_CLASS_FLAGS(EIoContainerFlags);
 
@@ -1009,6 +1010,11 @@ struct FIoContainerSettings
 	{
 		return !!(ContainerFlags & EIoContainerFlags::Indexed);
 	}
+	
+	bool IsOnDemand() const
+	{
+		return !!(ContainerFlags & EIoContainerFlags::OnDemand);
+	}
 };
 
 struct FIoStoreWriterResult
