Allow asset registry writeback to fail until we sort out asset registry staging issue to fix build break.

diff --git a/Engine/Source/Developer/IoStoreUtilities/Private/IoStoreUtilities.cpp b/Engine/Source/Developer/IoStoreUtilities/Private/IoStoreUtilities.cpp
--- a/Engine/Source/Developer/IoStoreUtilities/Private/IoStoreUtilities.cpp
+++ b/Engine/Source/Developer/IoStoreUtilities/Private/IoStoreUtilities.cpp
@@ -3348,7 +3348,7 @@ bool DoAssetRegistryWritebackDuringStage(EAssetRegistryWritebackMethod InMethod,
 	// The overwhelming majority of time for the asset registry writeback is loading and saving.
 	FString AssetRegistryFileName;
 	FAssetRegistryState AssetRegistry;
-	if (FindAndLoadDevelopmentAssetRegistry(InCookedDir, true, AssetRegistry, &AssetRegistryFileName) == false)
+	if (FindAndLoadDevelopmentAssetRegistry(InCookedDir, false, AssetRegistry, &AssetRegistryFileName) == false)
 	{
 		// already logged
 		return false;
