Remove spurious asset registry writeback warning.

diff --git a/Engine/Source/Developer/IoStoreUtilities/Private/IoStoreUtilities.cpp b/Engine/Source/Developer/IoStoreUtilities/Private/IoStoreUtilities.cpp
--- a/Engine/Source/Developer/IoStoreUtilities/Private/IoStoreUtilities.cpp
+++ b/Engine/Source/Developer/IoStoreUtilities/Private/IoStoreUtilities.cpp
@@ -3142,7 +3142,7 @@ static void AddChunkInfoToAssetRegistry(TMap<FPackageId, TArray<FIoStoreTocChunk
 		const TArray<FIoStoreTocChunkInfo, TInlineAllocator<2>>* PackageChunks = PackageToChunks.Find(FPackageId::FromName(AssetPackage.Key));
 		if (PackageChunks == nullptr)
 		{
-			UE_LOG(LogIoStore, Warning, TEXT("Package data for %s had no chunks %d -- %lld "), *AssetPackage.Key.ToString(), AssetPackage.Value->ChunkHashes.Num(), AssetPackage.Value->DiskSize);
+			// This happens when the package has been stripped by UAT prior to staging by e.g. PakDenyList.
 			continue;
 		}
 
