Use ini for cross-container refs config instead of json to allow comments.
Handle packages with multiple locations.

diff --git a/Engine/Source/Developer/IoStoreUtilities/Private/IoStoreUtilities.cpp b/Engine/Source/Developer/IoStoreUtilities/Private/IoStoreUtilities.cpp
--- a/Engine/Source/Developer/IoStoreUtilities/Private/IoStoreUtilities.cpp
+++ b/Engine/Source/Developer/IoStoreUtilities/Private/IoStoreUtilities.cpp
@@ -5794,59 +5794,32 @@ int32 ValidateCrossContainerRefs(
 {
 	TMultiMap<FString, FString> ValidEdges;
 	TArray<FString> IgnoreRefsFromAssets, IgnoreRefsToAssets;
-	if (FString ConfigString; FFileHelper::LoadFileToString(ConfigString, *ConfigPath))
+
+	FConfigFile ConfigFile;
+	if (!FConfigCacheIni::LoadLocalIniFile(ConfigFile, *ConfigPath, false))
 	{
-		TSharedPtr<FJsonObject> RootConfigObject;
-		TSharedRef<TJsonReader<>> Reader = TJsonReaderFactory<>::Create(ConfigString);
-		FJsonSerializer::Deserialize(Reader, RootConfigObject);
-		if (!RootConfigObject.IsValid())
-		{
-			UE_LOG(LogIoStore, Error, TEXT("Failed to read configuration from file %s, check if it is valid json"), *ConfigPath);
-			return -1;
-		}
-		if (TSharedPtr<FJsonObject> ValidEdgesObj = RootConfigObject->GetObjectField(TEXT("ValidEdges")); ValidEdgesObj.IsValid())
-		{
-			for (const TPair<FString, TSharedPtr<FJsonValue>>& EdgePair : ValidEdgesObj->Values) 
-			{
-				const FString& From = EdgePair.Key;
-				FString ToString;
-				if (EdgePair.Value->TryGetString(ToString))
-				{
-					ValidEdges.Add(From, ToString);
-				}
-				else if (TArray<TSharedPtr<FJsonValue>>* ToArray; EdgePair.Value->TryGetArray(ToArray))
-				{
-					for (const TSharedPtr<FJsonValue>& ToValue : *ToArray)
-					{
-						if (ToValue->TryGetString(ToString))
-						{
-							ValidEdges.Add(From, ToString);
-						}
-					}
-				}
-			}
-		}
-		const TArray<TSharedPtr<FJsonValue>> IgnoreRefsFromArray = RootConfigObject->GetArrayField(TEXT("IgnoreRefsFrom"));
-		for (const TSharedPtr<FJsonValue>& Entry : IgnoreRefsFromArray)
+		UE_LOG(LogIoStore, Error, TEXT("Failed to load config file %s"), *ConfigPath);
+		return -1;
+	}
+
+	if (FConfigSection* EdgesSection = ConfigFile.Find(TEXT("Edges")))
+	{
+		for (auto It = EdgesSection->CreateIterator(); It; ++It)
 		{
-			if (FString Value; Entry->TryGetString(Value))
-			{
-				IgnoreRefsFromAssets.Add(MoveTemp(Value));
-			}
+			ValidEdges.Add(It.Key().ToString(), It.Value().GetValue());
 		}
-		const TArray<TSharedPtr<FJsonValue>> IgnoreRefsToArray = RootConfigObject->GetArrayField(TEXT("IgnoreRefsTo"));
-		for (const TSharedPtr<FJsonValue>& Entry : IgnoreRefsToArray)
+	}
+	if (FConfigSection* DefaultEdgesSection = ConfigFile.Find(TEXT("DefaultEdges")))
+	{
+		for (auto It = DefaultEdgesSection->CreateIterator(); It; ++It)
 		{
-			if (FString Value; Entry->TryGetString(Value))
-			{
-				IgnoreRefsToAssets.Add(MoveTemp(Value));
-			}
+			ValidEdges.Add(FString(), It.Key().ToString());
 		}
 	}
-	else
+	if (FConfigSection* IgnoreSection = ConfigFile.Find(TEXT("Ignore")))
 	{
-		UE_LOG(LogIoStore, Error, TEXT("Failed to load config file %s"), *ConfigPath);
-		return -1;
+		IgnoreSection->MultiFind(TEXT("IgnoreRefsFrom"), IgnoreRefsFromAssets);
+		IgnoreSection->MultiFind(TEXT("IgnoreRefsTo"), IgnoreRefsToAssets);
 	}
 
 	if (ValidEdges.Num() == 0)
@@ -5961,16 +5934,10 @@ int32 ValidateCrossContainerRefs(
 		}
 	}
 
-	TMap<TTuple<const FContainerDesc*, const FContainerDesc*>, TSet<TTuple<const FPackageDesc*, const FPackageDesc*>>> Errors;
+	TMap<const FContainerDesc*, TSet<TTuple<const FPackageDesc*, const FPackageDesc*>>> Errors;
 	Algo::SortBy(Packages, [](FPackageDesc* Desc) { return Desc->PackageName; }, FNameLexicalLess());
 	for (const FPackageDesc* Package : Packages)
 	{
-		if (Package->Locations.Num() > 1)
-		{
-			UE_LOG(LogIoStore, Error, TEXT("ValidateCrossContainerRefs does not support assets with multiple locations. (%s)"), *Package->PackageName.ToString());
-			continue;
-		}
-		
 		bool bSkip = Algo::AnyOf(IgnoreRefsFromAssets, [Package](const FString& IgnoreString)
 		{
 			FString PackageNameString = Package->PackageName.ToString();
@@ -5989,7 +5956,6 @@ int32 ValidateCrossContainerRefs(
 			continue;
 		}
 
-		const FContainerDesc* Container = Package->Locations[0].Container;
 		bool bNeedsHeader = true;
 		for (const FImportDesc& Import : Package->Imports)
 		{
@@ -6000,11 +5966,6 @@ int32 ValidateCrossContainerRefs(
 					LogIoStore, Error, TEXT("Unresolved import of package %s by package %s"), *Import.Name.ToString(), *Package->PackageName.ToString());
 				continue;
 			}
-			if (ImportPackage->Locations.Num() > 1)
-			{
-				UE_LOG(LogIoStore, Error, TEXT("ValidateCrossContainerRefs does not support assets with multiple locations. (%s)"), *Package->PackageName.ToString());
-				continue;
-			}
 
 			bSkip = Algo::AnyOf(IgnoreRefsFromAssets, [ImportPackage](const FString& IgnoreString)
 			{
@@ -6024,11 +5985,16 @@ int32 ValidateCrossContainerRefs(
 				continue;
 			}
 
-			const FContainerDesc* ImportContainer = ImportPackage->Locations[0].Container;
-			if (!FinalValidEdges.FindPair(Container, ImportContainer))
+			// For each location of the importing paacka
+			for (const DescribeUtils::FPackageLocation& Location : Package->Locations)
 			{
-				Errors.FindOrAdd(TTuple<const FContainerDesc*, const FContainerDesc*>(Container, ImportContainer))
-					.Add({Package, ImportPackage});
+				bool bValid = Algo::AnyOf(ImportPackage->Locations, [&](const DescribeUtils::FPackageLocation& ImportLocation) {
+					return FinalValidEdges.FindPair(Location.Container, ImportLocation.Container) != nullptr;
+				});
+				if (!bValid)
+				{
+					Errors.FindOrAdd(Location.Container).Add({ Package, ImportPackage });
+				}
 			}
 		}
 	}
@@ -6055,12 +6021,13 @@ int32 ValidateCrossContainerRefs(
 		return 0;
 	}
 
-	for (const TPair<TTuple<const FContainerDesc*, const FContainerDesc*>, TSet<TTuple<const FPackageDesc*, const FPackageDesc*>>>& Pair : Errors)
+	for (const TPair<const FContainerDesc*, TSet<TTuple<const FPackageDesc*, const FPackageDesc*>>>& Pair : Errors)
 	{
-		OutputOverride->Logf(ELogVerbosity::Display, TEXT("%s -> %s"), *Pair.Key.Key->Name.ToString(), *Pair.Key.Value->Name.ToString());
+		OutputOverride->Logf(ELogVerbosity::Display, TEXT("%s"), *Pair.Key->Name.ToString());
 		for (const TTuple<const FPackageDesc*, const FPackageDesc*>& Error : Pair.Value) 
 		{
-			OutputOverride->Logf(ELogVerbosity::Display, TEXT("\t%s -> %s"), *Error.Key->PackageName.ToString(), *Error.Value->PackageName.ToString());
+			FString LocationsString = FString::JoinBy(Error.Value->Locations, TEXT(","), [](const DescribeUtils::FPackageLocation& Location) { return Location.Container->Name.ToString(); });
+			OutputOverride->Logf(ELogVerbosity::Display, TEXT("\t%s -> %s (%s)"), *Error.Key->PackageName.ToString(), *Error.Value->PackageName.ToString(), *LocationsString);
 		} 
 	}
 	return 0;
