I/O - fix non unity compile error in HTTP client

diff --git a/Engine/Source/Developer/IoStoreUtilities/IoStoreUtilities.Build.cs b/Engine/Source/Developer/IoStoreUtilities/IoStoreUtilities.Build.cs
--- a/Engine/Source/Developer/IoStoreUtilities/IoStoreUtilities.Build.cs
+++ b/Engine/Source/Developer/IoStoreUtilities/IoStoreUtilities.Build.cs
@@ -28,5 +28,6 @@ public class IoStoreUtilities : ModuleRules
         PrivateDependencyModuleNames.Add("RSA");
         PrivateDependencyModuleNames.Add("DeveloperToolSettings");
         PrivateDependencyModuleNames.Add("SandboxFile");
+        PrivateDependencyModuleNames.Add("IoStoreOnDemand");
 	}
 }
diff --git a/Engine/Source/Developer/IoStoreUtilities/Private/IoStoreOnDemand.cpp b/Engine/Source/Developer/IoStoreUtilities/Private/IoStoreOnDemand.cpp
deleted file mode 100644
diff --git a/Engine/Source/Developer/IoStoreUtilities/Private/IoStoreOnDemand.h b/Engine/Source/Developer/IoStoreUtilities/Private/IoStoreOnDemand.h
deleted file mode 100644
diff --git a/Engine/Source/Developer/IoStoreUtilities/Private/IoStoreUtilities.cpp b/Engine/Source/Developer/IoStoreUtilities/Private/IoStoreUtilities.cpp
--- a/Engine/Source/Developer/IoStoreUtilities/Private/IoStoreUtilities.cpp
+++ b/Engine/Source/Developer/IoStoreUtilities/Private/IoStoreUtilities.cpp
@@ -2,7 +2,6 @@
 
 #include "IoStoreUtilities.h"
 
-#include "IoStoreOnDemand.h"
 #include "Async/AsyncWork.h"
 #include "HAL/FileManager.h"
 #include "HAL/PlatformFileManager.h"
@@ -32,6 +31,7 @@
 #include "Serialization/ArrayReader.h"
 #include "Serialization/ArrayWriter.h"
 #include "Settings/ProjectPackagingSettings.h" // for EAssetRegistryWritebackMethod
+#include "IO/IoStoreOnDemand.h"
 #include "IO/PackageStore.h"
 #include "UObject/Class.h"
 #include "UObject/NameBatchSerialization.h"
@@ -3720,7 +3720,7 @@ int32 CreateTarget(const FIoStoreArguments& Arguments, const FIoStoreWriterSetti
 	}
 	TArray<TSharedPtr<IIoStoreWriter>> IoStoreWriters;
 	TSharedPtr<IIoStoreWriter> GlobalIoStoreWriter;
-	TUniquePtr<IIoStoreOnDemandWriter> IoStoreOnDemandWriter;
+	TUniquePtr<UE::IOnDemandIoStoreWriter> OnDemandIoStoreWriter;
 	{
 		IOSTORE_CPU_SCOPE(InitializeWriters);
 		if (!Arguments.IsDLC())
@@ -3747,11 +3747,11 @@ int32 CreateTarget(const FIoStoreArguments& Arguments, const FIoStoreWriterSetti
 
 			if (!ContainerTarget->OnDemandOutputPath.IsEmpty())
 			{
-				if (!IoStoreOnDemandWriter.IsValid())
+				if (!OnDemandIoStoreWriter.IsValid())
 				{
 					FIoStoreWriterSettings WriterSettings;
 					WriterSettings.CompressionMethod = NAME_Oodle;
-					IoStoreOnDemandWriter = MakeIoStoreOnDemandWriter(WriterSettings, ContainerTarget->OnDemandOutputPath);
+					OnDemandIoStoreWriter = UE::MakeOnDemandIoStoreWriter(WriterSettings, ContainerTarget->OnDemandOutputPath);
 				}
 
 				FIoContainerSettings ContainerSettings;
@@ -3770,7 +3770,7 @@ int32 CreateTarget(const FIoStoreArguments& Arguments, const FIoStoreWriterSetti
 						UE_LOG(LogIoStore, Warning, TEXT("Failed to find encryption key '%s' for container '%s'"), *ContainerTarget->EncryptionKeyGuid.ToString(), *ContainerTarget->Name.ToString());
 					}
 				}
-				ContainerTarget->IoStoreWriter = IoStoreOnDemandWriter->CreateContainer(ContainerTarget->Name.ToString(), ContainerSettings);
+				ContainerTarget->IoStoreWriter = OnDemandIoStoreWriter->CreateContainer(ContainerTarget->Name.ToString(), ContainerSettings);
 				IoStoreWriters.Add(ContainerTarget->IoStoreWriter);
 			}
 			else
@@ -3961,11 +3961,11 @@ int32 CreateTarget(const FIoStoreArguments& Arguments, const FIoStoreWriterSetti
 		GlobalIoStoreWriter->Append(CreateIoChunkId(0, 0, EIoChunkType::ScriptObjects), ScriptObjectsBuffer, WriteOptions);
 	}
 
-	if (IoStoreOnDemandWriter.IsValid())
+	if (OnDemandIoStoreWriter.IsValid())
 	{
 		UE_LOG(LogIoStore, Display, TEXT("Serializing ondemand content..."));
 		IOSTORE_CPU_SCOPE(SerializingContentOnDemand);
-		IoStoreOnDemandWriter->Flush();
+		OnDemandIoStoreWriter->Flush();
 	}
 
 	UE_LOG(LogIoStore, Display, TEXT("Serializing container(s)..."));
diff --git a/Engine/Source/Runtime/Experimental/IoStoreOnDemand/Private/CoreHttp/Http.cpp b/Engine/Source/Runtime/Experimental/IoStoreOnDemand/Private/CoreHttp/Http.cpp
--- a/Engine/Source/Runtime/Experimental/IoStoreOnDemand/Private/CoreHttp/Http.cpp
+++ b/Engine/Source/Runtime/Experimental/IoStoreOnDemand/Private/CoreHttp/Http.cpp
@@ -2,6 +2,16 @@
 
 #include "CoreHttp/Client.h"
 
+#if !defined(NO_UE_INCLUDES)
+#include "IO/IoBuffer.h"
+#include "Math/UnrealMathUtility.h"
+#include "Misc/ScopeExit.h"
+#include "Misc/ScopeLock.h"
+#include "Misc/StringBuilder.h"
+#include "ProfilingDebugging/CpuProfilerTrace.h"
+#include "Tasks/Task.h"
+#endif
+
 // {{{1 platforms ..............................................................
 
 #if PLATFORM_MICROSOFT
@@ -174,16 +184,6 @@
 #	undef PF_MAX
 #endif
 
-#if !defined(NO_UE_INCLUDES)
-#include "IO/IoDispatcher.h"
-#include "Math/UnrealMathUtility.h"
-#include "Misc/ScopeExit.h"
-#include "Misc/ScopeLock.h"
-#include "Misc/StringBuilder.h"
-#include "ProfilingDebugging/CpuProfilerTrace.h"
-#include "Tasks/Task.h"
-#endif
-
 static_assert(sizeof(sockaddr_in::sin_addr) == sizeof(uint32));
 
 #if PLATFORM_MICROSOFT
@@ -2482,4 +2482,4 @@ FEventLoop::~FEventLoop() {}
 
 #endif // CORE_HTTP_SUPPORTED
 
-/* vim: set noet : */
\ No newline at end of file
+/* vim: set noet : */
