exposed IoStore extract loop for in memory processing and added optional task throttle
added build scripts for content validations

diff --git a/Engine/Plugins/Experimental/PlatformCrypto/PlatformCrypto.uplugin b/Engine/Plugins/Experimental/PlatformCrypto/PlatformCrypto.uplugin
--- a/Engine/Plugins/Experimental/PlatformCrypto/PlatformCrypto.uplugin
+++ b/Engine/Plugins/Experimental/PlatformCrypto/PlatformCrypto.uplugin
@@ -14,7 +14,7 @@
 	"CanContainContent": false,
 	"IsBetaVersion": false,
 	"Installed": false,
-	"SupportedPrograms": [ "UnrealPak" ],
+	"SupportedPrograms": [ "UnrealPak", "UnrealPakEx" ],
 	"Modules": [
 		{
 			"Name": "PlatformCrypto",
@@ -31,7 +31,7 @@
 			"PlatformDenyList": [
 				"TVOS"
 			],
-			"ProgramAllowList": [ "UnrealPak" ]
+			"ProgramAllowList": [ "UnrealPak", "UnrealPakEx" ]
 		},
 		{
 			"Name": "PlatformCryptoTypes",
diff --git a/Engine/Source/Developer/IoStoreUtilities/Private/IoStoreUtilities.cpp b/Engine/Source/Developer/IoStoreUtilities/Private/IoStoreUtilities.cpp
--- a/Engine/Source/Developer/IoStoreUtilities/Private/IoStoreUtilities.cpp
+++ b/Engine/Source/Developer/IoStoreUtilities/Private/IoStoreUtilities.cpp
@@ -73,8 +73,6 @@
 #include "HAL/FileManagerGeneric.h"
 #include "Serialization/CompactBinarySerialization.h"
 
-//PRAGMA_DISABLE_OPTIMIZATION
-
 IMPLEMENT_MODULE(FDefaultModuleImpl, IoStoreUtilities);
 
 #define IOSTORE_CPU_SCOPE(NAME) TRACE_CPUPROFILER_EVENT_SCOPE(IoStore##NAME);
@@ -6748,6 +6746,23 @@ int32 GenerateZenFileSystemManifest(ITargetPlatform* TargetPlatform)
 	return 0;
 }
 
+bool ExtractFilesWriter(const FString& SrcFileName, const FString& DestFileName, const FIoChunkId& ChunkId, const uint8* Data, uint64 DataSize)
+{
+	TRACE_CPUPROFILER_EVENT_SCOPE(WriteFile);
+	TUniquePtr<FArchive> FileHandle(IFileManager::Get().CreateFileWriter(*DestFileName));
+	if (FileHandle)
+	{
+		FileHandle->Serialize(const_cast<uint8*>(Data), DataSize);
+		UE_CLOG(FileHandle->IsError(), LogIoStore, Error, TEXT("Failed writing to file \"%s\"."), *DestFileName);
+		return !FileHandle->IsError();
+	}
+	else
+	{
+		UE_LOG(LogIoStore, Error, TEXT("Unable to create file \"%s\"."), *DestFileName);
+		return false;
+	}
+};
+
 bool ExtractFilesFromIoStoreContainer(
 	const TCHAR* InContainerFilename,
 	const TCHAR* InDestPath,
@@ -6756,6 +6771,20 @@ bool ExtractFilesFromIoStoreContainer(
 	TMap<FString, uint64>* OutOrderMap,
 	TArray<FGuid>* OutUsedEncryptionKeys,
 	bool* bOutIsSigned)
+{
+	return ProcessFilesFromIoStoreContainer(InContainerFilename, InDestPath, InKeyChain, InFilter, ExtractFilesWriter, OutOrderMap, OutUsedEncryptionKeys, bOutIsSigned, -1);
+}
+
+bool ProcessFilesFromIoStoreContainer(
+	const TCHAR* InContainerFilename,
+	const TCHAR* InDestPath,
+	const FKeyChain& InKeyChain,
+	const FString* InFilter,
+	TFunction<bool(const FString&, const FString&, const FIoChunkId&, const uint8*, uint64)> FileProcessFunc,
+	TMap<FString, uint64>* OutOrderMap,
+	TArray<FGuid>* OutUsedEncryptionKeys,
+	bool* bOutIsSigned,
+	int32 MaxConcurrentReaders)
 {
 	TRACE_CPUPROFILER_EVENT_SCOPE(ExtractFilesFromIoStoreContainer);
 
@@ -6822,86 +6851,77 @@ bool ExtractFilesFromIoStoreContainer(
 
 	
 	const bool bContainerIsEncrypted = EnumHasAnyFlags(Reader->GetContainerFlags(), EIoContainerFlags::Encrypted);
+	const int32 MaxConcurrentTasks = (MaxConcurrentReaders <= 0) ? Entries.Num() : FMath::Min(MaxConcurrentReaders, Entries.Num());
+	int32 ErrorCount = 0;
 
-	auto WriteFile = [](const FString& FileName, const uint8* Data, uint64 DataSize)
+	for (int32 EntryStartIdx = 0; EntryStartIdx < Entries.Num(); )
 	{
-		TRACE_CPUPROFILER_EVENT_SCOPE(WriteFile);
-		TUniquePtr<FArchive> FileHandle(IFileManager::Get().CreateFileWriter(*FileName));
-		if (FileHandle)
-		{
-			FileHandle->Serialize(const_cast<uint8*>(Data), DataSize);
-			UE_CLOG(FileHandle->IsError(), LogIoStore, Error, TEXT("Failed writing to file \"%s\"."), *FileName);
-			return !FileHandle->IsError();
-		}
-		else
-		{
-			UE_LOG(LogIoStore, Error, TEXT("Unable to create file \"%s\"."), *FileName);
-			return false;
-		}
-	};
+		TArray<UE::Tasks::TTask<bool>> ExtractTasks;
+		ExtractTasks.Reserve(MaxConcurrentTasks);
+		EAsyncExecution ThreadPool = EAsyncExecution::ThreadPool;
 
-	TArray<UE::Tasks::TTask<bool>> ExtractTasks;
-	ExtractTasks.Reserve(Entries.Num());
-	EAsyncExecution ThreadPool = EAsyncExecution::ThreadPool;
-	for (int32 EntryIndex = 0; EntryIndex < Entries.Num(); ++EntryIndex)
-	{
-		const FEntry& Entry = Entries[EntryIndex];
+		const int32 NumTasks = FMath::Min(MaxConcurrentTasks, Entries.Num() - EntryStartIdx);
+		for (int32 TaskIndex = 0; TaskIndex < NumTasks; ++TaskIndex)
+		{
+			const FEntry& Entry = Entries[EntryStartIdx + TaskIndex];
 
-		UE::Tasks::TTask<TIoStatusOr<FIoBuffer>> ReadTask = Reader->ReadAsync(Entry.ChunkId, FIoReadOptions());
+			UE::Tasks::TTask<TIoStatusOr<FIoBuffer>> ReadTask = Reader->ReadAsync(Entry.ChunkId, FIoReadOptions());
 
-		// Once the read is done, write out the file.
-		ExtractTasks.Emplace(UE::Tasks::Launch(TEXT("IoStore_Extract"), 
-			[&Entry, &WriteFile, ReadTask]() mutable
-			{
-				TIoStatusOr<FIoBuffer> ReadChunkResult = ReadTask.GetResult();
-				if (!ReadChunkResult.IsOk())
+			// Once the read is done, write out the file.
+			ExtractTasks.Emplace(UE::Tasks::Launch(TEXT("IoStore_Extract"),
+				[&Entry, &FileProcessFunc, ReadTask]() mutable
 				{
-					UE_LOG(LogIoStore, Error, TEXT("Failed reading chunk for file \"%s\" (%s)."), *Entry.SourceFileName, *ReadChunkResult.Status().ToString());
-					return false;
-				}
-				
-				const uint8* Data = ReadChunkResult.ValueOrDie().Data();
-				uint64 DataSize = ReadChunkResult.ValueOrDie().DataSize();
-				if (Entry.ChunkId.GetChunkType() == EIoChunkType::ExportBundleData)
-				{
-					const FZenPackageSummary* PackageSummary = reinterpret_cast<const FZenPackageSummary*>(Data);
-					uint64 HeaderDataSize = PackageSummary->HeaderSize;
-					check(HeaderDataSize <= DataSize);
-					FString DestFileName = FPaths::ChangeExtension(Entry.DestFileName, TEXT(".uheader"));
-					if (!WriteFile(DestFileName, Data, HeaderDataSize))
+					TIoStatusOr<FIoBuffer> ReadChunkResult = ReadTask.GetResult();
+					if (!ReadChunkResult.IsOk())
 					{
+						UE_LOG(LogIoStore, Error, TEXT("Failed reading chunk for file \"%s\" (%s)."), *Entry.SourceFileName, *ReadChunkResult.Status().ToString());
 						return false;
 					}
-					DestFileName = FPaths::ChangeExtension(Entry.DestFileName, TEXT(".uexp"));
-					if (!WriteFile(DestFileName, Data + HeaderDataSize, DataSize - HeaderDataSize))
+
+					const uint8* Data = ReadChunkResult.ValueOrDie().Data();
+					uint64 DataSize = ReadChunkResult.ValueOrDie().DataSize();
+					if (Entry.ChunkId.GetChunkType() == EIoChunkType::ExportBundleData)
+					{
+						const FZenPackageSummary* PackageSummary = reinterpret_cast<const FZenPackageSummary*>(Data);
+						uint64 HeaderDataSize = PackageSummary->HeaderSize;
+						check(HeaderDataSize <= DataSize);
+						FString DestFileName = FPaths::ChangeExtension(Entry.DestFileName, TEXT(".uheader"));
+						if (!FileProcessFunc(Entry.SourceFileName, DestFileName, Entry.ChunkId, Data, HeaderDataSize))
+						{
+							return false;
+						}
+						DestFileName = FPaths::ChangeExtension(Entry.DestFileName, TEXT(".uexp"));
+						if (!FileProcessFunc(Entry.SourceFileName, DestFileName, Entry.ChunkId, Data + HeaderDataSize, DataSize - HeaderDataSize))
+						{
+							return false;
+						}
+					}
+					else if (!FileProcessFunc(Entry.SourceFileName, Entry.DestFileName, Entry.ChunkId, Data, DataSize))
 					{
 						return false;
 					}
-				}
-				else if (!WriteFile(Entry.DestFileName, Data, DataSize))
-				{
-					return false;
-				}
-				return true;
-			},
-			Prerequisites(ReadTask)));
-	}
+					return true;
+				},
+				Prerequisites(ReadTask)));
+		}
 
-	int32 ErrorCount = 0;
-	for (int32 EntryIndex = 0; EntryIndex < Entries.Num(); ++EntryIndex)
-	{
-		if (ExtractTasks[EntryIndex].GetResult())
+		for (int32 TaskIndex = 0; TaskIndex < NumTasks; ++TaskIndex)
 		{
-			const FEntry& Entry = Entries[EntryIndex];
-			if (OutOrderMap != nullptr)
+			if (ExtractTasks[TaskIndex].GetResult())
+			{
+				const FEntry& Entry = Entries[EntryStartIdx + TaskIndex];
+				if (OutOrderMap != nullptr)
+				{
+					OutOrderMap->Add(IndexReader.GetMountPoint() / Entry.SourceFileName, Entry.Offset);
+				}
+			}
+			else
 			{
-				OutOrderMap->Add(IndexReader.GetMountPoint() / Entry.SourceFileName, Entry.Offset);
+				++ErrorCount;
 			}
 		}
-		else
-		{
-			++ErrorCount;
-		}
+
+		EntryStartIdx += NumTasks;
 	}
 
 	UE_LOG(LogIoStore, Log, TEXT("Finished extracting %d chunks (including %d errors)."), Entries.Num(), ErrorCount);
diff --git a/Engine/Source/Developer/IoStoreUtilities/Public/IoStoreUtilities.h b/Engine/Source/Developer/IoStoreUtilities/Public/IoStoreUtilities.h
--- a/Engine/Source/Developer/IoStoreUtilities/Public/IoStoreUtilities.h
+++ b/Engine/Source/Developer/IoStoreUtilities/Public/IoStoreUtilities.h
@@ -11,6 +11,7 @@
 
 struct FGuid;
 struct FKeyChain;
+class FIoChunkId;
 
 IOSTOREUTILITIES_API int32 CreateIoStoreContainerFiles(const TCHAR* CmdLine);
 
@@ -39,4 +40,15 @@ IOSTOREUTILITIES_API bool ExtractFilesFromIoStoreContainer(
 	TArray<FGuid>* OutUsedEncryptionKeys,
 	bool* bOutIsSigned);
 
+IOSTOREUTILITIES_API bool ProcessFilesFromIoStoreContainer(
+	const TCHAR* InContainerFilename,
+	const TCHAR* InDestPath,
+	const FKeyChain& InKeyChain,
+	const FString* InFilter,
+	TFunction<bool (const FString&, const FString&, const FIoChunkId&, const uint8*, uint64)> FileProcessFunc,
+	TMap<FString, uint64>* OutOrderMap,
+	TArray<FGuid>* OutUsedEncryptionKeys,
+	bool* bOutIsSigned,
+	int32 MaxConcurrentReaders);
+
 IOSTOREUTILITIES_API bool SignIoStoreContainer(const TCHAR* InContainerFilename, const FRSAKeyHandle InSigningKey);
