Fixes compile issue in non-unity builds

diff --git a/Engine/Source/Developer/IoStoreUtilities/Private/IoStoreLooseFiles.h b/Engine/Source/Developer/IoStoreUtilities/Private/IoStoreLooseFiles.h
--- a/Engine/Source/Developer/IoStoreUtilities/Private/IoStoreLooseFiles.h
+++ b/Engine/Source/Developer/IoStoreUtilities/Private/IoStoreLooseFiles.h
@@ -3,7 +3,7 @@
 #pragma once
 
 #include "Containers/UnrealString.h"
-#include "Templates/UniquePtr.h"
+#include "Templates/SharedPointer.h"
 #include "UObject/NameTypes.h"
 
 class IIoStoreWriter;
