Ensure that multi-process cook doesn't cause workers to delete the oplog created by the director when using ZenStoreWriter.

diff --git a/Engine/Source/Developer/IoStoreUtilities/Private/ZenStoreWriter.cpp b/Engine/Source/Developer/IoStoreUtilities/Private/ZenStoreWriter.cpp
--- a/Engine/Source/Developer/IoStoreUtilities/Private/ZenStoreWriter.cpp
+++ b/Engine/Source/Developer/IoStoreUtilities/Private/ZenStoreWriter.cpp
@@ -403,7 +403,7 @@ void FZenStoreWriter::Initialize(const FCookInfo& Info)
 		FString OplogLifetimeMarkerPath = OutputPath / TEXT(".projectstore");
 		TUniquePtr<FArchive> OplogMarker(IFileManager::Get().CreateFileWriter(*OplogLifetimeMarkerPath));
 
-		bool bOplogEstablished = HttpClient->TryCreateOplog(ProjectId, OplogId, OplogLifetimeMarkerPath, Info.bFullBuild);
+		bool bOplogEstablished = HttpClient->TryCreateOplog(ProjectId, OplogId, OplogLifetimeMarkerPath, CleanBuild);
 
 		if (bOplogEstablished && OplogMarker)
 		{
