Reduce folders to scan when checking for config and localization files when creating zenfs manifest

diff --git a/Engine/Source/Developer/IoStoreUtilities/IoStoreUtilities.Build.cs b/Engine/Source/Developer/IoStoreUtilities/IoStoreUtilities.Build.cs
--- a/Engine/Source/Developer/IoStoreUtilities/IoStoreUtilities.Build.cs
+++ b/Engine/Source/Developer/IoStoreUtilities/IoStoreUtilities.Build.cs
@@ -13,7 +13,8 @@ public class IoStoreUtilities : ModuleRules
 		PrivateDependencyModuleNames.AddRange(new string[] {
             "Core",
             "CoreUObject",
-            "AssetRegistry",
+			"AssetRegistry",
+			"Projects",
 			"Zen",
 			"RenderCore",
         });
diff --git a/Engine/Source/Developer/IoStoreUtilities/Private/ZenFileSystemManifest.cpp b/Engine/Source/Developer/IoStoreUtilities/Private/ZenFileSystemManifest.cpp
--- a/Engine/Source/Developer/IoStoreUtilities/Private/ZenFileSystemManifest.cpp
+++ b/Engine/Source/Developer/IoStoreUtilities/Private/ZenFileSystemManifest.cpp
@@ -5,6 +5,7 @@
 #include "Algo/Sort.h"
 #include "HAL/PlatformFileManager.h"
 #include "Interfaces/ITargetPlatform.h"
+#include "Interfaces/IPluginManager.h"
 #include "Misc/App.h"
 #include "Misc/FileHelper.h"
 #include "Misc/Paths.h"
@@ -247,12 +248,118 @@ int32 FZenFileSystemManifest::Generate()
 		.IncludeExtension(TEXT("ini"));
 	AddFilesFromExtensionDirectories(TEXT("Config"), &ConfigFilter);
 
-	FFileFilter PluginFilter = FFileFilter()
-		.IncludeExtension(TEXT("ini"))
+	auto AddFromPluginPath = [this](const FString& EngineDir, const FString& ClientDirectory, const FString& SourcePath, const TCHAR* Path) {
+		//TRACE_CPUPROFILER_EVENT_SCOPE(AddManifestEntry);
+
+		FString ServerRelativePath = Path;
+		FPaths::MakePathRelativeTo(ServerRelativePath, *EngineDir);
+		ServerRelativePath = TEXT("/") + ServerRelativePath;
+
+		FString ClientPath = Path;
+		FPaths::MakePathRelativeTo(ClientPath, *(SourcePath / TEXT("")));
+		ClientPath = ClientDirectory / ClientPath;
+
+		const FIoChunkId FileChunkId = CreateExternalFileChunkId(ClientPath);
+		AddManifestEntry(
+			FileChunkId,
+			ServerRelativePath,
+			MoveTemp(ClientPath));
+	};
+
+	auto AddFromPluginDir = [this, &PlatformFile, &BaseFilter, &AddFromPluginPath](const FString& EngineDir, const FString& ClientDirectory, const FString& SourcePath, const FString& DirectoryPath, bool bIncludeSubdirs, FFileFilter* AdditionalFilter = nullptr) {
+		TArray<FString> DirectoriesToVisit;
+		auto VisitorFunc =
+			[this, &AddFromPluginPath, &EngineDir, &ClientDirectory, &SourcePath, &DirectoriesToVisit, &BaseFilter, bIncludeSubdirs, AdditionalFilter]//, &RootDir, &ClientDirectory, &LocalDirectory, &ServerRelativeDirectory, &BaseFilter, AdditionalFilter]
+			(const TCHAR* InnerFileNameOrDirectory, bool bIsDirectory)
+		{
+			if (bIsDirectory)
+			{
+				if (!bIncludeSubdirs)
+				{
+					return true;
+				}
+				FStringView DirectoryName = FPathViews::GetPathLeaf(InnerFileNameOrDirectory);
+				if (!BaseFilter.FilterDirectory(DirectoryName))
+				{
+					return true;
+				}
+				if (AdditionalFilter && !AdditionalFilter->FilterDirectory(DirectoryName))
+				{
+					return true;
+				}
+				DirectoriesToVisit.Add(InnerFileNameOrDirectory);
+				return true;
+			}
+			FStringView Extension = FPathViews::GetExtension(InnerFileNameOrDirectory);
+			if (!BaseFilter.FilterFile(Extension))
+			{
+				return true;
+			}
+			if (AdditionalFilter && !AdditionalFilter->FilterFile(Extension))
+			{
+				return true;
+			}
+			AddFromPluginPath(EngineDir, ClientDirectory, SourcePath, InnerFileNameOrDirectory);
+			return true;
+		};
+
+		DirectoriesToVisit.Push(DirectoryPath);
+		while (!DirectoriesToVisit.IsEmpty())
+		{
+			PlatformFile.IterateDirectory(*DirectoriesToVisit.Pop(false), VisitorFunc);
+		}
+	};
+
+	FFileFilter LocalizationFilter = FFileFilter()
 		.IncludeExtension(TEXT("locmeta"))
-		.IncludeExtension(TEXT("locres"))
-		.IncludeExtension(TEXT("uplugin"));
-	AddFilesFromExtensionDirectories(TEXT("Plugins"), &PluginFilter);
+		.IncludeExtension(TEXT("locres"));
+
+	const bool FilterDisabledPlugins = true;
+	FString PluginTargetPlatformString = PlatformInfo.UBTPlatformString;
+	IPluginManager& PluginManager = IPluginManager::Get();
+	TArray<TSharedRef<IPlugin>> DiscoveredPlugins = PluginManager.GetDiscoveredPlugins();
+	for (TSharedRef<IPlugin>& Plugin : DiscoveredPlugins)
+	{
+		FString ProjectName = Plugin->GetName();
+		if (FilterDisabledPlugins && !Plugin->IsEnabled())
+		{
+			UE_LOG(LogZenFileSystemManifest, Display, TEXT("Plugin '%s' disabled, skipping"), *ProjectName);
+			continue;
+		}
+		const FPluginDescriptor& Descriptor = Plugin->GetDescriptor();
+		if (!Descriptor.SupportsTargetPlatform(PluginTargetPlatformString))
+		{
+			UE_LOG(LogZenFileSystemManifest, Log, TEXT("Plugin '%s' not supported on platform '%s', skipping"), *ProjectName, *TargetPlatform.PlatformName());
+			for (const auto& SupportedTargetPlatform : Descriptor.SupportedTargetPlatforms)
+			{
+				UE_LOG(LogZenFileSystemManifest, Log, TEXT("       '%s' supports platform '%s'"), *ProjectName, *SupportedTargetPlatform);
+			}
+			continue;
+		}
+		FString BaseDir = Plugin->GetBaseDir();
+		FString ProjectFile = Plugin->GetDescriptorFileName();
+		FString ContentDir = Plugin->GetContentDir();
+		FString LocalizationDir = ContentDir / TEXT("Localization");
+		FString ConfigDir = BaseDir / TEXT("Config");
+		UE_LOG(LogZenFileSystemManifest, Display, TEXT("Plugin '%s': BaseDir: '%s'"), *ProjectName, *BaseDir);
+
+		FString ClientDirectory;
+		FString SourcePath;
+		switch (Plugin->GetLoadedFrom())
+		{
+		case EPluginLoadedFrom::Engine:
+			ClientDirectory = TEXT("/{engine}");
+			SourcePath = EngineDir;
+			break;
+		case EPluginLoadedFrom::Project:
+			ClientDirectory = TEXT("/{project}");
+			SourcePath = ProjectDir;
+			break;
+		}
+		AddFromPluginPath(EngineDir, ClientDirectory, SourcePath, *ProjectFile);
+		AddFromPluginDir(EngineDir, ClientDirectory, SourcePath, LocalizationDir, true, &LocalizationFilter);
+		AddFromPluginDir(EngineDir, ClientDirectory, SourcePath, ConfigDir, true, &ConfigFilter);
+	}
 
 	FString InternationalizationPresetAsString = UEnum::GetValueAsString(PackagingSettings->InternationalizationPreset);
 	const TCHAR* InternationalizationPresetPath = FCString::Strrchr(*InternationalizationPresetAsString, ':');
@@ -268,9 +375,6 @@ int32 FZenFileSystemManifest::Generate()
 	const TCHAR* ICUDataVersion = TEXT("icudt64l"); // TODO: Could this go into datadriven platform info? But it's basically always this.
 	AddFilesFromDirectory(*FPaths::Combine(TEXT("/{engine}"), TEXT("Content"), TEXT("Internationalization"), ICUDataVersion), FPaths::Combine(EngineDir, TEXT("Content"), TEXT("Internationalization"), InternationalizationPresetPath, ICUDataVersion), true);
 	
-	FFileFilter LocalizationFilter = FFileFilter()
-		.IncludeExtension(TEXT("locmeta"))
-		.IncludeExtension(TEXT("locres"));
 	AddFilesFromExtensionDirectories(TEXT("Content/Localization"), &LocalizationFilter);
 
 	FFileFilter ContentFilter = FFileFilter()
