Make sure we can gracefully destroy a FZenStoreWriter instance even if EndCook() is not called.

diff --git a/Engine/Source/Developer/IoStoreUtilities/Private/ZenStoreWriter.cpp b/Engine/Source/Developer/IoStoreUtilities/Private/ZenStoreWriter.cpp
--- a/Engine/Source/Developer/IoStoreUtilities/Private/ZenStoreWriter.cpp
+++ b/Engine/Source/Developer/IoStoreUtilities/Private/ZenStoreWriter.cpp
@@ -211,6 +211,13 @@ FZenStoreWriter::FZenStoreWriter(
 
 FZenStoreWriter::~FZenStoreWriter()
 {
+	if (CommitThread.IsValid())
+	{
+		UE_LOG(LogZenStoreWriter, Display, TEXT("Aborted, flushing..."));
+		CommitQueue->CompleteAdding();
+		CommitThread.Wait();
+	}
+
 	FScopeLock _(&PackagesCriticalSection);
 
 	if (PendingPackages.Num())
@@ -586,6 +593,7 @@ void FZenStoreWriter::EndCook(const FCookInfo& Info)
 	
 	CommitQueue->CompleteAdding();
 	CommitThread.Wait();
+	CommitThread.Reset();
 
 	if (!Info.bWorkerOnSharedSandbox)
 	{
