Store cooker computer remote ips in .projectstore

diff --git a/Engine/Source/Developer/IoStoreUtilities/IoStoreUtilities.Build.cs b/Engine/Source/Developer/IoStoreUtilities/IoStoreUtilities.Build.cs
--- a/Engine/Source/Developer/IoStoreUtilities/IoStoreUtilities.Build.cs
+++ b/Engine/Source/Developer/IoStoreUtilities/IoStoreUtilities.Build.cs
@@ -17,7 +17,8 @@ public class IoStoreUtilities : ModuleRules
 			"Projects",
 			"Zen",
 			"RenderCore",
-        });
+			"Sockets",
+		});
 
 		PublicIncludePathModuleNames.AddRange(new string[] {
 			"Zen",
diff --git a/Engine/Source/Developer/IoStoreUtilities/Private/ZenStoreWriter.cpp b/Engine/Source/Developer/IoStoreUtilities/Private/ZenStoreWriter.cpp
--- a/Engine/Source/Developer/IoStoreUtilities/Private/ZenStoreWriter.cpp
+++ b/Engine/Source/Developer/IoStoreUtilities/Private/ZenStoreWriter.cpp
@@ -16,6 +16,8 @@
 #include "Serialization/CompactBinaryWriter.h"
 #include "Serialization/CompactBinarySerialization.h"
 #include "Serialization/LargeMemoryWriter.h" 
+#include "SocketSubsystem.h"
+#include "IPAddress.h"
 #include "HAL/FileManager.h"
 #include "HAL/PlatformFileManager.h"
 #include "IO/IoDispatcher.h"
@@ -423,6 +425,23 @@ void FZenStoreWriter::Initialize(const FCookInfo& Info)
 #endif
 			ManifestWriter << "islocalhost" << IsRunningLocally;
 			ManifestWriter << "hostname" << HttpClient->GetHostName();
+			if (IsRunningLocally)
+			{
+				ISocketSubsystem* SocketSubsystem = ISocketSubsystem::Get();
+				if (SocketSubsystem != nullptr)
+				{
+					TArray<TSharedPtr<FInternetAddr>> Addresses;
+					if (SocketSubsystem->GetLocalAdapterAddresses(Addresses))
+					{
+						ManifestWriter.BeginArray("remotehostnames");
+						for (const TSharedPtr<FInternetAddr>& Address : Addresses)
+						{
+							ManifestWriter << Address->ToString(false);
+						}
+						ManifestWriter.EndArray();
+					}
+				}
+			}
 			ManifestWriter << "hostport" << HttpClient->GetPort();
 			ManifestWriter << "projectid" << ProjectId;
 			ManifestWriter << "oplogid" << OplogId;
diff --git a/Engine/Source/Runtime/StorageServerClient/Private/StorageServerPlatformFile.cpp b/Engine/Source/Runtime/StorageServerClient/Private/StorageServerPlatformFile.cpp
--- a/Engine/Source/Runtime/StorageServerClient/Private/StorageServerPlatformFile.cpp
+++ b/Engine/Source/Runtime/StorageServerClient/Private/StorageServerPlatformFile.cpp
@@ -286,6 +286,15 @@ bool FStorageServerPlatformFile::ShouldBeUsed(IPlatformFile* Inner, const TCHAR*
 			{
 				HostAddrs.Add(FString(HostName));
 			}
+			FCbArrayView RemoteHostNames = ZenServerField["remotehostnames"].AsArrayView();
+			for (FCbFieldView RemoteHostName : RemoteHostNames)
+			{
+				if (FUtf8StringView HostName = RemoteHostName.AsString(); !HostName.IsEmpty())
+				{
+					HostAddrs.Add(FString(HostName));
+				}
+			}
+
 			HostPort = ZenServerField["hostport"].AsUInt16(HostPort);
 			UE_LOG(LogStorageServerPlatformFile, Display, TEXT("Found .projectstore"));
 		}
