UnrealPak now uses the CookMetadata file to optionally write out size information for plugins. This information is written back to the cook metadata file, and is enable via a project packaging setting.

diff --git a/Engine/Source/Developer/CookMetadata/Private/CookMetadata.cpp b/Engine/Source/Developer/CookMetadata/Private/CookMetadata.cpp
--- a/Engine/Source/Developer/CookMetadata/Private/CookMetadata.cpp
+++ b/Engine/Source/Developer/CookMetadata/Private/CookMetadata.cpp
@@ -35,6 +35,10 @@ bool FCookMetadataState::Serialize(FArchive& Ar)
 	Ar << Version;
 	Ar << PluginHierarchy;
 	Ar << AssociatedDevelopmentAssetRegistryHash;
+	Ar << AssociatedDevelopmentAssetRegistryHashPostWriteback;
+	Ar << Platform;
+	Ar << BuildVersion;
+	Ar << HordeJobId;
 
 	return true;
 }
diff --git a/Engine/Source/Developer/CookMetadata/Public/CookMetadata.h b/Engine/Source/Developer/CookMetadata/Public/CookMetadata.h
--- a/Engine/Source/Developer/CookMetadata/Public/CookMetadata.h
+++ b/Engine/Source/Developer/CookMetadata/Public/CookMetadata.h
@@ -14,12 +14,100 @@ namespace UE::Cook
 enum class ECookMetadataStateVersion : uint8
 {
 	PluginHierarchy = 1,
+	PostWritebackHash = 2,
 
 	// Add new versions above this.
 	VersionCount,
 	LatestVersion = VersionCount - 1
 };
 
+/**
+*	We classify various bits of the plugin data based on how it will be delivered to the end user in order to
+*	more appropriately track the user's experience.
+* 
+*	Each iostore chunk (note: NOT pak chunk!) gets compressed during staging and unreal pak will update the size
+*	for the plugin based on how the chunk gets deployed.
+*/
+enum class EPluginSizeTypes : uint8
+{
+	// The iostore chunks will be deployed to servers for downloading on-demand by the game using the Individual Asset Streaming
+	// system.
+	Streaming,
+	// The iostore chunks are written to normal iostore containers that are expected to be distributed with the game. This
+	// includes the required global iostore container.
+	Installed,
+	// The iostore chunks are written to a separate container that isn't required to be distributed with the game. This is where e.g.
+	// OptionalMips go for textures. They appear as .uptnl files in the Cooked directory (if cooking to Loose Files),
+	// The CopyBuildToStagingDirectory script manually assigns the iostore chunk to a corresponding pak chunk with the name
+	// ending in "optional", e.g. pakChunk0optional.pak/ucas/utoc/sig. The only way to catch this in UnrealPak
+	// is to parse the filename.
+	Optional,
+	// These are sidecar files for distributing EditorOnly data for Cooked Editor builds. When cooking to loose files
+	// they will contain ".o" inside their filename, e.g. "myasset.o.ubulk". They are never intended to be shipped with
+	// a game.
+	OptionalSegment,
+	COUNT
+};
+constexpr uint8 EPluginSizeTypesCount = (uint8)EPluginSizeTypes::COUNT;
+
+struct FPluginSizeInfo
+{
+	uint64 Sizes[EPluginSizeTypesCount] = {};
+
+	void AddSizes(uint64 SizesPerType[EPluginSizeTypesCount])
+	{
+		for (uint8 Type = 0; Type < EPluginSizeTypesCount; Type++)
+		{
+			Sizes[Type] += SizesPerType[Type];
+		}
+	}
+	void Add(const FPluginSizeInfo& Other)
+	{
+		for (uint8 Type = 0; Type < EPluginSizeTypesCount; Type++)
+		{
+			Sizes[Type] += Other.Sizes[Type];
+		}
+	}
+	uint64 TotalSize() const
+	{
+		uint64 Total = 0;
+		for (uint8 Type = 0; Type < EPluginSizeTypesCount; Type++)
+		{
+			Total += Sizes[Type];
+		}
+		return Total;
+	}
+
+	friend FArchive& operator<<(FArchive& Ar, FPluginSizeInfo& SizeInfo)
+	{
+		for (uint8 i = 0; i < EPluginSizeTypesCount; i++)
+		{
+			Ar << SizeInfo.Sizes[i];
+		}
+		return Ar;
+	}
+
+	uint64& operator[](EPluginSizeTypes InType) { return Sizes[(uint8)InType]; }
+	const uint64& operator[](EPluginSizeTypes InType) const { return Sizes[(uint8)InType]; }
+};
+
+/*
+*	Unrealpak doesn't compress the data for all platforms. In those cases, the data written back
+*	during staging isn't compressed and is not representative of the final sizes that occurs after
+*	the corresponding SDK tools process them for deployment.
+*/
+enum class ECookMetadataSizesPresent
+{
+	// staging has not occured, or writeback wasn't enabled in project packaing settings.
+	NotPresent,
+
+	// unrealpak compressed the iostore chunks and the data we have is compressed sizes
+	Compressed,
+
+	// the selected platform isn't compressed by unrealpak, or package compression was disabled.
+	Uncompressed
+};
+
 
 /** The name and dependency information for a plugin that was enabled during cooking. */
 struct COOKMETADATA_API FCookMetadataPluginEntry
@@ -35,14 +123,24 @@ struct COOKMETADATA_API FCookMetadataPluginEntry
 	//		const UE::Cook::FCookMetadataPluginEntry& DependentPlugin = PluginHierarchy.PluginsEnabledAtCook[PluginHierarchy.PluginDependencies[DependencyIndex]];
 	//	}
 	//
-	uint16 DependencyIndexStart;
-	uint16 DependencyIndexEnd;
+	uint16 DependencyIndexStart = 0;
+	uint16 DependencyIndexEnd = 0;
+
+	//
+	// Theses sizes are set during staging by unrealpak if the option in project packaging is set.
+	// To determine if they are set, check FCookMetadataState::GetSizesPresent(). Inclusive contains
+	// the size of the plugin and all of its dependencies listed in its uplugin file.
+	//
+	FPluginSizeInfo InclusiveSizes;
+	FPluginSizeInfo ExclusiveSizes;
 
 	uint16 DependencyCount() const { return DependencyIndexEnd - DependencyIndexStart; }
 
 	friend FArchive& operator<<(FArchive& Ar, FCookMetadataPluginEntry& Entry)
 	{
-		return Ar << Entry.Name << Entry.DependencyIndexStart << Entry.DependencyIndexEnd;
+		Ar << Entry.Name << Entry.DependencyIndexStart << Entry.DependencyIndexEnd;
+		Ar << Entry.InclusiveSizes << Entry.ExclusiveSizes;
+		return Ar;
 	}
 };
 
@@ -88,6 +186,9 @@ public:
 	// Plugin hierarchy information
 	void SetPluginHierarchyInfo(FCookMetadataPluginHierarchy&& InPluginHierarchy) { PluginHierarchy = MoveTemp(InPluginHierarchy); }
 	const FCookMetadataPluginHierarchy& GetPluginHierarchy() const { return PluginHierarchy; }
+	
+	// So that unrealpak can update the sizes.
+	FCookMetadataPluginHierarchy& GetMutablePluginHierarchy() { return PluginHierarchy; }
 
 	/**
 	*	Associated DevAR Hash.
@@ -98,19 +199,54 @@ public:
 	*	Use this to ensure that the files you are working with were produced by the same cook and didn't
 	*	get out of sync somehow.
 	* 
+	*	*IMPORTANT* If asset registry writeback is enabled during staging, then the hash of the development
+	*	asset registry changes, and you'll need to check against GetAssociatedDevelopmentAssetRegistryHashPostWriteback.
+	*	If you don't know which one you have, check both - they are both valid.
+	* 
 	*	e.g.
 	*	uint64 CheckHash = GetAssociatedDevelopmentAssetRegistryHash();
 	*	bValidDevAr = ComputeHashOfDevelopmentAssetRegistry(MakeMemoryView(SerializedAssetRegistry)) == CheckHash;
 	*/
 	void SetAssociatedDevelopmentAssetRegistryHash(uint64 InHash) { AssociatedDevelopmentAssetRegistryHash = InHash; }
+	void SetAssociatedDevelopmentAssetRegistryHashPostWriteback(uint64 InHash) { AssociatedDevelopmentAssetRegistryHashPostWriteback = InHash; }
+
 	uint64 GetAssociatedDevelopmentAssetRegistryHash() const { return AssociatedDevelopmentAssetRegistryHash; }
+	uint64 GetAssociatedDevelopmentAssetRegistryHashPostWriteback() const { return AssociatedDevelopmentAssetRegistryHashPostWriteback; }
+
+	void SetPlatformAndBuildVersion(const FString& InPlatform, const TCHAR* InBuildVersion) { Platform = InPlatform; BuildVersion = FString(InBuildVersion); }
+	const FString& GetPlatform() const { return Platform; }
+	const FString GetBuildVersion() const { return BuildVersion; }
+
+	void SetHordeJobId(FString&& InHordeJobId) { HordeJobId = MoveTemp(InHordeJobId); }
+	const FString& GetHordeJobId() const { return HordeJobId; }
 
 	static uint64 ComputeHashOfDevelopmentAssetRegistry(FMemoryView InSerializedDevelopmentAssetRegistry);
+
+	// Returns what size information is present in FCookMetadataPluginEntry. This varies based on the platform
+	// and settings. 
+	ECookMetadataSizesPresent GetSizesPresent() const { return SizesPresent; }
+	void SetSizesPresent(ECookMetadataSizesPresent InSizesPresent) { SizesPresent = InSizesPresent; }
 private:
+
 	ECookMetadataStateVersion Version;
 	FCookMetadataPluginHierarchy PluginHierarchy;
 
-	uint64 AssociatedDevelopmentAssetRegistryHash;
+	uint64 AssociatedDevelopmentAssetRegistryHash = 0;
+
+	// Asset registry size writeback changes the AR, so we have a separate hash for that DevAR that this
+	// also matches.
+	uint64 AssociatedDevelopmentAssetRegistryHashPostWriteback = 0;
+
+	FString Platform;
+
+	// BUILD_VERSION definition from definitions.h for the cook.
+	FString BuildVersion;
+
+	// If cooked on Horde, this is the job id that cooked it.
+	FString HordeJobId;
+
+	// Updated by unrealpak when plugin size information is added.
+	ECookMetadataSizesPresent SizesPresent;
 };
 
 COOKMETADATA_API const FString& GetCookMetadataFilename();
diff --git a/Engine/Source/Developer/DeveloperToolSettings/Classes/Settings/ProjectPackagingSettings.h b/Engine/Source/Developer/DeveloperToolSettings/Classes/Settings/ProjectPackagingSettings.h
--- a/Engine/Source/Developer/DeveloperToolSettings/Classes/Settings/ProjectPackagingSettings.h
+++ b/Engine/Source/Developer/DeveloperToolSettings/Classes/Settings/ProjectPackagingSettings.h
@@ -326,6 +326,12 @@ public:
 	UPROPERTY(config, EditAnywhere, Category = Packaging, AdvancedDisplay)
 	EAssetRegistryWritebackMethod WriteBackMetadataToAssetRegistry;
 
+	/**
+	* Whether or not to write a json summary file that contains size information to the cooked Metadata/PluginJsons directory
+	*/
+	UPROPERTY(config, EditAnywhere, Category = Packaging, AdvancedDisplay, meta = (EditCondition = "WriteBackMetadataToAssetRegistry != EAssetRegistryWritebackMethod::Disabled"))
+	bool bWritePluginSizeSummaryJsons;
+
 	/**
 	 * Create compressed cooked packages (decreased deployment size)
 	 */
diff --git a/Engine/Source/Developer/IoStoreUtilities/IoStoreUtilities.Build.cs b/Engine/Source/Developer/IoStoreUtilities/IoStoreUtilities.Build.cs
--- a/Engine/Source/Developer/IoStoreUtilities/IoStoreUtilities.Build.cs
+++ b/Engine/Source/Developer/IoStoreUtilities/IoStoreUtilities.Build.cs
@@ -14,6 +14,7 @@ public class IoStoreUtilities : ModuleRules
             "Core",
             "CoreUObject",
 			"AssetRegistry",
+			"CookMetadata",
 			"Projects",
 			"Zen",
 			"RenderCore",
diff --git a/Engine/Source/Developer/IoStoreUtilities/Private/IoStoreUtilities.cpp b/Engine/Source/Developer/IoStoreUtilities/Private/IoStoreUtilities.cpp
--- a/Engine/Source/Developer/IoStoreUtilities/Private/IoStoreUtilities.cpp
+++ b/Engine/Source/Developer/IoStoreUtilities/Private/IoStoreUtilities.cpp
@@ -3,7 +3,9 @@
 #include "IoStoreUtilities.h"
 
 #include "IoStoreLooseFiles.h"
+#include "Algo/TopologicalSort.h"
 #include "Async/AsyncWork.h"
+#include "CookMetadata.h"
 #include "HAL/FileManager.h"
 #include "HAL/PlatformFileManager.h"
 #include "Hash/CityHash.h"
@@ -83,6 +85,14 @@ TRACE_DECLARE_MEMORY_COUNTER(IoStoreUsedFileBufferMemory, TEXT("IoStore/UsedFile
 // Helper to format numbers with comma separators to help readability: 1,234 vs 1234.
 static FString NumberString(uint64 N) { return FText::AsNumber(N).ToString(); }
 
+// Used for tracking how the chunk will get deployed and thus what classification its size is.
+struct FIoStoreChunkSource
+{
+	FIoStoreTocChunkInfo ChunkInfo;
+	UE::Cook::EPluginSizeTypes SizeType;
+
+};
+
 static const FName DefaultCompressionMethod = NAME_Zlib;
 static const uint64 DefaultCompressionBlockSize = 64 << 10;
 static const uint64 DefaultCompressionBlockAlignment = 64 << 10;
@@ -945,6 +955,7 @@ struct FIoStoreArguments
 	bool bFileRegions = false;
 	bool bUpload = false;
 	EAssetRegistryWritebackMethod WriteBackMetadataToAssetRegistry = EAssetRegistryWritebackMethod::Disabled;
+	bool bWritePluginSizeSummaryJsons = false; // Only valid if WriteBackMetadataToAssetRegistry != Disabled.
 
 	FOodleDataCompression::ECompressor ShaderOodleCompressor = FOodleDataCompression::ECompressor::Mermaid;
 	FOodleDataCompression::ECompressionLevel ShaderOodleLevel = FOodleDataCompression::ECompressionLevel::Normal;
@@ -3157,7 +3168,257 @@ public:
 	static constexpr uint64 BufferMemoryLimit = 2ull << 30;
 };
 
-static void AddChunkInfoToAssetRegistry(TMap<FPackageId, TArray<FIoStoreTocChunkInfo, TInlineAllocator<2>>>&& PackageToChunks, FAssetRegistryState& AssetRegistry, uint64 TotalCompressedSize)
+static bool WriteUtf8StringView(FUtf8StringView InView, const FString& InFilename)
+{
+	TUniquePtr<FArchive> Ar(IFileManager::Get().CreateFileWriter(*InFilename, 0));
+	if (!Ar)
+	{
+		return false;
+	}
+	UTF8CHAR UTF8BOM[] = { (UTF8CHAR)0xEF, (UTF8CHAR)0xBB, (UTF8CHAR)0xBF };
+	Ar->Serialize(&UTF8BOM, sizeof(UTF8BOM));
+	Ar->Serialize((void*)InView.GetData(), InView.Len() * sizeof(UTF8CHAR));
+	Ar->Close();
+	return true;
+}
+
+
+static bool SavePluginMetadata(const FString& InAssetRegistryFileName, const FString& PluginName, TUtf8StringBuilder<4096>& InPluginMetadataJson)
+{
+	FString PluginMetadataFilename = FPaths::GetPath(InAssetRegistryFileName) / TEXT("PluginJsons") / PluginName + TEXT(".json");
+	if (WriteUtf8StringView(InPluginMetadataJson.ToView(), PluginMetadataFilename) == false)
+	{
+		UE_LOG(LogIoStore, Error, TEXT("Unable to write plugin metadata file: %s"), *PluginMetadataFilename);
+		return false;
+	}
+	return true;
+}
+
+struct FPluginGraphEntry
+{
+	uint16 IndexInEnabledPlugins = 0;
+	TArray<const UE::Cook::FCookMetadataPluginEntry*> Dependencies;
+	TSet<const UE::Cook::FCookMetadataPluginEntry*> TotalDependencies;
+	uint32 DirectRefcount = 0;
+	bool bIsRoot = false;
+	UE::Cook::FPluginSizeInfo ExclusiveSize;
+	UE::Cook::FPluginSizeInfo InclusiveSize;
+};
+
+// Rework the hierarchy in to a graph where we have output edges resolved to pointers so we can pass to
+// library functions.
+static void GeneratePluginGraph(const UE::Cook::FCookMetadataPluginHierarchy& InPluginHierarchy, TMap<FStringView, FPluginGraphEntry>& OutPluginGraph)
+{
+	uint16 PluginIndex = 0;
+	for (const UE::Cook::FCookMetadataPluginEntry& Plugin : InPluginHierarchy.PluginsEnabledAtCook)
+	{
+		FPluginGraphEntry& OurEntry = OutPluginGraph.FindOrAdd(Plugin.Name);
+		OurEntry.IndexInEnabledPlugins = PluginIndex;
+
+		for (uint16 DependencyIndex = Plugin.DependencyIndexStart; DependencyIndex < Plugin.DependencyIndexEnd; DependencyIndex++)
+		{
+			const UE::Cook::FCookMetadataPluginEntry& DependentPlugin = InPluginHierarchy.PluginsEnabledAtCook[InPluginHierarchy.PluginDependencies[DependencyIndex]];
+
+			OurEntry.Dependencies.Add(&DependentPlugin);
+		}
+		PluginIndex++;
+	}
+}
+
+/**
+*	Use the name of the package to assign sizes so that we can track build size at a per-plugin level,
+*	and write out jsons files for each plugin in to the cooked metadata directory.
+* 
+*	Plugins insert themselves in to the package's path at the top level. Content that is unassigned
+*	to a plugin has either /Engine or /Game as it's top level path and will be assigned to a pseudo
+*	plugin.
+*/
+static void WritePluginMetadataJsons(const FString& InAssetRegistryFileName, TMap<FPackageId, TArray<FIoStoreChunkSource, TInlineAllocator<2>>>& PackageToChunks, FAssetRegistryState& AssetRegistry, UE::Cook::FCookMetadataState& CookMetadata)
+{
+	//
+	// Using the name of the packages, assign chunk sizes to their associated plugin.
+	//
+	const UE::Cook::FCookMetadataPluginHierarchy& PluginHierarchy = CookMetadata.GetPluginHierarchy();
+	TMap<FStringView, FPluginGraphEntry> PluginGraph;
+	GeneratePluginGraph(PluginHierarchy, PluginGraph);
+
+	for (uint16 RootIndex : PluginHierarchy.RootPlugins)
+	{
+		PluginGraph[PluginHierarchy.PluginsEnabledAtCook[RootIndex].Name].bIsRoot = true;
+	}
+
+
+	const TMap<FName, const FAssetPackageData*> AssetPackageMap = AssetRegistry.GetAssetPackageDataMap();
+	for (const TPair<FName, const FAssetPackageData*>& AssetPackage : AssetPackageMap)
+	{
+		if (AssetPackage.Value->DiskSize < 0)
+		{
+			// No data on disk!
+			continue;
+		}
+		
+		const TArray<FIoStoreChunkSource, TInlineAllocator<2>>* PackageChunks = PackageToChunks.Find(FPackageId::FromName(AssetPackage.Key));
+		if (PackageChunks == nullptr)
+		{
+			// This happens when the package has been stripped by UAT prior to staging by e.g. PakDenyList.
+			continue;
+		}
+
+		UE::Cook::FPluginSizeInfo PackageSizes;
+		for (const FIoStoreChunkSource& ChunkInfo : *PackageChunks)
+		{
+			PackageSizes[ChunkInfo.SizeType] += ChunkInfo.ChunkInfo.CompressedSize;
+		}
+
+		// Assign the size to the package's plugin.
+		{
+			TStringBuilder<FName::StringBufferSize> PackageNameStr(InPlace, AssetPackage.Key);
+			FStringView PackageName(PackageNameStr);
+
+			FStringView PluginName = FPackageName::SplitPackageNameRoot(PackageName, nullptr);
+			FPluginGraphEntry* PluginEntry = PluginGraph.Find(PluginName);
+			if (PluginEntry)
+			{
+				PluginEntry->ExclusiveSize.Add(PackageSizes);
+			}
+			else
+			{
+				UE_LOG(LogTemp, Warning, TEXT("Plugin for package not found: %s (%.*s)"), PackageNameStr.GetData(), PluginName.Len(), PluginName.GetData());
+			}
+		}
+	}
+
+	//
+	// Generate the inclusive size for each plugin.
+	//
+
+	// Sort the plugins topologically. This means that when we iterate linearly,
+	// we know that when we hit a plugin, we've already processed the dependencies.
+	// This takes some memory to track edges but is a depth first search
+	// and not anything quadratic or worse.
+	TArray<const UE::Cook::FCookMetadataPluginEntry*> SortedList;
+	{
+		for (const UE::Cook::FCookMetadataPluginEntry& Plugin : PluginHierarchy.PluginsEnabledAtCook)
+		{
+			SortedList.Add(&Plugin);
+		}
+
+		auto GetElementDependencies = [&PluginGraph](const UE::Cook::FCookMetadataPluginEntry* PluginEntry) -> const TArray<const UE::Cook::FCookMetadataPluginEntry*>&
+		{
+			return PluginGraph[PluginEntry->Name].Dependencies;
+		};
+
+		Algo::TopologicalSort(SortedList, GetElementDependencies);
+
+		// Make sure the topological sort worked correctly.
+#if DO_CHECK
+		TSet<uint16> Test;
+		for (const UE::Cook::FCookMetadataPluginEntry* Check : SortedList)
+		{
+			Test.Add(PluginGraph[Check->Name].IndexInEnabledPlugins);
+
+			for (uint16 DependencyIndex = Check->DependencyIndexStart; DependencyIndex < Check->DependencyIndexEnd; DependencyIndex++)
+			{
+				const UE::Cook::FCookMetadataPluginEntry& DependentPlugin = PluginHierarchy.PluginsEnabledAtCook[PluginHierarchy.PluginDependencies[DependencyIndex]];
+
+				check (Test.Contains(PluginHierarchy.PluginDependencies[DependencyIndex]));
+			}
+		}
+#endif
+	}
+
+	// With a topological sort we can do this linearly since we know the
+	// plugins before us have already calculated their inclusive sizes.
+	for (const UE::Cook::FCookMetadataPluginEntry* Plugin : SortedList)
+	{
+		FPluginGraphEntry& PluginEntry = PluginGraph[Plugin->Name];
+
+		for (uint16 DependencyIndex = Plugin->DependencyIndexStart; DependencyIndex < Plugin->DependencyIndexEnd; DependencyIndex++)
+		{
+			const UE::Cook::FCookMetadataPluginEntry& DependentPlugin = PluginHierarchy.PluginsEnabledAtCook[PluginHierarchy.PluginDependencies[DependencyIndex]];
+
+			PluginEntry.TotalDependencies.Add(&DependentPlugin);
+			
+			FPluginGraphEntry& DependentEntry = PluginGraph[DependentPlugin.Name];
+			DependentEntry.DirectRefcount++;
+
+			PluginEntry.TotalDependencies.Append(DependentEntry.TotalDependencies);
+		}
+
+		PluginEntry.InclusiveSize = PluginEntry.ExclusiveSize;
+		for (const UE::Cook::FCookMetadataPluginEntry* Dependency : PluginEntry.TotalDependencies)
+		{
+			PluginEntry.InclusiveSize.Add(Dependency->InclusiveSizes);
+		}
+	}
+
+	auto GeneratePluginJson = [](TUtf8StringBuilder<4096>& OutPluginMetadataJson, FStringView InName, const FPluginGraphEntry& InGraphEntry)
+	{
+		OutPluginMetadataJson.Reset();
+		OutPluginMetadataJson << "{\n";
+		OutPluginMetadataJson << "\t\"name\":\"" << InName << "\",\n";
+
+		OutPluginMetadataJson << "\t\"schema_version\":1,\n";
+
+		OutPluginMetadataJson << "\t\"is_root_plugin\":" << (InGraphEntry.bIsRoot ? TEXTVIEW("true") : TEXTVIEW("false")) << ",\n";
+
+		OutPluginMetadataJson << "\t\"exclusive_installed\":" << InGraphEntry.ExclusiveSize[UE::Cook::EPluginSizeTypes::Installed] << ",\n";
+		OutPluginMetadataJson << "\t\"exclusive_optional\":" << InGraphEntry.ExclusiveSize[UE::Cook::EPluginSizeTypes::Optional] << ",\n";
+		OutPluginMetadataJson << "\t\"exclusive_ias\":" << InGraphEntry.ExclusiveSize[UE::Cook::EPluginSizeTypes::Streaming] << ",\n";
+		OutPluginMetadataJson << "\t\"exclusive_optionalsegment\":" << InGraphEntry.ExclusiveSize[UE::Cook::EPluginSizeTypes::OptionalSegment] << ",\n";
+
+		OutPluginMetadataJson << "\t\"inclusive_installed\":" << InGraphEntry.InclusiveSize[UE::Cook::EPluginSizeTypes::Installed] << ",\n";
+		OutPluginMetadataJson << "\t\"inclusive_optional\":" << InGraphEntry.InclusiveSize[UE::Cook::EPluginSizeTypes::Optional] << ",\n";
+		OutPluginMetadataJson << "\t\"inclusive_ias\":" << InGraphEntry.InclusiveSize[UE::Cook::EPluginSizeTypes::Streaming] << ",\n";
+		OutPluginMetadataJson << "\t\"inclusive_optionalsegment\":" << InGraphEntry.InclusiveSize[UE::Cook::EPluginSizeTypes::OptionalSegment] << "\n";
+
+		OutPluginMetadataJson << "\t\"direct_refcount\":" << InGraphEntry.DirectRefcount << ",\n";
+
+		OutPluginMetadataJson << "}\n";
+	};
+
+	//
+	// Generate the plugin_summary jsons.	
+	//
+	TUtf8StringBuilder<4096> PluginMetadataJson;
+
+	// Also write a csv for easier browsing in spreadsheets.
+	TUtf8StringBuilder<4096> Csv;
+	Csv.Append("name,exclusive_installed,exclusive_optional,inclusive_installed,inclusive_optional,direct_refcount\n");
+
+	// We want to write the sizes back to the cook metadata, so we need a non-const version.
+	UE::Cook::FCookMetadataPluginHierarchy& MutablePluginHierarchy = CookMetadata.GetMutablePluginHierarchy();
+	for (UE::Cook::FCookMetadataPluginEntry& Plugin : MutablePluginHierarchy.PluginsEnabledAtCook)
+	{
+		const FPluginGraphEntry& PluginEntry = PluginGraph[Plugin.Name];
+		if (PluginEntry.InclusiveSize.TotalSize() == 0)
+		{
+			continue;
+		}
+
+		GeneratePluginJson(PluginMetadataJson, Plugin.Name, PluginEntry);
+		SavePluginMetadata(InAssetRegistryFileName, Plugin.Name, PluginMetadataJson);
+
+		Plugin.InclusiveSizes = PluginEntry.InclusiveSize;
+		Plugin.ExclusiveSizes = PluginEntry.ExclusiveSize;
+
+		Csv.Appendf("%ls,%llu,%llu,%llu,%llu,%u\n", *Plugin.Name, 
+			Plugin.ExclusiveSizes[UE::Cook::EPluginSizeTypes::Installed], Plugin.ExclusiveSizes[UE::Cook::EPluginSizeTypes::Optional],
+			Plugin.InclusiveSizes[UE::Cook::EPluginSizeTypes::Installed], Plugin.InclusiveSizes[UE::Cook::EPluginSizeTypes::Optional],
+			PluginEntry.DirectRefcount);
+	}
+
+	{
+		FString CsvFilename = FPaths::GetPath(InAssetRegistryFileName) / TEXT("plugin_sizes.csv");
+		if (WriteUtf8StringView(Csv.ToView(), CsvFilename) == false)
+		{
+			UE_LOG(LogIoStore, Error, TEXT("Unable to write plugin csv file: %s"), *CsvFilename);
+		}
+	}
+}
+
+
+static void AddChunkInfoToAssetRegistry(TMap<FPackageId, TArray<FIoStoreChunkSource, TInlineAllocator<2>>>&& PackageToChunks, FAssetRegistryState& AssetRegistry, uint64 TotalCompressedSize)
 {
 	//
 	// The asset registry has the chunks associate with each package, so we can just iterate the
@@ -3186,7 +3447,7 @@ static void AddChunkInfoToAssetRegistry(TMap<FPackageId, TArray<FIoStoreTocChunk
 			continue;
 		}
 
-		const TArray<FIoStoreTocChunkInfo, TInlineAllocator<2>>* PackageChunks = PackageToChunks.Find(FPackageId::FromName(AssetPackage.Key));
+		const TArray<FIoStoreChunkSource, TInlineAllocator<2>>* PackageChunks = PackageToChunks.Find(FPackageId::FromName(AssetPackage.Key));
 		if (PackageChunks == nullptr)
 		{
 			// This happens when the package has been stripped by UAT prior to staging by e.g. PakDenyList.
@@ -3196,11 +3457,11 @@ static void AddChunkInfoToAssetRegistry(TMap<FPackageId, TArray<FIoStoreTocChunk
 		int32 ChunkCount = 0;
 		int64 Size = 0;
 		int64 CompressedSize = 0;
-		for (const FIoStoreTocChunkInfo& ChunkInfo : *PackageChunks)
+		for (const FIoStoreChunkSource& ChunkInfo : *PackageChunks)
 		{
 			ChunkCount++;
-			Size += ChunkInfo.Size;
-			CompressedSize += ChunkInfo.CompressedSize;
+			Size += ChunkInfo.ChunkInfo.Size;
+			CompressedSize += ChunkInfo.ChunkInfo.CompressedSize;
 		}
 
 		FAssetDataTagMap TagsAndValues;
@@ -3221,9 +3482,9 @@ static void AddChunkInfoToAssetRegistry(TMap<FPackageId, TArray<FIoStoreTocChunk
 	uint64 RemainingByType[(uint8)EIoChunkType::MAX] = {};
 	for (auto PackageChunks : PackageToChunks)
 	{
-		for (FIoStoreTocChunkInfo& Info : PackageChunks.Value)
+		for (FIoStoreChunkSource& Info : PackageChunks.Value)
 		{
-			RemainingByType[(uint8)Info.ChunkType] += Info.CompressedSize;
+			RemainingByType[(uint8)Info.ChunkInfo.ChunkType] += Info.ChunkInfo.CompressedSize;
 		}
 	}
 
@@ -3245,15 +3506,33 @@ static void AddChunkInfoToAssetRegistry(TMap<FPackageId, TArray<FIoStoreTocChunk
 	}
 }
 
-static bool LoadAssetRegistry(const FString& InAssetRegistryFileName, FAssetRegistryState& OutAssetRegistry)
+// Returns the hash of the development asset registry or 0 on failure.
+static uint64 LoadAssetRegistry(const FString& InAssetRegistryFileName, FAssetRegistryState& OutAssetRegistry)
 {
 	FAssetRegistryVersion::Type Version;
 	FAssetRegistryLoadOptions Options(UE::AssetRegistry::ESerializationTarget::ForDevelopment);
-	bool bSucceeded = FAssetRegistryState::LoadFromDisk(*InAssetRegistryFileName, Options, OutAssetRegistry, &Version);
-	return bSucceeded;
+
+	TUniquePtr<FArchive> FileReader(IFileManager::Get().CreateFileReader(*InAssetRegistryFileName));
+	if (FileReader)
+	{
+		TArray64<uint8> Data;
+		Data.SetNumUninitialized(FileReader->TotalSize());
+		FileReader->Serialize(Data.GetData(), Data.Num());
+		check(!FileReader->IsError());
+
+		uint64 DevArHash = UE::Cook::FCookMetadataState::ComputeHashOfDevelopmentAssetRegistry(MakeMemoryView(Data));
+
+		FLargeMemoryReader MemoryReader(Data.GetData(), Data.Num());
+		if (OutAssetRegistry.Load(MemoryReader, Options, &Version))
+		{
+			return DevArHash;
+		}
+	}
+
+	return 0;;
 }
 
-static bool SaveAssetRegistry(const FString& InAssetRegistryFileName, FAssetRegistryState& InAssetRegistry, bool InSaveTempAndRename)
+static bool SaveAssetRegistry(const FString& InAssetRegistryFileName, FAssetRegistryState& InAssetRegistry, uint64* OutDevArHash)
 {
 	TRACE_CPUPROFILER_EVENT_SCOPE(SavingAssetRegistry);
 	FLargeMemoryWriter SerializedAssetRegistry;
@@ -3263,7 +3542,12 @@ static bool SaveAssetRegistry(const FString& InAssetRegistryFileName, FAssetRegi
 		return false;
 	}
 
-	FString OutputFileName = InSaveTempAndRename ? (InAssetRegistryFileName + TEXT(".temp")) : InAssetRegistryFileName;
+	if (OutDevArHash)
+	{
+		*OutDevArHash = UE::Cook::FCookMetadataState::ComputeHashOfDevelopmentAssetRegistry(MakeMemoryView(SerializedAssetRegistry.GetData(), SerializedAssetRegistry.TotalSize()));
+	}
+
+	FString OutputFileName = InAssetRegistryFileName + TEXT(".temp");
 
 	TUniquePtr<FArchive> Writer = TUniquePtr<FArchive>(IFileManager::Get().CreateFileWriter(*OutputFileName));
 	if (!Writer)
@@ -3283,14 +3567,11 @@ static bool SaveAssetRegistry(const FString& InAssetRegistryFileName, FAssetRegi
 		return false;
 	}
 
-	if (InSaveTempAndRename)
+	// Move our temp file over the original asset registry.
+	if (IFileManager::Get().Move(*InAssetRegistryFileName, *OutputFileName) == false)
 	{
-		// Move our temp file over the original asset registry.
-		if (IFileManager::Get().Move(*InAssetRegistryFileName, *OutputFileName) == false)
-		{
-			// Error already logged by FileManager
-			return false;
-		}
+		// Error already logged by FileManager
+		return false;
 	}
 
 	UE_LOG(LogIoStore, Display, TEXT("Saved asset registry to disk. (%s)"), *InAssetRegistryFileName);
@@ -3298,14 +3579,14 @@ static bool SaveAssetRegistry(const FString& InAssetRegistryFileName, FAssetRegi
 	return true;
 }
 
-int32 DoAssetRegistryWritebackAfterStage(const FString& InAssetRegistryFileName, FString&& InContainerDirectory, const FKeyChain& InKeyChain)
+static int32 DoAssetRegistryWritebackAfterStage(const FString& InAssetRegistryFileName, FString&& InContainerDirectory, const FKeyChain& InKeyChain)
 {
 	// This version called after the containers are already created, when you
 	// have a bunch of containers on disk and you want to add chunk info back to
 	// an asset registry.
 
 	FAssetRegistryState AssetRegistry;
-	if (LoadAssetRegistry(InAssetRegistryFileName, AssetRegistry) == false)
+	if (LoadAssetRegistry(InAssetRegistryFileName, AssetRegistry) == 0)
 	{
 		UE_LOG(LogIoStore, Error, TEXT("Unabled to open source asset registry: %s"), *InAssetRegistryFileName);
 		return 1;
@@ -3319,7 +3600,7 @@ int32 DoAssetRegistryWritebackAfterStage(const FString& InAssetRegistryFileName,
 	uint64 TotalCompressedSize = 0;
 	
 	// Grab all the package infos.
-	TMap<FPackageId, TArray<FIoStoreTocChunkInfo, TInlineAllocator<2>>> PackageToChunks;
+	TMap<FPackageId, TArray<FIoStoreChunkSource, TInlineAllocator<2>>> PackageToChunks;
 	for (const FString& Filename : FoundContainerFiles)
 	{
 		TUniquePtr<FIoStoreReader> Reader = CreateIoStoreReader(*(InContainerDirectory / Filename), InKeyChain);
@@ -3331,7 +3612,8 @@ int32 DoAssetRegistryWritebackAfterStage(const FString& InAssetRegistryFileName,
 		Reader->EnumerateChunks([&](const FIoStoreTocChunkInfo& ChunkInfo)
 		{
 			FPackageId PackageId = FPackageId::FromValue(*(int64*)(ChunkInfo.Id.GetData()));
-			PackageToChunks.FindOrAdd(PackageId).Add(ChunkInfo);
+			// (Deployment can't be ascertained after staging - plugin jsons are not generated.)
+			PackageToChunks.FindOrAdd(PackageId).Add({ChunkInfo, UE::Cook::EPluginSizeTypes::COUNT});
 			TotalCompressedSize += ChunkInfo.CompressedSize;
 			return true;
 		});
@@ -3339,10 +3621,22 @@ int32 DoAssetRegistryWritebackAfterStage(const FString& InAssetRegistryFileName,
 
 	AddChunkInfoToAssetRegistry(MoveTemp(PackageToChunks), AssetRegistry, TotalCompressedSize);
 
-	return SaveAssetRegistry(InAssetRegistryFileName, AssetRegistry, true) ? 0 : 1;
+	return SaveAssetRegistry(InAssetRegistryFileName, AssetRegistry, nullptr) ? 0 : 1;
 }
 
-static bool FindAndLoadDevelopmentAssetRegistry(const FString& InCookedDir, bool bInRequired, FAssetRegistryState& OutAssetRegistry, FString* OutAssetRegistryFileName /*optional, set on success*/)
+enum class ECookMetadataFiles
+{
+	None = 0,
+	AssetRegistry = 1,
+	CookMetadata = 2,
+	All = 4
+};
+ENUM_CLASS_FLAGS(ECookMetadataFiles);
+
+static ECookMetadataFiles FindAndLoadMetadataFiles(
+	const FString& InCookedDir, ECookMetadataFiles InRequiredFiles, 
+	FAssetRegistryState& OutAssetRegistry, FString* OutAssetRegistryFileName /*optional, set on success*/,
+	UE::Cook::FCookMetadataState* OutCookMetadata, FString* OutCookMetadataFileName /*optional, set on success or need*/)
 {
 	TRACE_CPUPROFILER_EVENT_SCOPE(LoadingAssetRegistry);
 
@@ -3361,7 +3655,7 @@ static bool FindAndLoadDevelopmentAssetRegistry(const FString& InCookedDir, bool
 
 	if (PossibleAssetRegistryFiles.Num() == 0)
 	{
-		if (bInRequired)
+		if (EnumHasAnyFlags(InRequiredFiles, ECookMetadataFiles::AssetRegistry))
 		{
 			UE_LOG(LogIoStore, Error, TEXT("No development asset registry file found!"));
 		}
@@ -3369,71 +3663,159 @@ static bool FindAndLoadDevelopmentAssetRegistry(const FString& InCookedDir, bool
 		{
 			UE_LOG(LogIoStore, Display, TEXT("No development asset registry file found!"));
 		}
-		return false;
+		return ECookMetadataFiles::None;
 	}
 
 	UE_LOG(LogIoStore, Display, TEXT("Using input asset registry: %s"), *PossibleAssetRegistryFiles[0]);
-	if (LoadAssetRegistry(PossibleAssetRegistryFiles[0], OutAssetRegistry) == false)
+	uint64 LoadedDevArHash = LoadAssetRegistry(PossibleAssetRegistryFiles[0], OutAssetRegistry);
+
+	if (LoadedDevArHash == 0)
+	{
+		return ECookMetadataFiles::None; // already logged
+	}
+
+	// If we found the asset registry, try and find the cook metadata that should be next to it.
+	ECookMetadataFiles ResultFiles = ECookMetadataFiles::AssetRegistry;
+
+	if (OutCookMetadata)
 	{
-		return false; // already logged
+		// The cook metadata file should be adjacent to the development asset registry.
+		FString CookMetadataFileName = FPaths::GetPath(PossibleAssetRegistryFiles[0]) / UE::Cook::GetCookMetadataFilename();
+		TUniquePtr<FArchive> FileReader(IFileManager::Get().CreateFileReader(*CookMetadataFileName));
+		if (FileReader)
+		{
+			TArray64<uint8> Data;
+			Data.SetNumUninitialized(FileReader->TotalSize());
+			FileReader->Serialize(Data.GetData(), Data.Num());
+			check(!FileReader->IsError());
+
+			FLargeMemoryReader MemoryReader(Data.GetData(), Data.Num());
+			if (OutCookMetadata->Serialize(MemoryReader) == false)
+			{
+				UE_LOG(LogIoStore, Error, TEXT("Failed to deserialize cook metadata file - loaded successfully, but invalid data. [%s]"), *CookMetadataFileName);
+				if (EnumHasAnyFlags(InRequiredFiles, ECookMetadataFiles::CookMetadata))
+				{
+					return ECookMetadataFiles::None;
+				}
+			}
+			else if (OutCookMetadata->GetAssociatedDevelopmentAssetRegistryHash() != LoadedDevArHash &&
+				OutCookMetadata->GetAssociatedDevelopmentAssetRegistryHashPostWriteback() != LoadedDevArHash) // during testing we can repeat stage after cook so we might have already edited it.
+			{
+				UE_LOG(LogIoStore, Error, 
+					TEXT("Cook metadata file mismatch: Hash of associated development asset registry does not match. [%s] %llx vs %llx (%llx post writeback)"), 
+					*CookMetadataFileName, LoadedDevArHash, OutCookMetadata->GetAssociatedDevelopmentAssetRegistryHash(), OutCookMetadata->GetAssociatedDevelopmentAssetRegistryHashPostWriteback());
+				if (EnumHasAnyFlags(InRequiredFiles, ECookMetadataFiles::CookMetadata))
+				{
+					return ECookMetadataFiles::None;
+				}
+			}
+			else
+			{
+				EnumAddFlags(ResultFiles, ECookMetadataFiles::CookMetadata);
+				if (OutCookMetadataFileName)
+				{
+					*OutCookMetadataFileName = MoveTemp(CookMetadataFileName);
+				}
+			}
+		}
+		else
+		{
+			if (EnumHasAnyFlags(InRequiredFiles, ECookMetadataFiles::CookMetadata))
+			{
+				UE_LOG(LogIoStore, Error, TEXT("Failed to open and read cook metadata file %s"), *CookMetadataFileName);
+				return ECookMetadataFiles::None;
+			}
+
+			UE_LOG(LogIoStore, Display, TEXT("No cook metadata file found, checked %s"), *CookMetadataFileName);
+			if (OutCookMetadataFileName)
+			{
+				*OutCookMetadataFileName = FString("");
+			}
+		}
 	}
 
+
 	if (OutAssetRegistryFileName)
 	{
 		*OutAssetRegistryFileName = MoveTemp(PossibleAssetRegistryFiles[0]);
 	}
-	return true;
+	return ResultFiles;
 }
 
-bool DoAssetRegistryWritebackDuringStage(EAssetRegistryWritebackMethod InMethod, const FString& InCookedDir, TArray<TSharedPtr<IIoStoreWriter>>& InIoStoreWriters)
+static bool DoAssetRegistryWritebackDuringStage(
+	EAssetRegistryWritebackMethod InMethod, 
+	bool bInWritePluginMetadata,
+	const FString& InCookedDir, 
+	bool bInCompressionEnabled,
+	TArray<TSharedPtr<IIoStoreWriter>>& InIoStoreWriters, 
+	TArray<UE::Cook::EPluginSizeTypes>& IoStoreWriterSizeClassifications)
 {
 	// This version called during container creation.
-
 	TRACE_CPUPROFILER_EVENT_SCOPE(UpdateAssetRegistryWithSizeInfo);
 	UE_LOG(LogIoStore, Display, TEXT("Adding staging metadata to asset registry..."));
 
 	// The overwhelming majority of time for the asset registry writeback is loading and saving.
 	FString AssetRegistryFileName;
 	FAssetRegistryState AssetRegistry;
-	if (FindAndLoadDevelopmentAssetRegistry(InCookedDir, true, AssetRegistry, &AssetRegistryFileName) == false)
+	FString CookMetadataFileName;
+	UE::Cook::FCookMetadataState CookMetadata;
+	ECookMetadataFiles FilesNeeded = ECookMetadataFiles::AssetRegistry;
+	if (bInWritePluginMetadata)
+	{
+		EnumAddFlags(FilesNeeded, ECookMetadataFiles::CookMetadata);
+	}
+	if (FindAndLoadMetadataFiles(InCookedDir, FilesNeeded, AssetRegistry, &AssetRegistryFileName, &CookMetadata, &CookMetadataFileName) == ECookMetadataFiles::None)
 	{
 		// already logged
 		return false;
 	}
 
-	// Create a map off the package id to all of its chunks. 2 inline allocation
-	// is for the export data and the bulk data. For a major test project, 2 covers
-	// 89% of packages, 1 covers 72%.
-	uint64 TotalCompressedSize = 0;
-	TMap<FPackageId, TArray<FIoStoreTocChunkInfo, TInlineAllocator<2>>> PackageToChunks;
-	for (TSharedPtr<IIoStoreWriter> IoStoreWriter : InIoStoreWriters)
+	//
+	// We want to separate out the sizes based on where they go in the end product.
+	//
+	UE::Cook::FPluginSizeInfo ProductSize;	
+	TMap<FPackageId, TArray<FIoStoreChunkSource, TInlineAllocator<2>>> PackageToChunks;
 	{
-		IoStoreWriter->EnumerateChunks([&](const FIoStoreTocChunkInfo& ChunkInfo)
+		int32 IoStoreWriterIndex = 0;
+		for (TSharedPtr<IIoStoreWriter> IoStoreWriter : InIoStoreWriters)
 		{
-			FPackageId PackageId = FPackageId::FromValue(*(int64*)(ChunkInfo.Id.GetData()));
-			PackageToChunks.FindOrAdd(PackageId).Add(ChunkInfo);
-			TotalCompressedSize += ChunkInfo.CompressedSize;
-			return true;
-		});
+			IoStoreWriter->EnumerateChunks([&PackageToChunks, IoStoreWriterClassification = IoStoreWriterSizeClassifications[IoStoreWriterIndex], &ProductSize](const FIoStoreTocChunkInfo& ChunkInfo)
+			{
+				FPackageId PackageId = FPackageId::FromValue(*(int64*)(ChunkInfo.Id.GetData()));
+				PackageToChunks.FindOrAdd(PackageId).Add({ChunkInfo, IoStoreWriterClassification });
+				ProductSize[IoStoreWriterClassification] += ChunkInfo.CompressedSize;
+				return true;
+			});
+
+			IoStoreWriterIndex++;
+		}
 	}
 
-	AddChunkInfoToAssetRegistry(MoveTemp(PackageToChunks), AssetRegistry, TotalCompressedSize);
+	uint64 UpdatedDevArHash = 0;
+	if (bInWritePluginMetadata)
+	{
+		UE_LOG(LogIoStore, Display, TEXT("Writing plugin size jsons..."));
+		WritePluginMetadataJsons(AssetRegistryFileName, PackageToChunks, AssetRegistry, CookMetadata);
+	}
 
-	FString OutputFileName;
+	AddChunkInfoToAssetRegistry(MoveTemp(PackageToChunks), AssetRegistry, ProductSize.TotalSize());
+	
+	
 	switch (InMethod)
 	{
 	case EAssetRegistryWritebackMethod::OriginalFile:
 		{
 			// Write to an adjacent file and move after
-			if (SaveAssetRegistry(AssetRegistryFileName, AssetRegistry, true) == false)
+			if (SaveAssetRegistry(AssetRegistryFileName, AssetRegistry, bInWritePluginMetadata ? &UpdatedDevArHash : nullptr) == false)
 			{
 				return false;
 			}
+
 			break;
 		}
 	case EAssetRegistryWritebackMethod::AdjacentFile:
 		{
-			if (SaveAssetRegistry(AssetRegistryFileName.Replace(TEXT(".bin"), TEXT("Staged.bin")), AssetRegistry, true) == false)
+			if (SaveAssetRegistry(AssetRegistryFileName.Replace(TEXT(".bin"), TEXT("Staged.bin")), AssetRegistry, bInWritePluginMetadata ? &UpdatedDevArHash : nullptr) == false)
 			{
 				return false;
 			}
@@ -3445,6 +3827,32 @@ bool DoAssetRegistryWritebackDuringStage(EAssetRegistryWritebackMethod InMethod,
 			return false;
 		}
 	}
+
+	// Since we modified the dev ar, we need to save the updated hash in the cook metadata so it can still validate.
+	if (bInWritePluginMetadata)
+	{
+		CookMetadata.SetSizesPresent(bInCompressionEnabled ? UE::Cook::ECookMetadataSizesPresent::Compressed : UE::Cook::ECookMetadataSizesPresent::Uncompressed);
+		CookMetadata.SetAssociatedDevelopmentAssetRegistryHashPostWriteback(UpdatedDevArHash);
+
+		FArrayWriter SerializedCookMetadata;
+		CookMetadata.Serialize(SerializedCookMetadata);
+
+		FString TempFileName = CookMetadataFileName + TEXT(".temp");
+		if (FFileHelper::SaveArrayToFile(SerializedCookMetadata, *TempFileName))
+		{
+			// Move our temp file over the original asset registry.
+			if (IFileManager::Get().Move(*CookMetadataFileName, *TempFileName) == false)
+			{
+				// Error already logged by FileManager
+				return false;
+			}
+		}
+		else
+		{
+			UE_LOG(LogIoStore, Error, TEXT("Failed to save temp file for write updated cook metadata file (%s"), *TempFileName);
+			return false;
+		}
+	}
 	
 	return true;
 }
@@ -3462,7 +3870,7 @@ public:
 	{
 		FString AssetRegistryFileName;
 		FAssetRegistryState AssetRegistry;
-		if (FindAndLoadDevelopmentAssetRegistry(InCookedDir, false, AssetRegistry, nullptr) == false)
+		if (FindAndLoadMetadataFiles(InCookedDir, ECookMetadataFiles::None, AssetRegistry, nullptr, nullptr, nullptr) == ECookMetadataFiles::None)
 		{
 			// already logged
 			return false;
@@ -3767,6 +4175,7 @@ int32 CreateTarget(const FIoStoreArguments& Arguments, const FIoStoreWriterSetti
 	}
 	TArray<FString> OnDemandContainers;
 	TArray<TSharedPtr<IIoStoreWriter>> IoStoreWriters;
+	TArray<UE::Cook::EPluginSizeTypes> IoStoreWriterSizeClassifications;
 	TSharedPtr<IIoStoreWriter> GlobalIoStoreWriter;
 	{
 		IOSTORE_CPU_SCOPE(InitializeWriters);
@@ -3781,6 +4190,7 @@ int32 CreateTarget(const FIoStoreArguments& Arguments, const FIoStoreWriterSetti
 			}
 			GlobalIoStoreWriter = IoStoreWriterContext->CreateContainer(*Arguments.GlobalContainerPath, GlobalContainerSettings);
 			IoStoreWriters.Add(GlobalIoStoreWriter);
+			IoStoreWriterSizeClassifications.Add(UE::Cook::EPluginSizeTypes::Installed);
 		}
 		for (FContainerTargetSpec* ContainerTarget : ContainerTargets)
 		{
@@ -3797,6 +4207,7 @@ int32 CreateTarget(const FIoStoreArguments& Arguments, const FIoStoreWriterSetti
 				FLooseFilesWriterSettings WriterSettings;
 				WriterSettings.TargetRootPath = ContainerTarget->StageLooseFileRootPath;
 				ContainerTarget->IoStoreWriter = MakeLooseFilesIoStoreWriter(WriterSettings);
+				IoStoreWriterSizeClassifications.Add(UE::Cook::EPluginSizeTypes::Streaming); // LooseFiles currently end up as a streamed source.
 				IoStoreWriters.Add(ContainerTarget->IoStoreWriter);
 			}
 			else
@@ -3831,11 +4242,37 @@ int32 CreateTarget(const FIoStoreArguments& Arguments, const FIoStoreWriterSetti
 				ContainerTarget->IoStoreWriter->SetReferenceChunkDatabase(ChunkDatabase);
 				ContainerTarget->IoStoreWriter->SetHashDatabase(HashDatabase, Arguments.bVerifyHashDatabase);
 				IoStoreWriters.Add(ContainerTarget->IoStoreWriter);
+
+				if (EnumHasAnyFlags(ContainerTarget->ContainerFlags, EIoContainerFlags::OnDemand))
+				{
+					IoStoreWriterSizeClassifications.Add(UE::Cook::EPluginSizeTypes::Streaming);
+				}
+				else
+				{
+					// There's no way to know whether a container is optional without parsing the filename, see
+					// EIoStoreWriterType.
+					FString BaseFileName = FPaths::GetBaseFilename(ContainerTarget->OutputPath, true);
+					// Strip the platform identifier off the pak
+					if (int32 DashIndex=0; BaseFileName.FindLastChar(TEXT('-'), DashIndex))
+					{
+						BaseFileName.LeftInline(DashIndex);
+					}
+					if (BaseFileName.EndsWith(TEXT("optional")))
+					{
+						IoStoreWriterSizeClassifications.Add(UE::Cook::EPluginSizeTypes::Optional);
+					}
+					else
+					{
+						IoStoreWriterSizeClassifications.Add(UE::Cook::EPluginSizeTypes::Installed);
+					}
+				}
+
 				if (!ContainerTarget->OptionalSegmentOutputPath.IsEmpty())
 				{
 					ContainerTarget->OptionalSegmentIoStoreWriter = IoStoreWriterContext->CreateContainer(*ContainerTarget->OptionalSegmentOutputPath, ContainerSettings);
 					ContainerTarget->OptionalSegmentIoStoreWriter->SetReferenceChunkDatabase(ChunkDatabase);
 					IoStoreWriters.Add(ContainerTarget->OptionalSegmentIoStoreWriter);
+					IoStoreWriterSizeClassifications.Add(UE::Cook::EPluginSizeTypes::OptionalSegment);
 				}
 				if (EnumHasAnyFlags(ContainerTarget->ContainerFlags, EIoContainerFlags::OnDemand))
 				{
@@ -4129,7 +4566,7 @@ int32 CreateTarget(const FIoStoreArguments& Arguments, const FIoStoreWriterSetti
 
 	if (Arguments.WriteBackMetadataToAssetRegistry != EAssetRegistryWritebackMethod::Disabled)
 	{
-		DoAssetRegistryWritebackDuringStage(Arguments.WriteBackMetadataToAssetRegistry, Arguments.CookedDir, IoStoreWriters);
+		DoAssetRegistryWritebackDuringStage(Arguments.WriteBackMetadataToAssetRegistry, Arguments.bWritePluginSizeSummaryJsons, Arguments.CookedDir, GeneralIoWriterSettings.CompressionMethod != NAME_None, IoStoreWriters, IoStoreWriterSizeClassifications);
 	}
 
 	TArray<FIoStoreWriterResult> IoStoreWriterResults;
@@ -7342,13 +7779,14 @@ bool ParseContainerGenerationArguments(FIoStoreArguments& Arguments, FIoStoreWri
 	// This should be a path to your last released containers. If those containers are encrypted, be sure to
 	// provide keys via -ReferenceContainerCryptoKeys.
 	//
-	FParse::Value(FCommandLine::Get(), TEXT("-ReferenceContainerGlobalFileName="), Arguments.ReferenceChunkGlobalContainerFileName);
-
-	FString CryptoKeysCacheFilename;
-	if (FParse::Value(FCommandLine::Get(), TEXT("-ReferenceContainerCryptoKeys="), CryptoKeysCacheFilename))
+	if (FParse::Value(FCommandLine::Get(), TEXT("-ReferenceContainerGlobalFileName="), Arguments.ReferenceChunkGlobalContainerFileName))
 	{
-		UE_LOG(LogIoStore, Display, TEXT("Parsing reference container crypto keys from a crypto key cache file '%s'"), *CryptoKeysCacheFilename);
-		KeyChainUtilities::LoadKeyChainFromFile(CryptoKeysCacheFilename, Arguments.ReferenceChunkKeys);
+		FString CryptoKeysCacheFilename;
+		if (FParse::Value(FCommandLine::Get(), TEXT("-ReferenceContainerCryptoKeys="), CryptoKeysCacheFilename))
+		{
+			UE_LOG(LogIoStore, Display, TEXT("Parsing reference container crypto keys from a crypto key cache file '%s'"), *CryptoKeysCacheFilename);
+			KeyChainUtilities::LoadKeyChainFromFile(CryptoKeysCacheFilename, Arguments.ReferenceChunkKeys);
+		}
 	}
 
 	// By default, we use any hashes in the asset registry that exist in order to avoid reading and hashing
@@ -7827,6 +8265,8 @@ int32 CreateIoStoreContainerFiles(const TCHAR* CmdLine)
 			UE_LOG(LogIoStore, Error, TEXT("Valid options are: AdjacentFile, OriginalFile, Disabled."), *WriteBackMetadataToAssetRegistry);
 			return -1;
 		}
+
+		Arguments.bWritePluginSizeSummaryJsons = FParse::Param(FCommandLine::Get(), TEXT("WritePluginSizeSummaryJsons"));
 	}
 
 	FString PackageStoreManifestFilename;
diff --git a/Engine/Source/Editor/UnrealEd/Private/CookOnTheFlyServer.cpp b/Engine/Source/Editor/UnrealEd/Private/CookOnTheFlyServer.cpp
--- a/Engine/Source/Editor/UnrealEd/Private/CookOnTheFlyServer.cpp
+++ b/Engine/Source/Editor/UnrealEd/Private/CookOnTheFlyServer.cpp
@@ -9269,9 +9269,15 @@ void UCookOnTheFlyServer::WriteCookMetadata(const ITargetPlatform* InTargetPlatf
 	TArray<TSharedRef<IPlugin>> EnabledPlugins = IPluginManager::Get().GetEnabledPlugins();
 
 	// Remove any plugins that don't support the target platform.
-	EnabledPlugins.RemoveAllSwap([InTargetPlatform](const TSharedRef<IPlugin>& Plugin)
-	{
-		return InTargetPlatform->IsEnabledForPlugin(Plugin.Get()) == false;
+	// NOTE: We can't use IsEnabledForPlugin because it has an issue with the AllowTargets list where preventing a plugin on
+	// a target at the uproject level doesn't get overridden by a dependent plugins' reference, so we remove plugins that
+	// are actually enabled.
+	// We switch to just directly checking the plugin if it's enabled for the platform. This won't handle transitive disabling,
+	// and also won't handle any direct disabling... however disabling should have already been taken care of by using GetEnabledPlugins
+	// as our base list.
+	EnabledPlugins.RemoveAllSwap([&PlatformNameString](const TSharedRef<IPlugin>& Plugin)
+	{
+		return Plugin->GetDescriptor().SupportsTargetPlatform(PlatformNameString) == false;
 	});
 
 	// Filter to the DLC plugin + dependencies if we are a DLC cook.
@@ -9329,9 +9335,10 @@ void UCookOnTheFlyServer::WriteCookMetadata(const ITargetPlatform* InTargetPlatf
 		}
 	}
 
-	if (IntFitsIn<uint16>(EnabledPlugins.Num()) == false)
+	constexpr int32 AdditionalPseudoPlugins = 2; // /Engine and /Game.
+	if (IntFitsIn<uint16>(EnabledPlugins.Num() + AdditionalPseudoPlugins) == false)
 	{
-		UE_LOG(LogCook, Warning, TEXT("Number of plugins exceeds 64k, unable to write cook metadata file (count = %d"), EnabledPlugins.Num());
+		UE_LOG(LogCook, Warning, TEXT("Number of plugins exceeds 64k, unable to write cook metadata file (count = %d"), EnabledPlugins.Num() + AdditionalPseudoPlugins);
 	}
 	else
 	{
@@ -9395,6 +9402,10 @@ void UCookOnTheFlyServer::WriteCookMetadata(const ITargetPlatform* InTargetPlatf
 			}
 		}
 
+		// Add the /Engine and /Game pseudo plugins. These are placeholders for holding size information when unrealpak runs.
+		PluginsToAdd.AddDefaulted_GetRef().Name = TEXT("Engine");
+		PluginsToAdd.AddDefaulted_GetRef().Name = TEXT("Game");
+
 		if (IntFitsIn<uint16>(PluginChildArray.Num()) == false)
 		{
 			UE_LOG(LogCook, Warning, TEXT("Number of child plugins exceeds 64k, unable to write cook metadata file (count = %d)"), PluginChildArray.Num());
@@ -9411,9 +9422,15 @@ void UCookOnTheFlyServer::WriteCookMetadata(const ITargetPlatform* InTargetPlatf
 			MetadataState.SetPluginHierarchyInfo(MoveTemp(PluginHierarchy));
 			MetadataState.SetAssociatedDevelopmentAssetRegistryHash(InDevelopmentAssetRegistryHash);
 
+			MetadataState.SetPlatformAndBuildVersion(PlatformNameString, FApp::GetBuildVersion());
+			MetadataState.SetHordeJobId(FPlatformMisc::GetEnvironmentVariable(TEXT("UE_HORDE_JOBID")));
+
 			FArrayWriter SerializedCookMetadata;
 			MetadataState.Serialize(SerializedCookMetadata);
-			FFileHelper::SaveArrayToFile(SerializedCookMetadata, *GetCookedCookMetadataFilename(PlatformNameString));
+			if (FFileHelper::SaveArrayToFile(SerializedCookMetadata, *GetCookedCookMetadataFilename(PlatformNameString)) == false)
+			{
+				UE_LOG(LogCook, Error, TEXT("Failed to write cook metadata file (%s)"), *GetCookedCookMetadataFilename(PlatformNameString));
+			}
 		}
 	}
 }
diff --git a/Engine/Source/Programs/AutomationTool/Scripts/CopyBuildToStagingDirectory.Automation.cs b/Engine/Source/Programs/AutomationTool/Scripts/CopyBuildToStagingDirectory.Automation.cs
--- a/Engine/Source/Programs/AutomationTool/Scripts/CopyBuildToStagingDirectory.Automation.cs
+++ b/Engine/Source/Programs/AutomationTool/Scripts/CopyBuildToStagingDirectory.Automation.cs
@@ -3585,6 +3585,13 @@ namespace AutomationScripts
 					{
 						AdditionalArgs += " -WriteBackMetadataToAssetRegistry=" + WriteBackMetadataToAssetRegistry;
 					}
+
+					bool bWritePluginSizeSummaryJsons = false;
+					if (PlatformGameConfig.GetBool("/Script/UnrealEd.ProjectPackagingSettings", "bWritePluginSizeSummaryJsons", out bWritePluginSizeSummaryJsons) &&
+						bWritePluginSizeSummaryJsons)
+					{
+						AdditionalArgs += " -WritePluginSizeSummaryJsons";
+					}
 				}
 
 
