Quick fix -> should be summing exclusive sizes no inclusive

diff --git a/Engine/Source/Developer/IoStoreUtilities/Private/IoStoreUtilities.cpp b/Engine/Source/Developer/IoStoreUtilities/Private/IoStoreUtilities.cpp
--- a/Engine/Source/Developer/IoStoreUtilities/Private/IoStoreUtilities.cpp
+++ b/Engine/Source/Developer/IoStoreUtilities/Private/IoStoreUtilities.cpp
@@ -3348,7 +3348,7 @@ static void WritePluginMetadataJsons(const FString& InAssetRegistryFileName, TMa
 		PluginEntry.InclusiveSize = PluginEntry.ExclusiveSize;
 		for (const UE::Cook::FCookMetadataPluginEntry* Dependency : PluginEntry.TotalDependencies)
 		{
-			PluginEntry.InclusiveSize.Add(Dependency->InclusiveSizes);
+			PluginEntry.InclusiveSize.Add(Dependency->ExclusiveSizes);
 		}
 	}
 
