Track timing for the computation since it's N^2, in case it goes large at some point in the future.

diff --git a/Engine/Source/Developer/IoStoreUtilities/Private/IoStoreUtilities.cpp b/Engine/Source/Developer/IoStoreUtilities/Private/IoStoreUtilities.cpp
--- a/Engine/Source/Developer/IoStoreUtilities/Private/IoStoreUtilities.cpp
+++ b/Engine/Source/Developer/IoStoreUtilities/Private/IoStoreUtilities.cpp
@@ -3235,6 +3235,8 @@ static void GeneratePluginGraph(const UE::Cook::FCookMetadataPluginHierarchy& In
 */
 static void WritePluginMetadataJsons(const FString& InAssetRegistryFileName, TMap<FPackageId, TArray<FIoStoreChunkSource, TInlineAllocator<2>>>& PackageToChunks, FAssetRegistryState& AssetRegistry, UE::Cook::FCookMetadataState& CookMetadata)
 {
+	double WritePluginStart = FPlatformTime::Seconds();
+
 	//
 	// Using the name of the packages, assign chunk sizes to their associated plugin.
 	//
@@ -3327,8 +3329,10 @@ static void WritePluginMetadataJsons(const FString& InAssetRegistryFileName, TMa
 #endif
 	}
 
-	// With a topological sort we can do this linearly since we know the
-	// plugins before us have already calculated their inclusive sizes.
+	// For inclusive sizes we need to gather the set of all dependencies. This ends up being technically
+	// O(N^2) in the worst case. It's highly unlikely our plugin DAG will cause that, but we track the times
+	// just so we can keep an eye on it if it ends up taking measurable amounts of time.
+	double InclusiveComputeStart = FPlatformTime::Seconds();
 	for (const UE::Cook::FCookMetadataPluginEntry* Plugin : SortedList)
 	{
 		FPluginGraphEntry& PluginEntry = PluginGraph[Plugin->Name];
@@ -3345,12 +3349,14 @@ static void WritePluginMetadataJsons(const FString& InAssetRegistryFileName, TMa
 			PluginEntry.TotalDependencies.Append(DependentEntry.TotalDependencies);
 		}
 
+		// In the worse case this is another O(N) iteration, which makes us overall O(N^2)
 		PluginEntry.InclusiveSize = PluginEntry.ExclusiveSize;
 		for (const UE::Cook::FCookMetadataPluginEntry* Dependency : PluginEntry.TotalDependencies)
 		{
 			PluginEntry.InclusiveSize.Add(Dependency->ExclusiveSizes);
 		}
 	}
+	double InclusiveComputeEnd = FPlatformTime::Seconds();
 
 	auto GeneratePluginJson = [](TUtf8StringBuilder<4096>& OutPluginMetadataJson, FStringView InName, const FPluginGraphEntry& InGraphEntry)
 	{
@@ -3413,8 +3419,12 @@ static void WritePluginMetadataJsons(const FString& InAssetRegistryFileName, TMa
 		if (WriteUtf8StringView(Csv.ToView(), CsvFilename) == false)
 		{
 			UE_LOG(LogIoStore, Error, TEXT("Unable to write plugin csv file: %s"), *CsvFilename);
+			return;
 		}
 	}
+
+	double WritePluginEnd = FPlatformTime::Seconds();
+	UE_LOG(LogIoStore, Display, TEXT("Wrote plugin size jsons/csv in %.2f seconds, inclusive computation was %.2f of that"), WritePluginEnd - WritePluginStart, InclusiveComputeEnd - InclusiveComputeStart);
 }
 
 
@@ -3794,7 +3804,6 @@ static bool DoAssetRegistryWritebackDuringStage(
 	uint64 UpdatedDevArHash = 0;
 	if (bInWritePluginMetadata)
 	{
-		UE_LOG(LogIoStore, Display, TEXT("Writing plugin size jsons..."));
 		WritePluginMetadataJsons(AssetRegistryFileName, PackageToChunks, AssetRegistry, CookMetadata);
 	}
 
