Only cause an error when loading the cook metadata file if it's a required file.

diff --git a/Engine/Source/Developer/CookMetadata/Public/CookMetadata.h b/Engine/Source/Developer/CookMetadata/Public/CookMetadata.h
--- a/Engine/Source/Developer/CookMetadata/Public/CookMetadata.h
+++ b/Engine/Source/Developer/CookMetadata/Public/CookMetadata.h
@@ -13,6 +13,7 @@ namespace UE::Cook
 
 enum class ECookMetadataStateVersion : uint8
 {
+	InvalidVersion = 0,
 	PluginHierarchy = 1,
 	PostWritebackHash = 2,
 
@@ -173,13 +174,15 @@ class COOKMETADATA_API FCookMetadataState
 {
 public:
 	FCookMetadataState() = default;
-	FCookMetadataState(const FCookMetadataState&) = delete;
-	FCookMetadataState(FCookMetadataState&& Rhs) = delete;
+	FCookMetadataState(const FCookMetadataState&) = default;
+	FCookMetadataState(FCookMetadataState&& Rhs) = default;
 	~FCookMetadataState() = default;
 
-	FCookMetadataState& operator=(const FCookMetadataState&) = delete;
-	FCookMetadataState& operator=(FCookMetadataState&& O) = delete;
+	FCookMetadataState& operator=(const FCookMetadataState&) = default;
+	FCookMetadataState& operator=(FCookMetadataState&& O) = default;
 
+	bool IsValid() const { return Version == ECookMetadataStateVersion::LatestVersion; }
+	void Reset() { *this = FCookMetadataState(); }
 	
 	bool Serialize(FArchive& Ar);
 
@@ -228,7 +231,7 @@ public:
 	void SetSizesPresent(ECookMetadataSizesPresent InSizesPresent) { SizesPresent = InSizesPresent; }
 private:
 
-	ECookMetadataStateVersion Version;
+	ECookMetadataStateVersion Version = ECookMetadataStateVersion::InvalidVersion;
 	FCookMetadataPluginHierarchy PluginHierarchy;
 
 	uint64 AssociatedDevelopmentAssetRegistryHash = 0;
diff --git a/Engine/Source/Developer/IoStoreUtilities/Private/IoStoreUtilities.cpp b/Engine/Source/Developer/IoStoreUtilities/Private/IoStoreUtilities.cpp
--- a/Engine/Source/Developer/IoStoreUtilities/Private/IoStoreUtilities.cpp
+++ b/Engine/Source/Developer/IoStoreUtilities/Private/IoStoreUtilities.cpp
@@ -3711,13 +3711,20 @@ static ECookMetadataFiles FindAndLoadMetadataFiles(
 			else if (OutCookMetadata->GetAssociatedDevelopmentAssetRegistryHash() != LoadedDevArHash &&
 				OutCookMetadata->GetAssociatedDevelopmentAssetRegistryHashPostWriteback() != LoadedDevArHash) // during testing we can repeat stage after cook so we might have already edited it.
 			{
-				UE_LOG(LogIoStore, Error, 
-					TEXT("Cook metadata file mismatch: Hash of associated development asset registry does not match. [%s] %llx vs %llx (%llx post writeback)"), 
-					*CookMetadataFileName, LoadedDevArHash, OutCookMetadata->GetAssociatedDevelopmentAssetRegistryHash(), OutCookMetadata->GetAssociatedDevelopmentAssetRegistryHashPostWriteback());
 				if (EnumHasAnyFlags(InRequiredFiles, ECookMetadataFiles::CookMetadata))
 				{
+					UE_LOG(LogIoStore, Error,
+						TEXT("Cook metadata file mismatch: Hash of associated development asset registry does not match. [%s] %llx vs %llx (%llx post writeback)"),
+						*CookMetadataFileName, LoadedDevArHash, OutCookMetadata->GetAssociatedDevelopmentAssetRegistryHash(), OutCookMetadata->GetAssociatedDevelopmentAssetRegistryHashPostWriteback());
 					return ECookMetadataFiles::None;
 				}
+				else
+				{
+					UE_LOG(LogIoStore, Display,
+						TEXT("Cook metadata file mismatch: Hash of associated development asset registry does not match. [%s] %llx vs %llx (%llx post writeback)"),
+						*CookMetadataFileName, LoadedDevArHash, OutCookMetadata->GetAssociatedDevelopmentAssetRegistryHash(), OutCookMetadata->GetAssociatedDevelopmentAssetRegistryHashPostWriteback());
+					OutCookMetadata->Reset();
+				}
 			}
 			else
 			{
