Fix for goof in json format to address build warning

diff --git a/Engine/Source/Developer/IoStoreUtilities/Private/IoStoreUtilities.cpp b/Engine/Source/Developer/IoStoreUtilities/Private/IoStoreUtilities.cpp
--- a/Engine/Source/Developer/IoStoreUtilities/Private/IoStoreUtilities.cpp
+++ b/Engine/Source/Developer/IoStoreUtilities/Private/IoStoreUtilities.cpp
@@ -3376,9 +3376,9 @@ static void WritePluginMetadataJsons(const FString& InAssetRegistryFileName, TMa
 		OutPluginMetadataJson << "\t\"inclusive_installed\":" << InGraphEntry.InclusiveSize[UE::Cook::EPluginSizeTypes::Installed] << ",\n";
 		OutPluginMetadataJson << "\t\"inclusive_optional\":" << InGraphEntry.InclusiveSize[UE::Cook::EPluginSizeTypes::Optional] << ",\n";
 		OutPluginMetadataJson << "\t\"inclusive_ias\":" << InGraphEntry.InclusiveSize[UE::Cook::EPluginSizeTypes::Streaming] << ",\n";
-		OutPluginMetadataJson << "\t\"inclusive_optionalsegment\":" << InGraphEntry.InclusiveSize[UE::Cook::EPluginSizeTypes::OptionalSegment] << "\n";
+		OutPluginMetadataJson << "\t\"inclusive_optionalsegment\":" << InGraphEntry.InclusiveSize[UE::Cook::EPluginSizeTypes::OptionalSegment] << ",\n";
 
-		OutPluginMetadataJson << "\t\"direct_refcount\":" << InGraphEntry.DirectRefcount << ",\n";
+		OutPluginMetadataJson << "\t\"direct_refcount\":" << InGraphEntry.DirectRefcount << "\n";
 
 		OutPluginMetadataJson << "}\n";
 	};
