I/O - added on demand container stats to summary

#rnx

diff --git a/Engine/Source/Developer/IoStoreUtilities/Private/IoStoreUtilities.cpp b/Engine/Source/Developer/IoStoreUtilities/Private/IoStoreUtilities.cpp
--- a/Engine/Source/Developer/IoStoreUtilities/Private/IoStoreUtilities.cpp
+++ b/Engine/Source/Developer/IoStoreUtilities/Private/IoStoreUtilities.cpp
@@ -2657,15 +2657,23 @@ void InitializeContainerTargetsAndPackages(
 
 void LogWriterResults(const TArray<FIoStoreWriterResult>& Results)
 {
-	UE_LOG(LogIoStore, Display, TEXT("--------------------------------------------------- IoDispatcher --------------------------------------------------------"));
+	struct FContainerStats
+	{
+		uint64 TocCount = 0;
+		uint64 TocSize = 0;
+		uint64 UncompressedContainerSize = 0;
+		uint64 CompressedContainerSize = 0;
+		uint64 PaddingSize = 0;
+	};
+
+	UE_LOG(LogIoStore, Display, TEXT("------------------------------------------------ Container Summary ------------------------------------------------------"));
 	UE_LOG(LogIoStore, Display, TEXT(""));
 	UE_LOG(LogIoStore, Display, TEXT("%-30s %10s %15s %15s %15s %25s"),
 		TEXT("Container"), TEXT("Flags"), TEXT("TOC Size (KB)"), TEXT("TOC Entries"), TEXT("Size (MB)"), TEXT("Compressed (MB)"));
 	UE_LOG(LogIoStore, Display, TEXT("-------------------------------------------------------------------------------------------------------------------------"));
-	uint64 TotalTocSize = 0;
-	uint64 TotalTocEntryCount = 0;
-	uint64 TotalUncompressedContainerSize = 0;
-	uint64 TotalPaddingSize = 0;
+
+	FContainerStats TotalStats;
+	FContainerStats OnDemandStats;
 	for (const FIoStoreWriterResult& Result : Results)
 	{
 		FString CompressionInfo = TEXT("-");
@@ -2694,25 +2702,46 @@ void LogWriterResults(const TArray<FIoStoreWriterResult>& Results)
 			(double)Result.UncompressedContainerSize / 1024.0 / 1024.0,
 			*CompressionInfo);
 
+		if (EnumHasAnyFlags(Result.ContainerFlags, EIoContainerFlags::OnDemand))
+		{
+			OnDemandStats.TocCount += Result.TocEntryCount;
+			OnDemandStats.TocSize += Result.TocSize;
+			OnDemandStats.UncompressedContainerSize += Result.UncompressedContainerSize;
+			OnDemandStats.CompressedContainerSize += Result.CompressedContainerSize;
+		}
+
+		TotalStats.TocCount += Result.TocEntryCount;
+		TotalStats.TocSize += Result.TocSize;
+		TotalStats.UncompressedContainerSize += Result.UncompressedContainerSize;
+		TotalStats.CompressedContainerSize += Result.CompressedContainerSize;
+		TotalStats.PaddingSize += Result.PaddingSize;
+	}
+
+	UE_LOG(LogIoStore, Display, TEXT("-------------------------------------------------------------------------------------------------------------------------"));
 
-		TotalTocSize += Result.TocSize;
-		TotalTocEntryCount += Result.TocEntryCount;
-		TotalUncompressedContainerSize += Result.UncompressedContainerSize;
-		TotalPaddingSize += Result.PaddingSize;
+	if (OnDemandStats.TocCount > 0)
+	{
+		UE_LOG(LogIoStore, Display, TEXT("%-30s %10s %15.2lf %15llu %15.2lf %25.2lf"),
+			TEXT("Total On Demand"),
+			TEXT(""),
+			(double)OnDemandStats.TocSize / 1024.0,
+			OnDemandStats.TocCount,
+			(double)OnDemandStats.UncompressedContainerSize / 1024.0 / 1024.0,
+			(double)OnDemandStats.CompressedContainerSize / 1024.0 / 1024.0);
 	}
 
-	UE_LOG(LogIoStore, Display, TEXT("%-30s %10s %15.2lf %15llu %15.2lf %25s"),
-		TEXT("TOTAL"),
+	UE_LOG(LogIoStore, Display, TEXT("%-30s %10s %15.2lf %15llu %15.2lf %25.2lf "),
+		TEXT("Total"),
 		TEXT(""),
-		(double)TotalTocSize / 1024.0,
-		TotalTocEntryCount,
-		(double)TotalUncompressedContainerSize / 1024.0 / 1024.0,
-		TEXT("-"));
+		(double)TotalStats.TocSize / 1024.0,
+		TotalStats.TocCount,
+		(double)TotalStats.UncompressedContainerSize / 1024.0 / 1024.0,
+		(double)TotalStats.CompressedContainerSize / 1024.0 / 1024.0);
 
 	UE_LOG(LogIoStore, Display, TEXT(""));
 	UE_LOG(LogIoStore, Display, TEXT("** Flags: (C)ompressed / (E)ncrypted / (S)igned) / (I)ndexed) / (O)nDemand **"));
 	UE_LOG(LogIoStore, Display, TEXT(""));
-	UE_LOG(LogIoStore, Display, TEXT("Compression block padding: %8.2lf MB"), (double)TotalPaddingSize / 1024.0 / 1024.0);
+	UE_LOG(LogIoStore, Display, TEXT("Compression block padding: %8.2lf MB"), (double)TotalStats.PaddingSize / 1024.0 / 1024.0);
 	UE_LOG(LogIoStore, Display, TEXT(""));
 
 	UE_LOG(LogIoStore, Display, TEXT("-------------------------------------------- Container Directory Index --------------------------------------------------"));
