Misc CookMetadataFixes:
* Restaging a build no longer just constantly increases size
* Inclusive size actually looks at valid data
* Horde URL environment variable corrected.
* Add a new root pseudo plugin json event that holds the size of all plugins that don't belong to a root plugin.

diff --git a/Engine/Source/Developer/CookMetadata/Public/CookMetadata.h b/Engine/Source/Developer/CookMetadata/Public/CookMetadata.h
--- a/Engine/Source/Developer/CookMetadata/Public/CookMetadata.h
+++ b/Engine/Source/Developer/CookMetadata/Public/CookMetadata.h
@@ -56,6 +56,7 @@ struct FPluginSizeInfo
 {
 	uint64 Sizes[EPluginSizeTypesCount] = {};
 
+	void Zero() { FMemory::Memzero(this, sizeof(*this)); }
 	void AddSizes(uint64 SizesPerType[EPluginSizeTypesCount])
 	{
 		for (uint8 Type = 0; Type < EPluginSizeTypesCount; Type++)
diff --git a/Engine/Source/Developer/IoStoreUtilities/Private/IoStoreUtilities.cpp b/Engine/Source/Developer/IoStoreUtilities/Private/IoStoreUtilities.cpp
--- a/Engine/Source/Developer/IoStoreUtilities/Private/IoStoreUtilities.cpp
+++ b/Engine/Source/Developer/IoStoreUtilities/Private/IoStoreUtilities.cpp
@@ -3278,7 +3278,7 @@ static void WritePluginMetadataJsons(const FString& InAssetRegistryFileName, TMa
 		PluginGraph[PluginHierarchy.PluginsEnabledAtCook[RootIndex].Name].bIsRoot = true;
 	}
 
-
+	double AssetPackageMapStart = FPlatformTime::Seconds();
 	const TMap<FName, const FAssetPackageData*> AssetPackageMap = AssetRegistry.GetAssetPackageDataMap();
 	for (const TPair<FName, const FAssetPackageData*>& AssetPackage : AssetPackageMap)
 	{
@@ -3322,11 +3322,21 @@ static void WritePluginMetadataJsons(const FString& InAssetRegistryFileName, TMa
 	//
 	// Generate the inclusive size for each plugin.
 	//
+	{
+		// This is so re-staging the same cook is consistent.
+		UE::Cook::FCookMetadataPluginHierarchy& MutablePluginHierarchy = CookMetadata.GetMutablePluginHierarchy();
+		for (UE::Cook::FCookMetadataPluginEntry& Plugin : MutablePluginHierarchy.PluginsEnabledAtCook)
+		{
+			Plugin.InclusiveSizes.Zero();
+			Plugin.ExclusiveSizes.Zero();
+		}
+	}
 
 	// Sort the plugins topologically. This means that when we iterate linearly,
 	// we know that when we hit a plugin, we've already processed the dependencies.
 	// This takes some memory to track edges but is a depth first search
 	// and not anything quadratic or worse.
+	double TopologicalSortStart = FPlatformTime::Seconds();
 	TArray<const UE::Cook::FCookMetadataPluginEntry*> SortedList;
 	{
 		for (const UE::Cook::FCookMetadataPluginEntry& Plugin : PluginHierarchy.PluginsEnabledAtCook)
@@ -3366,26 +3376,63 @@ static void WritePluginMetadataJsons(const FString& InAssetRegistryFileName, TMa
 	{
 		FPluginGraphEntry& PluginEntry = PluginGraph[Plugin->Name];
 
+		PluginEntry.InclusiveSize = PluginEntry.ExclusiveSize;
+
 		for (uint16 DependencyIndex = Plugin->DependencyIndexStart; DependencyIndex < Plugin->DependencyIndexEnd; DependencyIndex++)
 		{
 			const UE::Cook::FCookMetadataPluginEntry& DependentPlugin = PluginHierarchy.PluginsEnabledAtCook[PluginHierarchy.PluginDependencies[DependencyIndex]];
 
-			PluginEntry.TotalDependencies.Add(&DependentPlugin);
+			bool bAlreadyInSet = false;
+			PluginEntry.TotalDependencies.Add(&DependentPlugin, &bAlreadyInSet);
+
+			if (bAlreadyInSet == false)
+			{
+				const FPluginGraphEntry* DependencyGraphEntry = &PluginGraph[DependentPlugin.Name];
+				PluginEntry.InclusiveSize.Add(DependencyGraphEntry->ExclusiveSize);
+			}
 			
 			FPluginGraphEntry& DependentEntry = PluginGraph[DependentPlugin.Name];
 			DependentEntry.DirectRefcount++;
 
-			PluginEntry.TotalDependencies.Append(DependentEntry.TotalDependencies);
+			// In the worse case this is another O(N) iteration, which makes us overall O(N^2)
+			for (const UE::Cook::FCookMetadataPluginEntry* TotalDependencyEntry : DependentEntry.TotalDependencies)
+			{
+				PluginEntry.TotalDependencies.Add(TotalDependencyEntry, &bAlreadyInSet);
+				if (bAlreadyInSet == false)
+				{
+					const FPluginGraphEntry* DependencyGraphEntry = &PluginGraph[TotalDependencyEntry->Name];
+					PluginEntry.InclusiveSize.Add(DependencyGraphEntry->ExclusiveSize);
+				}
+			}
 		}
+	}
+	double InclusiveComputeEnd = FPlatformTime::Seconds();
 
-		// In the worse case this is another O(N) iteration, which makes us overall O(N^2)
-		PluginEntry.InclusiveSize = PluginEntry.ExclusiveSize;
-		for (const UE::Cook::FCookMetadataPluginEntry* Dependency : PluginEntry.TotalDependencies)
+	// Find the total size of all plugins that aren't rooted in the root set.
+	UE::Cook::FPluginSizeInfo UnrootedTotal;
+	TSet<const UE::Cook::FCookMetadataPluginEntry*> UnrootedPlugins;
+	{		
+		for (const UE::Cook::FCookMetadataPluginEntry* Plugin : SortedList)
+		{
+			UnrootedPlugins.Add(Plugin);
+		}
+
+		for (uint16 RootIndex : PluginHierarchy.RootPlugins)
 		{
-			PluginEntry.InclusiveSize.Add(Dependency->ExclusiveSizes);
+			UnrootedPlugins.Remove(&PluginHierarchy.PluginsEnabledAtCook[RootIndex]);
+
+			FPluginGraphEntry& PluginEntry = PluginGraph[PluginHierarchy.PluginsEnabledAtCook[RootIndex].Name];
+			for (const UE::Cook::FCookMetadataPluginEntry* Plugin : PluginEntry.TotalDependencies)
+			{
+				UnrootedPlugins.Remove(Plugin);
+			}
+		}
+
+		for (const UE::Cook::FCookMetadataPluginEntry* Plugin : UnrootedPlugins)
+		{
+			UnrootedTotal.Add(PluginGraph[Plugin->Name].ExclusiveSize);
 		}
 	}
-	double InclusiveComputeEnd = FPlatformTime::Seconds();
 
 	auto GeneratePluginJson = [](TUtf8StringBuilder<4096>& OutPluginMetadataJson, FStringView InName, const FPluginGraphEntry& InGraphEntry)
 	{
@@ -3419,7 +3466,7 @@ static void WritePluginMetadataJsons(const FString& InAssetRegistryFileName, TMa
 
 	// Also write a csv for easier browsing in spreadsheets.
 	TUtf8StringBuilder<4096> Csv;
-	Csv.Append("name,exclusive_installed,exclusive_optional,inclusive_installed,inclusive_optional,direct_refcount\n");
+	Csv.Append("name,exclusive_installed,exclusive_optional,exclusive_ias,inclusive_installed,inclusive_optional,inclusive_ias,direct_refcount,total_dependency_count\n");
 
 	// We want to write the sizes back to the cook metadata, so we need a non-const version.
 	UE::Cook::FCookMetadataPluginHierarchy& MutablePluginHierarchy = CookMetadata.GetMutablePluginHierarchy();
@@ -3437,10 +3484,25 @@ static void WritePluginMetadataJsons(const FString& InAssetRegistryFileName, TMa
 		Plugin.InclusiveSizes = PluginEntry.InclusiveSize;
 		Plugin.ExclusiveSizes = PluginEntry.ExclusiveSize;
 
-		Csv.Appendf("%ls,%llu,%llu,%llu,%llu,%u\n", *Plugin.Name, 
-			Plugin.ExclusiveSizes[UE::Cook::EPluginSizeTypes::Installed], Plugin.ExclusiveSizes[UE::Cook::EPluginSizeTypes::Optional],
-			Plugin.InclusiveSizes[UE::Cook::EPluginSizeTypes::Installed], Plugin.InclusiveSizes[UE::Cook::EPluginSizeTypes::Optional],
-			PluginEntry.DirectRefcount);
+		Csv.Appendf("%ls,%llu,%llu,%llu,%llu,%llu,%llu,%u,%u\n", *Plugin.Name, 
+			Plugin.ExclusiveSizes[UE::Cook::EPluginSizeTypes::Installed], Plugin.ExclusiveSizes[UE::Cook::EPluginSizeTypes::Optional], Plugin.ExclusiveSizes[UE::Cook::EPluginSizeTypes::Streaming],
+			Plugin.InclusiveSizes[UE::Cook::EPluginSizeTypes::Installed], Plugin.InclusiveSizes[UE::Cook::EPluginSizeTypes::Optional], Plugin.InclusiveSizes[UE::Cook::EPluginSizeTypes::Streaming],
+			PluginEntry.DirectRefcount,PluginEntry.TotalDependencies.Num());
+	}
+
+	// Also write a json that contains the sizes for the plugins that don't belong to any root plugin,
+	// so that size information never gets lost.
+	{
+		FPluginGraphEntry OrphanedEntry;
+		OrphanedEntry.bIsRoot = true;
+		OrphanedEntry.DirectRefcount = 0;
+		OrphanedEntry.InclusiveSize = UnrootedTotal;
+		GeneratePluginJson(PluginMetadataJson, TEXT("OrphanedPlugins"), OrphanedEntry);
+		SavePluginMetadata(InAssetRegistryFileName, TEXT("OrphanedPlugins"), PluginMetadataJson);
+
+		Csv.Appendf("OrphanedPlugins,0,0,0,%llu,%llu,%llu,0,%u\n",
+			UnrootedTotal[UE::Cook::EPluginSizeTypes::Installed], UnrootedTotal[UE::Cook::EPluginSizeTypes::Optional], UnrootedTotal[UE::Cook::EPluginSizeTypes::Streaming],
+			UnrootedPlugins.Num());
 	}
 
 	{
@@ -3453,7 +3515,9 @@ static void WritePluginMetadataJsons(const FString& InAssetRegistryFileName, TMa
 	}
 
 	double WritePluginEnd = FPlatformTime::Seconds();
-	UE_LOG(LogIoStore, Display, TEXT("Wrote plugin size jsons/csv in %.2f seconds, inclusive computation was %.2f of that"), WritePluginEnd - WritePluginStart, InclusiveComputeEnd - InclusiveComputeStart);
+	UE_LOG(LogIoStore, Display, TEXT("Wrote plugin size jsons/csv in %.2f seconds [inclusive computation: %.2fs; plugin graph: %.2fs; topological sort: %.2fs; package mapping: %.2fs]"), 
+		WritePluginEnd - WritePluginStart, InclusiveComputeEnd - InclusiveComputeStart, AssetPackageMapStart - WritePluginStart, 
+		InclusiveComputeStart - TopologicalSortStart, TopologicalSortStart - AssetPackageMapStart);
 }
 
 
