ZenWriter: Log the fatal assertion about invalid CommitPackage call on the CommitPackage thread in addition to the async thread so it is easier to debug.
#rnx

diff --git a/Engine/Source/Developer/IoStoreUtilities/Private/ZenStoreWriter.cpp b/Engine/Source/Developer/IoStoreUtilities/Private/ZenStoreWriter.cpp
--- a/Engine/Source/Developer/IoStoreUtilities/Private/ZenStoreWriter.cpp
+++ b/Engine/Source/Developer/IoStoreUtilities/Private/ZenStoreWriter.cpp
@@ -778,6 +778,8 @@ void FZenStoreWriter::CommitPackage(FCommitPackageInfo&& Info)
 	}
 
 	TUniquePtr<FPendingPackageState> PackageState = RemovePendingPackage(Info.PackageName);
+	checkf(Info.Status != ECommitStatus::Success || !PackageState->PackageData.IsEmpty(),
+		TEXT("CommitPackage called with CommitStatus::Success but without first calling WritePackageData"));
 	FZenCommitInfo ZenCommitInfo{ Forward<FCommitPackageInfo>(Info), MoveTemp(PackageState) };
 	if (FPlatformProcess::SupportsMultithreading())
 	{
