Cooker: ZenStore: Downgrade ZenFileSystemManifest's per-plugin log messages to verbose to make it easier to read other end of cook logs.
#rnx

diff --git a/Engine/Source/Developer/IoStoreUtilities/Private/ZenFileSystemManifest.cpp b/Engine/Source/Developer/IoStoreUtilities/Private/ZenFileSystemManifest.cpp
--- a/Engine/Source/Developer/IoStoreUtilities/Private/ZenFileSystemManifest.cpp
+++ b/Engine/Source/Developer/IoStoreUtilities/Private/ZenFileSystemManifest.cpp
@@ -323,16 +323,16 @@ int32 FZenFileSystemManifest::Generate()
 		FString ProjectName = Plugin->GetName();
 		if (FilterDisabledPlugins && !Plugin->IsEnabled())
 		{
-			UE_LOG(LogZenFileSystemManifest, Display, TEXT("Plugin '%s' disabled, skipping"), *ProjectName);
+			UE_LOG(LogZenFileSystemManifest, Verbose, TEXT("Plugin '%s' disabled, skipping"), *ProjectName);
 			continue;
 		}
 		const FPluginDescriptor& Descriptor = Plugin->GetDescriptor();
 		if (!Descriptor.SupportsTargetPlatform(PluginTargetPlatformString))
 		{
-			UE_LOG(LogZenFileSystemManifest, Log, TEXT("Plugin '%s' not supported on platform '%s', skipping"), *ProjectName, *TargetPlatform.PlatformName());
+			UE_LOG(LogZenFileSystemManifest, Verbose, TEXT("Plugin '%s' not supported on platform '%s', skipping"), *ProjectName, *TargetPlatform.PlatformName());
 			for (const auto& SupportedTargetPlatform : Descriptor.SupportedTargetPlatforms)
 			{
-				UE_LOG(LogZenFileSystemManifest, Log, TEXT("       '%s' supports platform '%s'"), *ProjectName, *SupportedTargetPlatform);
+				UE_LOG(LogZenFileSystemManifest, Verbose, TEXT("       '%s' supports platform '%s'"), *ProjectName, *SupportedTargetPlatform);
 			}
 			continue;
 		}
@@ -341,7 +341,7 @@ int32 FZenFileSystemManifest::Generate()
 		FString ContentDir = Plugin->GetContentDir();
 		FString LocalizationDir = ContentDir / TEXT("Localization");
 		FString ConfigDir = BaseDir / TEXT("Config");
-		UE_LOG(LogZenFileSystemManifest, Display, TEXT("Plugin '%s': BaseDir: '%s'"), *ProjectName, *BaseDir);
+		UE_LOG(LogZenFileSystemManifest, Verbose, TEXT("Plugin '%s': BaseDir: '%s'"), *ProjectName, *BaseDir);
 
 		FString ClientDirectory;
 		FString SourcePath;
