I/O - minor change to summary stats

#rnx

diff --git a/Engine/Source/Developer/IoStoreUtilities/Private/IoStoreUtilities.cpp b/Engine/Source/Developer/IoStoreUtilities/Private/IoStoreUtilities.cpp
--- a/Engine/Source/Developer/IoStoreUtilities/Private/IoStoreUtilities.cpp
+++ b/Engine/Source/Developer/IoStoreUtilities/Private/IoStoreUtilities.cpp
@@ -2666,11 +2666,13 @@ void LogWriterResults(const TArray<FIoStoreWriterResult>& Results)
 		uint64 PaddingSize = 0;
 	};
 
-	UE_LOG(LogIoStore, Display, TEXT("------------------------------------------------ Container Summary ------------------------------------------------------"));
 	UE_LOG(LogIoStore, Display, TEXT(""));
-	UE_LOG(LogIoStore, Display, TEXT("%-30s %10s %15s %15s %15s %25s"),
-		TEXT("Container"), TEXT("Flags"), TEXT("TOC Size (KB)"), TEXT("TOC Entries"), TEXT("Size (MB)"), TEXT("Compressed (MB)"));
-	UE_LOG(LogIoStore, Display, TEXT("-------------------------------------------------------------------------------------------------------------------------"));
+	UE_LOG(LogIoStore, Display, TEXT("Container Summary"));
+	UE_LOG(LogIoStore, Display, TEXT("=================="));
+	UE_LOG(LogIoStore, Display, TEXT(""));
+	UE_LOG(LogIoStore, Display, TEXT("%-50s %10s %15s %15s %20s %20s"),
+		TEXT("Container"), TEXT("Flags"), TEXT("Chunk(s) #"), TEXT("TOC (KiB)"), TEXT("Raw Size (MiB)"), TEXT("Size (MiB)"));
+	UE_LOG(LogIoStore, Display, TEXT("----------------------------------------------------------------------------------------------------------------------------------------------"));
 
 	FContainerStats TotalStats;
 	FContainerStats OnDemandStats;
@@ -2680,9 +2682,8 @@ void LogWriterResults(const TArray<FIoStoreWriterResult>& Results)
 
 		if (Result.CompressionMethod != NAME_None)
 		{
-			double Procentage = (double(Result.UncompressedContainerSize - Result.CompressedContainerSize) / double(Result.UncompressedContainerSize)) * 100.0;
-			CompressionInfo = FString::Printf(TEXT("%.2lf (%.2lf%% %s)"),
-				(double)Result.CompressedContainerSize / 1024.0 / 1024.0,
+			const double Procentage = (double(Result.UncompressedContainerSize - Result.CompressedContainerSize) / double(Result.UncompressedContainerSize)) * 100.0;
+			CompressionInfo = FString::Printf(TEXT("(%.2lf%% %s)"),
 				Procentage,
 				*Result.CompressionMethod.ToString());
 		}
@@ -2694,12 +2695,13 @@ void LogWriterResults(const TArray<FIoStoreWriterResult>& Results)
 			EnumHasAnyFlags(Result.ContainerFlags, EIoContainerFlags::Indexed) ? TEXT("I") : TEXT("-"),
 			EnumHasAnyFlags(Result.ContainerFlags, EIoContainerFlags::OnDemand) ? TEXT("O") : TEXT("-"));
 
-		UE_LOG(LogIoStore, Display, TEXT("%-30s %10s %15.2lf %15llu %15.2lf %25s"),
+		UE_LOG(LogIoStore, Display, TEXT("%-50s %10s %15llu %15.2lf %20.2lf %20.2lf %s"),
 			*Result.ContainerName,
 			*ContainerSettings,
-			(double)Result.TocSize / 1024.0,
 			Result.TocEntryCount,
+			(double)Result.TocSize / 1024.0,
 			(double)Result.UncompressedContainerSize / 1024.0 / 1024.0,
+			(double)Result.CompressedContainerSize / 1024.0 / 1024.0,
 			*CompressionInfo);
 
 		if (EnumHasAnyFlags(Result.ContainerFlags, EIoContainerFlags::OnDemand))
@@ -2716,47 +2718,53 @@ void LogWriterResults(const TArray<FIoStoreWriterResult>& Results)
 		TotalStats.CompressedContainerSize += Result.CompressedContainerSize;
 		TotalStats.PaddingSize += Result.PaddingSize;
 	}
+	UE_LOG(LogIoStore, Display, TEXT("----------------------------------------------------------------------------------------------------------------------------------------------"));
 
-	UE_LOG(LogIoStore, Display, TEXT("-------------------------------------------------------------------------------------------------------------------------"));
 
 	if (OnDemandStats.TocCount > 0)
 	{
-		UE_LOG(LogIoStore, Display, TEXT("%-30s %10s %15.2lf %15llu %15.2lf %25.2lf"),
+		UE_LOG(LogIoStore, Display, TEXT("%-50s %10s %15llu %15.2lf %20.2lf %20.2lf"),
 			TEXT("Total On Demand"),
 			TEXT(""),
-			(double)OnDemandStats.TocSize / 1024.0,
 			OnDemandStats.TocCount,
+			(double)OnDemandStats.TocSize / 1024.0,
 			(double)OnDemandStats.UncompressedContainerSize / 1024.0 / 1024.0,
 			(double)OnDemandStats.CompressedContainerSize / 1024.0 / 1024.0);
 	}
 
-	UE_LOG(LogIoStore, Display, TEXT("%-30s %10s %15.2lf %15llu %15.2lf %25.2lf "),
+	UE_LOG(LogIoStore, Display, TEXT("%-50s %10s %15llu %15.2lf %20.2lf %20.2lf"),
 		TEXT("Total"),
 		TEXT(""),
-		(double)TotalStats.TocSize / 1024.0,
 		TotalStats.TocCount,
+		(double)TotalStats.TocSize / 1024.0,
 		(double)TotalStats.UncompressedContainerSize / 1024.0 / 1024.0,
 		(double)TotalStats.CompressedContainerSize / 1024.0 / 1024.0);
 
 	UE_LOG(LogIoStore, Display, TEXT(""));
 	UE_LOG(LogIoStore, Display, TEXT("** Flags: (C)ompressed / (E)ncrypted / (S)igned) / (I)ndexed) / (O)nDemand **"));
 	UE_LOG(LogIoStore, Display, TEXT(""));
-	UE_LOG(LogIoStore, Display, TEXT("Compression block padding: %8.2lf MB"), (double)TotalStats.PaddingSize / 1024.0 / 1024.0);
+	UE_LOG(LogIoStore, Display, TEXT("Compression block padding: %8.2lf MiB"), (double)TotalStats.PaddingSize / 1024.0 / 1024.0);
 	UE_LOG(LogIoStore, Display, TEXT(""));
 
-	UE_LOG(LogIoStore, Display, TEXT("-------------------------------------------- Container Directory Index --------------------------------------------------"));
-	UE_LOG(LogIoStore, Display, TEXT("%-30s %15s"), TEXT("Container"), TEXT("Size (KB)"));
+	UE_LOG(LogIoStore, Display, TEXT("Container Directory Index"));
+	UE_LOG(LogIoStore, Display, TEXT("=========================="));
+	UE_LOG(LogIoStore, Display, TEXT(""));
+	UE_LOG(LogIoStore, Display, TEXT("%-45s %15s"), TEXT("Container"), TEXT("Size (KiB)"));
+	UE_LOG(LogIoStore, Display, TEXT("----------------------------------------------------------------------------------------------------------------------------------------------"));
 	for (const FIoStoreWriterResult& Result : Results)
 	{
-		UE_LOG(LogIoStore, Display, TEXT("%-30s %15.2lf"), *Result.ContainerName, double(Result.DirectoryIndexSize) / 1024.0);
+		UE_LOG(LogIoStore, Display, TEXT("%-45s %15.2lf"), *Result.ContainerName, double(Result.DirectoryIndexSize) / 1024.0);
 	}
 
 	UE_LOG(LogIoStore, Display, TEXT(""));
-	UE_LOG(LogIoStore, Display, TEXT("---------------------------------------------- Container Patch Report ---------------------------------------------------"));
-	UE_LOG(LogIoStore, Display, TEXT("%-30s %16s %16s %16s %16s %16s"), TEXT("Container"), TEXT("Total (count)"), TEXT("Modified (count)"), TEXT("Added (count)"), TEXT("Modified (MB)"), TEXT("Added (MB)"));
+	UE_LOG(LogIoStore, Display, TEXT("Container Patch Report"));
+	UE_LOG(LogIoStore, Display, TEXT("========================"));
+	UE_LOG(LogIoStore, Display, TEXT(""));
+	UE_LOG(LogIoStore, Display, TEXT("%-44s %16s %16s %16s %16s %16s"), TEXT("Container"), TEXT("Total #"), TEXT("Modified #"), TEXT("Added #"), TEXT("Modified (MiB)"), TEXT("Added (MiB)"));
+	UE_LOG(LogIoStore, Display, TEXT("----------------------------------------------------------------------------------------------------------------------------------------------"));
 	for (const FIoStoreWriterResult& Result : Results)
 	{
-		UE_LOG(LogIoStore, Display, TEXT("%-30s %16d %16d %16d %16.2lf %16.2lf"), *Result.ContainerName, Result.TocEntryCount, Result.ModifiedChunksCount, Result.AddedChunksCount, Result.ModifiedChunksSize / 1024.0 / 1024.0, Result.AddedChunksSize / 1024.0 / 1024.0);
+		UE_LOG(LogIoStore, Display, TEXT("%-44s %16d %16d %16d %16.2lf %16.2lf"), *Result.ContainerName, Result.TocEntryCount, Result.ModifiedChunksCount, Result.AddedChunksCount, Result.ModifiedChunksSize / 1024.0 / 1024.0, Result.AddedChunksSize / 1024.0 / 1024.0);
 	}
 }
 
@@ -2767,15 +2775,15 @@ void LogContainerPackageInfo(const TArray<FContainerTargetSpec*>& ContainerTarge
 	uint64 TotalLocalizedPackageCount = 0;
 
 	UE_LOG(LogIoStore, Display, TEXT(""));
+	UE_LOG(LogIoStore, Display, TEXT("PackageStore"));
+	UE_LOG(LogIoStore, Display, TEXT("============="));
 	UE_LOG(LogIoStore, Display, TEXT(""));
-	UE_LOG(LogIoStore, Display, TEXT("--------------------------------------------------- PackageStore (KB) ---------------------------------------------------"));
-	UE_LOG(LogIoStore, Display, TEXT(""));
-	UE_LOG(LogIoStore, Display, TEXT("%-30s %20s %20s %20s"),
+	UE_LOG(LogIoStore, Display, TEXT("%-45s %15s %15s %15s"),
 		TEXT("Container"),
-		TEXT("Store Size"),
-		TEXT("Packages"),
-		TEXT("Localized"));
-	UE_LOG(LogIoStore, Display, TEXT("-------------------------------------------------------------------------------------------------------------------------"));
+		TEXT("Size (KiB)"),
+		TEXT("Packages #"),
+		TEXT("Localized #"));
+	UE_LOG(LogIoStore, Display, TEXT("----------------------------------------------------------------------------------------------------------------------------------------------"));
 
 	for (const FContainerTargetSpec* ContainerTarget : ContainerTargets)
 	{
@@ -2783,7 +2791,7 @@ void LogContainerPackageInfo(const TArray<FContainerTargetSpec*>& ContainerTarge
 		uint64 PackageCount = ContainerTarget->Packages.Num();
 		uint64 LocalizedPackageCount = ContainerTarget->Header.LocalizedPackages.Num();
 
-		UE_LOG(LogIoStore, Display, TEXT("%-30s %20.0lf %20llu %20llu"),
+		UE_LOG(LogIoStore, Display, TEXT("%-45s %15.0lf %15llu %15llu"),
 			*ContainerTarget->Name.ToString(),
 			(double)StoreSize / 1024.0,
 			PackageCount,
@@ -2793,8 +2801,9 @@ void LogContainerPackageInfo(const TArray<FContainerTargetSpec*>& ContainerTarge
 		TotalPackageCount += PackageCount;
 		TotalLocalizedPackageCount += LocalizedPackageCount;
 	}
-	UE_LOG(LogIoStore, Display, TEXT("%-30s %20.0lf %20llu %20llu"),
-		TEXT("TOTAL"),
+	UE_LOG(LogIoStore, Display, TEXT("----------------------------------------------------------------------------------------------------------------------------------------------"));
+	UE_LOG(LogIoStore, Display, TEXT("%-45s %15.0lf %15llu %15llu"),
+		TEXT("Total"),
 		(double)TotalStoreSize / 1024.0,
 		TotalPackageCount,
 		TotalLocalizedPackageCount);
