#UE Only assert that you have package data to write when you are intending to write data or computehash. The new SkipSave cook option does not write data intentionally

diff --git a/Engine/Source/Developer/IoStoreUtilities/Private/ZenStoreWriter.cpp b/Engine/Source/Developer/IoStoreUtilities/Private/ZenStoreWriter.cpp
--- a/Engine/Source/Developer/IoStoreUtilities/Private/ZenStoreWriter.cpp
+++ b/Engine/Source/Developer/IoStoreUtilities/Private/ZenStoreWriter.cpp
@@ -778,8 +778,11 @@ void FZenStoreWriter::CommitPackage(FCommitPackageInfo&& Info)
 	}
 
 	TUniquePtr<FPendingPackageState> PackageState = RemovePendingPackage(Info.PackageName);
-	checkf(Info.Status != ECommitStatus::Success || !PackageState->PackageData.IsEmpty(),
-		TEXT("CommitPackage called with CommitStatus::Success but without first calling WritePackageData"));
+	if (EnumHasAnyFlags(Info.WriteOptions, IPackageWriter::EWriteOptions::Write | IPackageWriter::EWriteOptions::ComputeHash))
+	{
+		checkf(Info.Status != ECommitStatus::Success || !PackageState->PackageData.IsEmpty(),
+			TEXT("CommitPackage called with CommitStatus::Success but without first calling WritePackageData"));
+	}
 	FZenCommitInfo ZenCommitInfo{ Forward<FCommitPackageInfo>(Info), MoveTemp(PackageState) };
 	if (FPlatformProcess::SupportsMultithreading())
 	{
