When writing file manifest to zenstore, don't attempt to set both the "data" field on a file entry and the "serverpath" field on the file entry, as zenserver will treat the two as mutually exclusive, and setting the "data" field will mean we no loger get "serverpath" out of zenserver.  This causes issues when doing an iterative cook and trying to read the file manifest back from zenserver.

diff --git a/Engine/Source/Developer/IoStoreUtilities/Private/ZenStoreWriter.cpp b/Engine/Source/Developer/IoStoreUtilities/Private/ZenStoreWriter.cpp
--- a/Engine/Source/Developer/IoStoreUtilities/Private/ZenStoreWriter.cpp
+++ b/Engine/Source/Developer/IoStoreUtilities/Private/ZenStoreWriter.cpp
@@ -988,7 +988,10 @@ void FZenStoreWriter::CommitPackageInternal(FZenCommitInfo&& ZenCommitInfo)
 
 				OplogEntryDesc.BeginObject();
 				OplogEntryDesc << "id" << ToObjectId(File.Info.ChunkId);
-				OplogEntryDesc << "data" << FileDataAttachment;
+				// ZenServer treats the hash stored in "data" as mutually exlusive with the string stored in "serverpath".
+				// We must write data as a zero hash (or exclude it entirely) if we want to be able to get the serverpath from ZenServer later.
+				// This is relevant to iterative cooks which will obtain the filesystem manifest contents from ZenServer.
+				OplogEntryDesc << "data" << FIoHash::Zero;
 				OplogEntryDesc << "serverpath" << File.ZenManifestServerPath;
 				OplogEntryDesc << "clientpath" << File.ZenManifestClientPath;
 				OplogEntryDesc.EndObject();
