Fix for asset registry writeback hash not matching when plugin size writeback is disabled.

diff --git a/Engine/Source/Developer/IoStoreUtilities/Private/IoStoreUtilities.cpp b/Engine/Source/Developer/IoStoreUtilities/Private/IoStoreUtilities.cpp
--- a/Engine/Source/Developer/IoStoreUtilities/Private/IoStoreUtilities.cpp
+++ b/Engine/Source/Developer/IoStoreUtilities/Private/IoStoreUtilities.cpp
@@ -3882,10 +3882,10 @@ static bool DoAssetRegistryWritebackDuringStage(
 	FString CookMetadataFileName;
 	UE::Cook::FCookMetadataState CookMetadata;
 	ECookMetadataFiles FilesNeeded = ECookMetadataFiles::AssetRegistry;
-	if (bInWritePluginMetadata)
-	{
-		EnumAddFlags(FilesNeeded, ECookMetadataFiles::CookMetadata);
-	}
+
+	// We always need the cook metadata in order to update the corresponding hash.
+	EnumAddFlags(FilesNeeded, ECookMetadataFiles::CookMetadata);
+
 	if (FindAndLoadMetadataFiles(InCookedDir, FilesNeeded, AssetRegistry, &AssetRegistryFileName, &CookMetadata, &CookMetadataFileName) == ECookMetadataFiles::None)
 	{
 		// already logged
@@ -3927,7 +3927,7 @@ static bool DoAssetRegistryWritebackDuringStage(
 	case EAssetRegistryWritebackMethod::OriginalFile:
 		{
 			// Write to an adjacent file and move after
-			if (SaveAssetRegistry(AssetRegistryFileName, AssetRegistry, bInWritePluginMetadata ? &UpdatedDevArHash : nullptr) == false)
+			if (SaveAssetRegistry(AssetRegistryFileName, AssetRegistry, &UpdatedDevArHash) == false)
 			{
 				return false;
 			}
@@ -3936,7 +3936,7 @@ static bool DoAssetRegistryWritebackDuringStage(
 		}
 	case EAssetRegistryWritebackMethod::AdjacentFile:
 		{
-			if (SaveAssetRegistry(AssetRegistryFileName.Replace(TEXT(".bin"), TEXT("Staged.bin")), AssetRegistry, bInWritePluginMetadata ? &UpdatedDevArHash : nullptr) == false)
+			if (SaveAssetRegistry(AssetRegistryFileName.Replace(TEXT(".bin"), TEXT("Staged.bin")), AssetRegistry, &UpdatedDevArHash) == false)
 			{
 				return false;
 			}
@@ -3950,30 +3950,27 @@ static bool DoAssetRegistryWritebackDuringStage(
 	}
 
 	// Since we modified the dev ar, we need to save the updated hash in the cook metadata so it can still validate.
-	if (bInWritePluginMetadata)
-	{
-		CookMetadata.SetSizesPresent(bInCompressionEnabled ? UE::Cook::ECookMetadataSizesPresent::Compressed : UE::Cook::ECookMetadataSizesPresent::Uncompressed);
-		CookMetadata.SetAssociatedDevelopmentAssetRegistryHashPostWriteback(UpdatedDevArHash);
+	CookMetadata.SetSizesPresent(bInCompressionEnabled ? UE::Cook::ECookMetadataSizesPresent::Compressed : UE::Cook::ECookMetadataSizesPresent::Uncompressed);
+	CookMetadata.SetAssociatedDevelopmentAssetRegistryHashPostWriteback(UpdatedDevArHash);
 
-		FArrayWriter SerializedCookMetadata;
-		CookMetadata.Serialize(SerializedCookMetadata);
+	FArrayWriter SerializedCookMetadata;
+	CookMetadata.Serialize(SerializedCookMetadata);
 
-		FString TempFileName = CookMetadataFileName + TEXT(".temp");
-		if (FFileHelper::SaveArrayToFile(SerializedCookMetadata, *TempFileName))
-		{
-			// Move our temp file over the original asset registry.
-			if (IFileManager::Get().Move(*CookMetadataFileName, *TempFileName) == false)
-			{
-				// Error already logged by FileManager
-				return false;
-			}
-		}
-		else
+	FString TempFileName = CookMetadataFileName + TEXT(".temp");
+	if (FFileHelper::SaveArrayToFile(SerializedCookMetadata, *TempFileName))
+	{
+		// Move our temp file over the original asset registry.
+		if (IFileManager::Get().Move(*CookMetadataFileName, *TempFileName) == false)
 		{
-			UE_LOG(LogIoStore, Error, TEXT("Failed to save temp file for write updated cook metadata file (%s"), *TempFileName);
+			// Error already logged by FileManager
 			return false;
 		}
 	}
+	else
+	{
+		UE_LOG(LogIoStore, Error, TEXT("Failed to save temp file for write updated cook metadata file (%s"), *TempFileName);
+		return false;
+	}
 	
 	return true;
 }
