* Add a system for routing plugin custom fields in to the cook metadata file.
* Use that system to upload whether a plugin is installed or not to analytics.
* Compute "unique" size for root plugins to more accurately represent the incremental cost of a plugin in analytics.
* Provide sizes for key asset types in addition to tracking total cost of a plugin.

diff --git a/Engine/Source/Developer/CookMetadata/Public/CookMetadata.h b/Engine/Source/Developer/CookMetadata/Public/CookMetadata.h
--- a/Engine/Source/Developer/CookMetadata/Public/CookMetadata.h
+++ b/Engine/Source/Developer/CookMetadata/Public/CookMetadata.h
@@ -3,6 +3,7 @@
 #pragma once
 
 #include "Containers/Array.h"
+#include "Containers/Map.h"
 #include "Containers/UnrealString.h"
 #include "CoreTypes.h"
 #include "Memory/MemoryFwd.h"
@@ -17,6 +18,7 @@ enum class ECookMetadataStateVersion : uint8
 	PluginHierarchy = 1,
 	PostWritebackHash = 2,
 	FixSerialization = 3,
+	AddedCustomFields = 4,
 
 	// Add new versions above this.
 	VersionCount,
@@ -118,6 +120,12 @@ struct COOKMETADATA_API FCookMetadataPluginEntry
 {
 	FString Name;
 
+	// These contain values pulled from the uplugin json file and hold fields that are not
+	// part of the engine FPluginDescriptor. They are for per-project values. The keys for the maps
+	// are indices in to FCookMetadataPluginHierarchy::CustomFieldNames. See the comment for CustomFieldNames.
+	TMap<uint8, bool> CustomBoolFields;
+	TMap<uint8, FString> CustomStringFields;
+
 	// The dependencies are stored in the FCookMetadataPluginHierarchy::PluginDependencies array,
 	// and this is an index into it. From there you can get a further index into PluginsEnabledAtCook
 	// to get the plugin information.
@@ -144,6 +152,7 @@ struct COOKMETADATA_API FCookMetadataPluginEntry
 	{
 		Ar << Entry.Name << Entry.DependencyIndexStart << Entry.DependencyIndexEnd;
 		Ar << Entry.InclusiveSizes << Entry.ExclusiveSizes;
+		Ar << Entry.CustomBoolFields << Entry.CustomStringFields;
 		return Ar;
 	}
 };
@@ -162,9 +171,44 @@ struct COOKMETADATA_API FCookMetadataPluginHierarchy
 	// The list of root plugins for the project as defined by the Editor.ini file.
 	TArray<uint16> RootPlugins;
 
+	// The list of custom field names. The values stored on an entry index in to this for the name.
+	// These values are copied from the plugins' json descriptor to allow for carrying through project
+	// specific data in to the metadata.
+	//
+	// To enable custom field replication, add the CookMetadataCustomPluginFields section to the relevant Editor.ini file
+	// and add to the follow arrays, as desired:
+	//
+	//  BoolFields, StringFields, PerPlatformBoolFields, PerPlatformStringFields
+	// 
+	// Bool and String refer to the json type, and will emit as bool or FString here. The PerPlatform arrays
+	// will check in the uplugin json for a field named PerPlatform<ArrayValue>. That field is expected to be
+	// an array of override objects. Each object is required to have the platform name and override value. E.g.:
+	//
+	//	[CookMetadataCustomPluginFields]
+	//	+PerPlatformBoolFields = ExampleProjectBool
+	//	+StringField = ExampleProjectString
+	//
+	// Will look in the uplugin file for this:
+	//
+	// 	"ExampleProjectBool": true
+	//	"PerPlatformExampleProjectBool": [
+	//		{
+	//			"Platform": "Windows",
+	//			"Value": false
+	//		}
+	//	]
+	//	"ExampleProjectString": "AProjectString"
+	//
+	// Note that in the per platform case, the base value is not required (i.e. ExampleProjectBool above), however it
+	// provides a default value for unlisted platforms. If it does not exist, false is used for bool fields, and an empty string
+	// for string fields.
+	TArray<FString> CustomFieldNames;
+
 	friend FArchive& operator<<(FArchive& Ar, FCookMetadataPluginHierarchy& Hierarchy)
 	{
-		return Ar << Hierarchy.PluginsEnabledAtCook << Hierarchy.PluginDependencies << Hierarchy.RootPlugins;
+		Ar << Hierarchy.PluginsEnabledAtCook << Hierarchy.PluginDependencies;
+		Ar << Hierarchy.RootPlugins << Hierarchy.CustomFieldNames;
+		return Ar;
 	}
 };
 
diff --git a/Engine/Source/Developer/IoStoreUtilities/Private/IoStoreUtilities.cpp b/Engine/Source/Developer/IoStoreUtilities/Private/IoStoreUtilities.cpp
--- a/Engine/Source/Developer/IoStoreUtilities/Private/IoStoreUtilities.cpp
+++ b/Engine/Source/Developer/IoStoreUtilities/Private/IoStoreUtilities.cpp
@@ -3220,10 +3220,30 @@ static bool WriteUtf8StringView(FUtf8StringView InView, const FString& InFilenam
 	return true;
 }
 
+// When we emit plugin size information, we emit one json for each class of sizes we
+// are monitoring.
+enum class EPluginGraphSizeClass : uint8
+{
+	All,
+	Texture,
+	StaticMesh,
+	SoundWave,
+	COUNT
+};
+
+static const UTF8CHAR* PluginGraphEntryClassNames[] = 
+{
+	UTF8TEXT("all"),
+	UTF8TEXT("texture"),
+	UTF8TEXT("staticmesh"),
+	UTF8TEXT("soundwave")
+};
 
-static bool SavePluginMetadata(const FString& InAssetRegistryFileName, const FString& PluginName, TUtf8StringBuilder<4096>& InPluginMetadataJson)
+static_assert( UE_ARRAY_COUNT(PluginGraphEntryClassNames) == (size_t)EPluginGraphSizeClass::COUNT, "Must have a name for each plugin graph size class!");
+
+static bool SavePluginMetadata(const FString& InAssetRegistryFileName, const FString& PluginName, TUtf8StringBuilder<4096>& InPluginMetadataJson, EPluginGraphSizeClass InAssetClass)
 {
-	FString PluginMetadataFilename = FPaths::GetPath(InAssetRegistryFileName) / TEXT("PluginJsons") / PluginName + TEXT(".json");
+	FString PluginMetadataFilename = FPaths::GetPath(InAssetRegistryFileName) / TEXT("PluginJsons") / PluginName + TEXT("_") + PluginGraphEntryClassNames[(uint8)InAssetClass] + TEXT(".json");
 	if (WriteUtf8StringView(InPluginMetadataJson.ToView(), PluginMetadataFilename) == false)
 	{
 		UE_LOG(LogIoStore, Error, TEXT("Unable to write plugin metadata file: %s"), *PluginMetadataFilename);
@@ -3232,15 +3252,23 @@ static bool SavePluginMetadata(const FString& InAssetRegistryFileName, const FSt
 	return true;
 }
 
+
 struct FPluginGraphEntry
 {
 	uint16 IndexInEnabledPlugins = 0;
+	const UE::Cook::FCookMetadataPluginEntry* Self = nullptr;
 	TArray<const UE::Cook::FCookMetadataPluginEntry*> Dependencies;
 	TSet<const UE::Cook::FCookMetadataPluginEntry*> TotalDependencies;
 	uint32 DirectRefcount = 0;
 	bool bIsRoot = false;
-	UE::Cook::FPluginSizeInfo ExclusiveSize;
-	UE::Cook::FPluginSizeInfo InclusiveSize;
+
+	// The list of root plugins this plugin can trace a route from
+	TArray<const UE::Cook::FCookMetadataPluginEntry*> Roots;
+
+	static constexpr uint8 ClassCount = (uint8)EPluginGraphSizeClass::COUNT;
+	UE::Cook::FPluginSizeInfo ExclusiveSizes[ClassCount];
+	UE::Cook::FPluginSizeInfo InclusiveSizes[ClassCount];
+	UE::Cook::FPluginSizeInfo UniqueSizes[ClassCount];
 };
 
 // Rework the hierarchy in to a graph where we have output edges resolved to pointers so we can pass to
@@ -3252,6 +3280,7 @@ static void GeneratePluginGraph(const UE::Cook::FCookMetadataPluginHierarchy& In
 	{
 		FPluginGraphEntry& OurEntry = OutPluginGraph.FindOrAdd(Plugin.Name);
 		OurEntry.IndexInEnabledPlugins = PluginIndex;
+		OurEntry.Self = &Plugin;
 
 		for (uint16 DependencyIndex = Plugin.DependencyIndexStart; DependencyIndex < Plugin.DependencyIndexEnd; DependencyIndex++)
 		{
@@ -3287,6 +3316,14 @@ static void WritePluginMetadataJsons(const FString& InAssetRegistryFileName, TMa
 		PluginGraph[PluginHierarchy.PluginsEnabledAtCook[RootIndex].Name].bIsRoot = true;
 	}
 
+	FTopLevelAssetPath Texture2DPath(TEXT("/Script/Engine.Texture2D"));
+	FTopLevelAssetPath Texture2DArrayPath(TEXT("/Script/Engine.Texture2DArray"));
+	FTopLevelAssetPath Texture3DPath(TEXT("/Script/Engine.Texture3D"));
+	FTopLevelAssetPath TextureCubePath(TEXT("/Script/Engine.TextureCube"));
+	FTopLevelAssetPath TextureCubeArrayPath(TEXT("/Script/Engine.TextureCubeArray"));
+	FTopLevelAssetPath StaticMeshPath(TEXT("/Script/Engine.StaticMesh"));
+	FTopLevelAssetPath SoundWavePath(TEXT("/Script/Engine.SoundWave"));
+
 	double AssetPackageMapStart = FPlatformTime::Seconds();
 	const TMap<FName, const FAssetPackageData*> AssetPackageMap = AssetRegistry.GetAssetPackageDataMap();
 	for (const TPair<FName, const FAssetPackageData*>& AssetPackage : AssetPackageMap)
@@ -3296,6 +3333,10 @@ static void WritePluginMetadataJsons(const FString& InAssetRegistryFileName, TMa
 			// No data on disk!
 			continue;
 		}
+
+		// Grab the most important asset out and use it to track largest asset classes for the plugin.
+		// This might be null!
+		const FAssetData* AssetData = UE::AssetRegistry::GetMostImportantAsset(AssetRegistry.GetAssetsByPackageName(AssetPackage.Key), UE::AssetRegistry::EGetMostImportantAssetFlags::IgnoreSkipClasses);
 		
 		const TArray<FIoStoreChunkSource, TInlineAllocator<2>>* PackageChunks = PackageToChunks.Find(FPackageId::FromName(AssetPackage.Key));
 		if (PackageChunks == nullptr)
@@ -3319,7 +3360,28 @@ static void WritePluginMetadataJsons(const FString& InAssetRegistryFileName, TMa
 			FPluginGraphEntry* PluginEntry = PluginGraph.Find(PluginName);
 			if (PluginEntry)
 			{
-				PluginEntry->ExclusiveSize.Add(PackageSizes);
+				PluginEntry->ExclusiveSizes[(uint8)EPluginGraphSizeClass::All].Add(PackageSizes);
+
+				// If we have asset class info and it's a top contender, track it also.
+				if (AssetData != nullptr)
+				{
+					if (AssetData->AssetClassPath == Texture2DPath ||
+						AssetData->AssetClassPath == Texture3DPath ||
+						AssetData->AssetClassPath == TextureCubePath ||
+						AssetData->AssetClassPath == TextureCubeArrayPath ||
+						AssetData->AssetClassPath == Texture2DArrayPath)
+					{
+						PluginEntry->ExclusiveSizes[(uint8)EPluginGraphSizeClass::Texture].Add(PackageSizes);
+					}
+					else if (AssetData->AssetClassPath == StaticMeshPath)
+					{
+						PluginEntry->ExclusiveSizes[(uint8)EPluginGraphSizeClass::StaticMesh].Add(PackageSizes);
+					}
+					else if (AssetData->AssetClassPath == SoundWavePath)
+					{
+						PluginEntry->ExclusiveSizes[(uint8)EPluginGraphSizeClass::SoundWave].Add(PackageSizes);
+					}
+				}
 			}
 			else
 			{
@@ -3385,7 +3447,10 @@ static void WritePluginMetadataJsons(const FString& InAssetRegistryFileName, TMa
 	{
 		FPluginGraphEntry& PluginEntry = PluginGraph[Plugin->Name];
 
-		PluginEntry.InclusiveSize = PluginEntry.ExclusiveSize;
+		for (uint8 ClassIndex = 0; ClassIndex < FPluginGraphEntry::ClassCount; ClassIndex++)
+		{
+			PluginEntry.InclusiveSizes[ClassIndex] = PluginEntry.ExclusiveSizes[ClassIndex];
+		}
 
 		for (uint16 DependencyIndex = Plugin->DependencyIndexStart; DependencyIndex < Plugin->DependencyIndexEnd; DependencyIndex++)
 		{
@@ -3397,7 +3462,11 @@ static void WritePluginMetadataJsons(const FString& InAssetRegistryFileName, TMa
 			if (bAlreadyInSet == false)
 			{
 				const FPluginGraphEntry* DependencyGraphEntry = &PluginGraph[DependentPlugin.Name];
-				PluginEntry.InclusiveSize.Add(DependencyGraphEntry->ExclusiveSize);
+				for (uint8 ClassIndex = 0; ClassIndex < FPluginGraphEntry::ClassCount; ClassIndex++)
+				{
+					PluginEntry.InclusiveSizes[ClassIndex].Add(DependencyGraphEntry->ExclusiveSizes[ClassIndex]);
+				}
+				
 			}
 			
 			FPluginGraphEntry& DependentEntry = PluginGraph[DependentPlugin.Name];
@@ -3410,13 +3479,64 @@ static void WritePluginMetadataJsons(const FString& InAssetRegistryFileName, TMa
 				if (bAlreadyInSet == false)
 				{
 					const FPluginGraphEntry* DependencyGraphEntry = &PluginGraph[TotalDependencyEntry->Name];
-					PluginEntry.InclusiveSize.Add(DependencyGraphEntry->ExclusiveSize);
+					for (uint8 ClassIndex = 0; ClassIndex < FPluginGraphEntry::ClassCount; ClassIndex++)
+					{
+						PluginEntry.InclusiveSizes[ClassIndex].Add(DependencyGraphEntry->ExclusiveSizes[ClassIndex]);
+					}
 				}
 			}
 		}
 	}
 	double InclusiveComputeEnd = FPlatformTime::Seconds();
 
+	// Now we need to find the unique size for each root plugin. This is the size of dependencies that only
+	// belong to the root plugin and not to another. Conceptually this is the "assuming all other roots are
+	// installed, this is the size cost to add this plugin to the install". These dependencies could be referred
+	// to by another plugin within the unique set - the only requirement is that there exists no path from _another_
+	// root to the dependency.
+	for (uint16 RootIndex : PluginHierarchy.RootPlugins)
+	{
+		FPluginGraphEntry& RootPluginEntry = PluginGraph[PluginHierarchy.PluginsEnabledAtCook[RootIndex].Name];
+
+		// Add us as a root entry for all our total dependencies so all plugins know which root they are in.
+		for (const UE::Cook::FCookMetadataPluginEntry* Dependency : RootPluginEntry.TotalDependencies)
+		{
+			PluginGraph[Dependency->Name].Roots.Add(RootPluginEntry.Self);
+		}
+
+		// Duplicate the TotalDependencies and then remove any dependency that
+		// exists for another root.
+		TSet<const UE::Cook::FCookMetadataPluginEntry*> UniqueDependencies = RootPluginEntry.TotalDependencies;
+
+		for (uint16 InnerRootIndex : PluginHierarchy.RootPlugins)
+		{
+			if (RootIndex == InnerRootIndex)
+			{
+				continue;
+			}
+
+			FPluginGraphEntry& InnerRootPluginEntry = PluginGraph[PluginHierarchy.PluginsEnabledAtCook[InnerRootIndex].Name];
+			for (const UE::Cook::FCookMetadataPluginEntry* InnerRootDependency : InnerRootPluginEntry.TotalDependencies)
+			{
+				UniqueDependencies.Remove(InnerRootDependency);
+			}
+		}
+
+		// Sum the exclusive size of the unique set to get the unique size.
+		for (uint8 ClassIndex = 0; ClassIndex < FPluginGraphEntry::ClassCount; ClassIndex++)
+		{
+			RootPluginEntry.UniqueSizes[ClassIndex].Zero();
+		}
+
+		for (const UE::Cook::FCookMetadataPluginEntry* UniqueDependency : UniqueDependencies)
+		{
+			for (uint8 ClassIndex = 0; ClassIndex < FPluginGraphEntry::ClassCount; ClassIndex++)
+			{
+				RootPluginEntry.UniqueSizes[ClassIndex].Add(PluginGraph[UniqueDependency->Name].ExclusiveSizes[ClassIndex]);
+			}
+		}
+	}
+	
 	// Find the total size of all plugins that aren't rooted in the root set.
 	UE::Cook::FPluginSizeInfo UnrootedTotal;
 	TSet<const UE::Cook::FCookMetadataPluginEntry*> UnrootedPlugins;
@@ -3439,31 +3559,74 @@ static void WritePluginMetadataJsons(const FString& InAssetRegistryFileName, TMa
 
 		for (const UE::Cook::FCookMetadataPluginEntry* Plugin : UnrootedPlugins)
 		{
-			UnrootedTotal.Add(PluginGraph[Plugin->Name].ExclusiveSize);
+			UnrootedTotal.Add(PluginGraph[Plugin->Name].ExclusiveSizes[(uint8)EPluginGraphSizeClass::All]);
 		}
 	}
 
-	auto GeneratePluginJson = [](TUtf8StringBuilder<4096>& OutPluginMetadataJson, FStringView InName, const FPluginGraphEntry& InGraphEntry)
+	auto GeneratePluginJson = [&PluginHierarchy](TUtf8StringBuilder<4096>& OutPluginMetadataJson, FStringView InName, const FPluginGraphEntry& InGraphEntry, EPluginGraphSizeClass InSizeClass)
 	{
 		OutPluginMetadataJson.Reset();
 		OutPluginMetadataJson << "{\n";
 		OutPluginMetadataJson << "\t\"name\":\"" << InName << "\",\n";
 
-		OutPluginMetadataJson << "\t\"schema_version\":1,\n";
+		OutPluginMetadataJson << "\t\"schema_version\":2,\n";
 
 		OutPluginMetadataJson << "\t\"is_root_plugin\":" << (InGraphEntry.bIsRoot ? TEXTVIEW("true") : TEXTVIEW("false")) << ",\n";
 
-		OutPluginMetadataJson << "\t\"exclusive_installed\":" << InGraphEntry.ExclusiveSize[UE::Cook::EPluginSizeTypes::Installed] << ",\n";
-		OutPluginMetadataJson << "\t\"exclusive_optional\":" << InGraphEntry.ExclusiveSize[UE::Cook::EPluginSizeTypes::Optional] << ",\n";
-		OutPluginMetadataJson << "\t\"exclusive_ias\":" << InGraphEntry.ExclusiveSize[UE::Cook::EPluginSizeTypes::Streaming] << ",\n";
-		OutPluginMetadataJson << "\t\"exclusive_optionalsegment\":" << InGraphEntry.ExclusiveSize[UE::Cook::EPluginSizeTypes::OptionalSegment] << ",\n";
+		uint8 SizeClass = (uint8)InSizeClass;
+
+		OutPluginMetadataJson << "\t\"asset_sizes_class\":\"" << PluginGraphEntryClassNames[SizeClass] << "\",\n";
+
+		OutPluginMetadataJson << "\t\"exclusive_installed\":" << InGraphEntry.ExclusiveSizes[SizeClass][UE::Cook::EPluginSizeTypes::Installed] << ",\n";
+		OutPluginMetadataJson << "\t\"exclusive_optional\":" << InGraphEntry.ExclusiveSizes[SizeClass][UE::Cook::EPluginSizeTypes::Optional] << ",\n";
+		OutPluginMetadataJson << "\t\"exclusive_ias\":" << InGraphEntry.ExclusiveSizes[SizeClass][UE::Cook::EPluginSizeTypes::Streaming] << ",\n";
+		OutPluginMetadataJson << "\t\"exclusive_optionalsegment\":" << InGraphEntry.ExclusiveSizes[SizeClass][UE::Cook::EPluginSizeTypes::OptionalSegment] << ",\n";
+		
+		OutPluginMetadataJson << "\t\"inclusive_installed\":" << InGraphEntry.InclusiveSizes[SizeClass][UE::Cook::EPluginSizeTypes::Installed] << ",\n";
+		OutPluginMetadataJson << "\t\"inclusive_optional\":" << InGraphEntry.InclusiveSizes[SizeClass][UE::Cook::EPluginSizeTypes::Optional] << ",\n";
+		OutPluginMetadataJson << "\t\"inclusive_ias\":" << InGraphEntry.InclusiveSizes[SizeClass][UE::Cook::EPluginSizeTypes::Streaming] << ",\n";
+		OutPluginMetadataJson << "\t\"inclusive_optionalsegment\":" << InGraphEntry.InclusiveSizes[SizeClass][UE::Cook::EPluginSizeTypes::OptionalSegment] << ",\n";
+
+		// this only has values for is_root_plugin == true.
+		OutPluginMetadataJson << "\t\"unique_installed\":" << InGraphEntry.UniqueSizes[SizeClass][UE::Cook::EPluginSizeTypes::Installed] << ",\n";
+		OutPluginMetadataJson << "\t\"unique_optional\":" << InGraphEntry.UniqueSizes[SizeClass][UE::Cook::EPluginSizeTypes::Optional] << ",\n";
+		OutPluginMetadataJson << "\t\"unique_ias\":" << InGraphEntry.UniqueSizes[SizeClass][UE::Cook::EPluginSizeTypes::Streaming] << ",\n";
+		OutPluginMetadataJson << "\t\"unique_optionalsegment\":" << InGraphEntry.UniqueSizes[SizeClass][UE::Cook::EPluginSizeTypes::OptionalSegment] << ",\n";
+
+		OutPluginMetadataJson << "\t\"direct_refcount\":" << InGraphEntry.DirectRefcount << ",\n";
+
+		// pass through any custom fields that were added to the cook metadata.
+		if (InGraphEntry.Self)
+		{
+			for (const TPair<uint8, bool>& BoolValue : InGraphEntry.Self->CustomBoolFields)
+			{
+				const FString& FieldName = PluginHierarchy.CustomFieldNames[BoolValue.Key];
+				OutPluginMetadataJson << "\t\"" << FieldName << (BoolValue.Value ? "\":true,\n" : "\":false,\n");				
+			}
+			for (const TPair<uint8, FString>& StringValue : InGraphEntry.Self->CustomStringFields)
+			{
+				const FString& FieldName = PluginHierarchy.CustomFieldNames[StringValue.Key];
+				OutPluginMetadataJson << "\t\"" << FieldName << "\":\"" << StringValue.Value << "\",\n";
+			}
+		}
+
+		{
+			OutPluginMetadataJson << "\t\"roots\":[";
+
+			for (int32 RootIndex = 0; RootIndex < InGraphEntry.Roots.Num(); RootIndex++)
+			{
+				const UE::Cook::FCookMetadataPluginEntry* Root = InGraphEntry.Roots[RootIndex];
+				OutPluginMetadataJson << "\"" << Root->Name << "\"";
+				if (RootIndex + 1 < InGraphEntry.Roots.Num())
+				{
+					OutPluginMetadataJson << ",";
+				}
+			}
+
+			OutPluginMetadataJson << "]\n";
+		}
 
-		OutPluginMetadataJson << "\t\"inclusive_installed\":" << InGraphEntry.InclusiveSize[UE::Cook::EPluginSizeTypes::Installed] << ",\n";
-		OutPluginMetadataJson << "\t\"inclusive_optional\":" << InGraphEntry.InclusiveSize[UE::Cook::EPluginSizeTypes::Optional] << ",\n";
-		OutPluginMetadataJson << "\t\"inclusive_ias\":" << InGraphEntry.InclusiveSize[UE::Cook::EPluginSizeTypes::Streaming] << ",\n";
-		OutPluginMetadataJson << "\t\"inclusive_optionalsegment\":" << InGraphEntry.InclusiveSize[UE::Cook::EPluginSizeTypes::OptionalSegment] << ",\n";
 
-		OutPluginMetadataJson << "\t\"direct_refcount\":" << InGraphEntry.DirectRefcount << "\n";
 
 		OutPluginMetadataJson << "}\n";
 	};
@@ -3475,28 +3638,38 @@ static void WritePluginMetadataJsons(const FString& InAssetRegistryFileName, TMa
 
 	// Also write a csv for easier browsing in spreadsheets.
 	TUtf8StringBuilder<4096> Csv;
-	Csv.Append("name,exclusive_installed,exclusive_optional,exclusive_ias,inclusive_installed,inclusive_optional,inclusive_ias,direct_refcount,total_dependency_count\n");
+	Csv.Append("name,asset_sizes_class,exclusive_installed,exclusive_optional,exclusive_ias,inclusive_installed,inclusive_optional,inclusive_ias,unique_installed,unique_optional,unique_ias,direct_refcount,total_dependency_count\n");
 
 	// We want to write the sizes back to the cook metadata, so we need a non-const version.
 	UE::Cook::FCookMetadataPluginHierarchy& MutablePluginHierarchy = CookMetadata.GetMutablePluginHierarchy();
 	for (UE::Cook::FCookMetadataPluginEntry& Plugin : MutablePluginHierarchy.PluginsEnabledAtCook)
 	{
 		const FPluginGraphEntry& PluginEntry = PluginGraph[Plugin.Name];
-		if (PluginEntry.InclusiveSize.TotalSize() == 0)
+		if (PluginEntry.InclusiveSizes[(uint8)EPluginGraphSizeClass::All].TotalSize() == 0)
 		{
 			continue;
 		}
 
-		GeneratePluginJson(PluginMetadataJson, Plugin.Name, PluginEntry);
-		SavePluginMetadata(InAssetRegistryFileName, Plugin.Name, PluginMetadataJson);
+		Plugin.InclusiveSizes = PluginEntry.InclusiveSizes[(uint8)EPluginGraphSizeClass::All];
+		Plugin.ExclusiveSizes = PluginEntry.ExclusiveSizes[(uint8)EPluginGraphSizeClass::All];
 
-		Plugin.InclusiveSizes = PluginEntry.InclusiveSize;
-		Plugin.ExclusiveSizes = PluginEntry.ExclusiveSize;
+		for (uint8 ClassIndex = 0; ClassIndex < FPluginGraphEntry::ClassCount; ClassIndex++)
+		{
+			GeneratePluginJson(PluginMetadataJson, Plugin.Name, PluginEntry, (EPluginGraphSizeClass)ClassIndex);
+			SavePluginMetadata(InAssetRegistryFileName, Plugin.Name, PluginMetadataJson, (EPluginGraphSizeClass)ClassIndex);
 
-		Csv.Appendf("%ls,%llu,%llu,%llu,%llu,%llu,%llu,%u,%u\n", *Plugin.Name, 
-			Plugin.ExclusiveSizes[UE::Cook::EPluginSizeTypes::Installed], Plugin.ExclusiveSizes[UE::Cook::EPluginSizeTypes::Optional], Plugin.ExclusiveSizes[UE::Cook::EPluginSizeTypes::Streaming],
-			Plugin.InclusiveSizes[UE::Cook::EPluginSizeTypes::Installed], Plugin.InclusiveSizes[UE::Cook::EPluginSizeTypes::Optional], Plugin.InclusiveSizes[UE::Cook::EPluginSizeTypes::Streaming],
-			PluginEntry.DirectRefcount,PluginEntry.TotalDependencies.Num());
+			Csv.Appendf("%ls,%s,%llu,%llu,%llu,%llu,%llu,%llu,%llu,%llu,%llu,%u,%u\n", *Plugin.Name, PluginGraphEntryClassNames[ClassIndex],
+				PluginEntry.ExclusiveSizes[ClassIndex][UE::Cook::EPluginSizeTypes::Installed],
+				PluginEntry.ExclusiveSizes[ClassIndex][UE::Cook::EPluginSizeTypes::Optional],
+				PluginEntry.ExclusiveSizes[ClassIndex][UE::Cook::EPluginSizeTypes::Streaming],
+				PluginEntry.InclusiveSizes[ClassIndex][UE::Cook::EPluginSizeTypes::Installed],
+				PluginEntry.InclusiveSizes[ClassIndex][UE::Cook::EPluginSizeTypes::Optional],
+				PluginEntry.InclusiveSizes[ClassIndex][UE::Cook::EPluginSizeTypes::Streaming],
+				PluginEntry.UniqueSizes[ClassIndex][UE::Cook::EPluginSizeTypes::Installed],
+				PluginEntry.UniqueSizes[ClassIndex][UE::Cook::EPluginSizeTypes::Optional],
+				PluginEntry.UniqueSizes[ClassIndex][UE::Cook::EPluginSizeTypes::Streaming],
+				PluginEntry.DirectRefcount,PluginEntry.TotalDependencies.Num());
+		}
 	}
 
 	// Also write a json that contains the sizes for the plugins that don't belong to any root plugin,
@@ -3505,11 +3678,11 @@ static void WritePluginMetadataJsons(const FString& InAssetRegistryFileName, TMa
 		FPluginGraphEntry OrphanedEntry;
 		OrphanedEntry.bIsRoot = true;
 		OrphanedEntry.DirectRefcount = 0;
-		OrphanedEntry.InclusiveSize = UnrootedTotal;
-		GeneratePluginJson(PluginMetadataJson, TEXT("OrphanedPlugins"), OrphanedEntry);
-		SavePluginMetadata(InAssetRegistryFileName, TEXT("OrphanedPlugins"), PluginMetadataJson);
+		OrphanedEntry.InclusiveSizes[(uint8)EPluginGraphSizeClass::All] = UnrootedTotal;
+		GeneratePluginJson(PluginMetadataJson, TEXT("OrphanedPlugins"), OrphanedEntry, EPluginGraphSizeClass::All);
+		SavePluginMetadata(InAssetRegistryFileName, TEXT("OrphanedPlugins"), PluginMetadataJson, EPluginGraphSizeClass::All);
 
-		Csv.Appendf("OrphanedPlugins,0,0,0,%llu,%llu,%llu,0,%u\n",
+		Csv.Appendf("OrphanedPlugins,all,0,0,0,%llu,%llu,%llu,0,%u\n",
 			UnrootedTotal[UE::Cook::EPluginSizeTypes::Installed], UnrootedTotal[UE::Cook::EPluginSizeTypes::Optional], UnrootedTotal[UE::Cook::EPluginSizeTypes::Streaming],
 			UnrootedPlugins.Num());
 	}
diff --git a/Engine/Source/Editor/UnrealEd/Private/CookOnTheFlyServer.cpp b/Engine/Source/Editor/UnrealEd/Private/CookOnTheFlyServer.cpp
--- a/Engine/Source/Editor/UnrealEd/Private/CookOnTheFlyServer.cpp
+++ b/Engine/Source/Editor/UnrealEd/Private/CookOnTheFlyServer.cpp
@@ -46,6 +46,7 @@
 #include "CookPackageSplitter.h"
 #include "DerivedDataCacheInterface.h"
 #include "DistanceFieldAtlas.h"
+#include "Dom/JsonObject.h"
 #include "Editor.h"
 #include "Editor/UnrealEdEngine.h"
 #include "EditorDomain/EditorDomain.h"
@@ -9482,7 +9483,7 @@ void UCookOnTheFlyServer::WriteCookMetadata(const ITargetPlatform* InTargetPlatf
 	constexpr int32 AdditionalPseudoPlugins = 2; // /Engine and /Game.
 	if (IntFitsIn<uint16>(EnabledPlugins.Num() + AdditionalPseudoPlugins) == false)
 	{
-		UE_LOG(LogCook, Warning, TEXT("Number of plugins exceeds 64k, unable to write cook metadata file (count = %d"), EnabledPlugins.Num() + AdditionalPseudoPlugins);
+		UE_LOG(LogCook, Warning, TEXT("Number of plugins exceeds 64k, unable to write cook metadata file (count = %d)"), EnabledPlugins.Num() + AdditionalPseudoPlugins);
 	}
 	else
 	{
@@ -9500,6 +9501,44 @@ void UCookOnTheFlyServer::WriteCookMetadata(const ITargetPlatform* InTargetPlatf
 			AddIndex++;
 		}
 
+		TArray<FString> BoolCustomFieldsList;
+		TArray<FString> StringCustomFieldsList;
+		TArray<FString> PerPlatformBoolCustomFieldsList;
+		TArray<FString> PerPlatformStringCustomFieldsList;
+		GConfig->GetArray(TEXT("CookMetadataCustomPluginFields"), TEXT("BoolFields"), BoolCustomFieldsList, GEditorIni);		
+		GConfig->GetArray(TEXT("CookMetadataCustomPluginFields"), TEXT("StringFields"), StringCustomFieldsList, GEditorIni);
+		GConfig->GetArray(TEXT("CookMetadataCustomPluginFields"), TEXT("PerPlatformBoolFields"), PerPlatformBoolCustomFieldsList, GEditorIni);
+		GConfig->GetArray(TEXT("CookMetadataCustomPluginFields"), TEXT("PerPlatformStringFields"), PerPlatformStringCustomFieldsList, GEditorIni);
+
+		// Get the names as a unique list for indexing the names for serialization
+		TArray<FString> CustomFieldNames;
+		TMap<FString, uint8> CustomFieldNameIndex;
+		{
+			auto GetNames = [&CustomFieldNameIndex, &CustomFieldNames](const TArray<FString>& FieldList)
+			{
+				for (const FString& FieldName : FieldList)
+				{
+					uint8& FoundAtIndex = CustomFieldNameIndex.FindOrAdd(FieldName, MAX_uint8);
+					if (FoundAtIndex == MAX_uint8)
+					{
+						CustomFieldNames.Add(FieldName);
+						FoundAtIndex = (uint8)(CustomFieldNames.Num() - 1);
+					}
+				}
+			};
+
+			GetNames(BoolCustomFieldsList);
+			GetNames(StringCustomFieldsList);
+			GetNames(PerPlatformBoolCustomFieldsList);
+			GetNames(PerPlatformStringCustomFieldsList);
+		}
+
+		if (CustomFieldNames.Num() > 255)
+		{
+			// Sanity check integer limits - should never hit this, but if we do all bets are off.
+			UE_LOG(LogCook, Warning, TEXT("Number of custom plugin fields exceeds 255 (count = %d), custom fields will be incorrect!"), CustomFieldNames.Num());
+		}
+
 		// Add the /Engine and /Game pseudo plugins. These are placeholders for holding size information when unrealpak runs.
 		const uint16 EnginePluginIndex = (uint16)PluginsToAdd.Num();
 		PluginsToAdd.AddDefaulted_GetRef().Name = TEXT("Engine");
@@ -9525,6 +9564,117 @@ void UCookOnTheFlyServer::WriteCookMetadata(const ITargetPlatform* InTargetPlatf
 				RootPlugins.Add(SelfIndex);
 			}
 
+			// Pull in any custom fields the project wants to pass on.
+			auto CheckForCustomFields = [&Descriptor, &PlatformNameString, &EnabledPlugin, &Entry, &CustomFieldNameIndex](bool bIsBool, bool bIsPerPlatform, const TArray<FString>& FieldNames)
+			{
+				for (const FString& FieldName : FieldNames)
+				{
+					TVariant<bool, FString> FieldValue;
+			
+					bool bHasField = false;
+
+					if (bIsBool)
+					{
+						bool bPlatformAgnosticValue = false;
+						bHasField = Descriptor.CachedJson->TryGetBoolField(FieldName, bPlatformAgnosticValue);
+						FieldValue.Set<bool>(bPlatformAgnosticValue);
+					}
+					else
+					{
+						FString PlatformAgnosticString;
+						bHasField = Descriptor.CachedJson->TryGetStringField(FieldName, PlatformAgnosticString);
+						FieldValue.Set<FString>(MoveTemp(PlatformAgnosticString));
+					}
+
+					if (bIsPerPlatform)
+					{
+						// If the field is marked as per-platform, then it has a field with the same name except
+						// prepended with PerPlatform. Inside that is an array of objects with Platform and Value
+						// to specify the override for specific platforms.
+						bool bHasPerPlatform = false;
+
+						const TArray<TSharedPtr<FJsonValue>>* Array;
+						if (Descriptor.CachedJson->TryGetArrayField(TEXT("PerPlatform") + FieldName, Array))
+						{
+							bHasPerPlatform = true;
+
+							for (const TSharedPtr<FJsonValue>& Value : *Array)
+							{
+								const TSharedPtr<FJsonObject>& ValueObject = Value->AsObject();
+								if (ValueObject.IsValid())
+								{
+									FString OverridePlatformName;
+									if (!ValueObject->TryGetStringField(TEXT("Platform"), OverridePlatformName))
+									{
+										UE_LOG(LogCook, Error, TEXT("Unable to get Platform field from PerPlatform%s array in plugin %s json."), *FieldName, *EnabledPlugin->GetName());
+										continue;
+									}
+
+									if (OverridePlatformName == PlatformNameString)
+									{
+										bool bGotOverride = false;
+
+										if (bIsBool)
+										{
+											bool bPlatformValue = false;
+											bGotOverride = ValueObject->TryGetBoolField(TEXT("Value"), bPlatformValue);
+											FieldValue.Set<bool>(bPlatformValue);
+										}
+										else
+										{
+											FString PlatformString;
+											bGotOverride = ValueObject->TryGetStringField(TEXT("Value"), PlatformString);
+											FieldValue.Set<FString>(MoveTemp(PlatformString));
+										}
+
+										if (!bGotOverride)
+										{
+											UE_LOG(LogCook, Error, TEXT("Unable to get Value field from PerPlatform%s array in plugin %s json for platform %s"), *FieldName, *EnabledPlugin->GetName(), *PlatformNameString);
+											continue;
+										}
+										bHasField = true;
+									}
+								}
+							}
+						} // end if the plugin has overrides
+
+						// If the field has a per platform value, but no value for this platform in either the
+						// agnostic or the specific area, then we fill it with default values so it's still
+						// present in the output, even if it's just default values.
+						if (bHasPerPlatform && !bHasField)
+						{
+							bHasField = true;
+							if (bIsBool)
+							{
+								FieldValue.Set<bool>(false);
+							}
+							else
+							{
+								FieldValue.Set<FString>(FString());
+							}
+						}
+					} // end if field is per platform
+
+					if (bHasField)
+					{
+						if (bIsBool)
+						{
+							Entry.CustomBoolFields.Add(CustomFieldNameIndex[FieldName], FieldValue.Get<bool>());
+						}
+						else
+						{
+							Entry.CustomStringFields.Add(CustomFieldNameIndex[FieldName], MoveTemp(FieldValue.Get<FString>()));
+						}
+					}
+				} // end each field name
+
+			}; // end local lambda
+
+			CheckForCustomFields(true, false, BoolCustomFieldsList);
+			CheckForCustomFields(true, true, PerPlatformBoolCustomFieldsList);
+			CheckForCustomFields(false, false, StringCustomFieldsList);
+			CheckForCustomFields(false, true, PerPlatformStringCustomFieldsList);
+
 			for (FPluginReferenceDescriptor ChildPlugin : Descriptor.Plugins)
 			{
 				if (uint16* ChildIndex = IndexForPlugin.Find(ChildPlugin.Name); ChildIndex != nullptr)
@@ -9567,6 +9717,7 @@ void UCookOnTheFlyServer::WriteCookMetadata(const ITargetPlatform* InTargetPlatf
 			PluginHierarchy.PluginsEnabledAtCook = MoveTemp(PluginsToAdd);
 			PluginHierarchy.PluginDependencies = MoveTemp(PluginChildArray);
 			PluginHierarchy.RootPlugins = MoveTemp(RootPlugins);
+			PluginHierarchy.CustomFieldNames = MoveTemp(CustomFieldNames);
 
 			MetadataState.SetPluginHierarchyInfo(MoveTemp(PluginHierarchy));
 			MetadataState.SetAssociatedDevelopmentAssetRegistryHash(InDevelopmentAssetRegistryHash);
