IoStore: Make sure the FZenPackageHeader array views starting with ImportedPublicExportHashes are 8-byte aligned

diff --git a/Engine/Source/Developer/IoStoreUtilities/Private/PackageStoreOptimizer.cpp b/Engine/Source/Developer/IoStoreUtilities/Private/PackageStoreOptimizer.cpp
--- a/Engine/Source/Developer/IoStoreUtilities/Private/PackageStoreOptimizer.cpp
+++ b/Engine/Source/Developer/IoStoreUtilities/Private/PackageStoreOptimizer.cpp
@@ -735,17 +735,22 @@ void FPackageStoreOptimizer::FinalizePackageHeader(FPackageStorePackage* Package
 	}
 	uint64 BulkDataMapSize = BulkDataMapAr.Tell();
 
-	uint64 HeaderSize =
+	uint64 SizeBeforePublicExportHashes = 
 		sizeof(FZenPackageSummary)
 		+ VersioningInfoSize
 		+ NameMapSize
+		+ BulkDataMapSize + sizeof(int64);
+
+	uint64 AlignedSizeBeforePublicExportHashes = Align(SizeBeforePublicExportHashes, sizeof(uint64));
+
+	uint64 HeaderSize =
+		AlignedSizeBeforePublicExportHashes
 		+ ImportedPublicExportHashesSize
 		+ ImportMapSize
 		+ ExportMapSize
 		+ ExportBundleEntriesSize
 		+ GraphDataSize
-		+ ImportedPackagesSize
-		+ BulkDataMapSize + sizeof(int64);
+		+ ImportedPackagesSize;
 
 	Package->HeaderBuffer = FIoBuffer(HeaderSize);
 	uint8* HeaderData = Package->HeaderBuffer.Data();
@@ -773,12 +778,21 @@ void FPackageStoreOptimizer::FinalizePackageHeader(FPackageStorePackage* Package
 	HeaderArchive << BulkDataMapSize;
 	HeaderArchive.Serialize(BulkDataMapAr.GetWriterData(), BulkDataMapSize);
 
+	if (uint64 Pad=AlignedSizeBeforePublicExportHashes-SizeBeforePublicExportHashes; Pad > 0)
+	{
+		uint8 PadBytes[sizeof(uint64)] = {};
+		HeaderArchive.Serialize(&PadBytes[0], Pad);
+	}
+	check(HeaderArchive.Tell() == AlignedSizeBeforePublicExportHashes);
+
+	// raw arrays of 8-byte aligned items
 	PackageSummary->ImportedPublicExportHashesOffset = HeaderArchive.Tell();
 	HeaderArchive.Serialize(ImportedPublicExportHashesArchive.GetWriterData(), ImportedPublicExportHashesArchive.Tell());
 	PackageSummary->ImportMapOffset = HeaderArchive.Tell();
 	HeaderArchive.Serialize(ImportMapArchive.GetWriterData(), ImportMapArchive.Tell());
 	PackageSummary->ExportMapOffset = HeaderArchive.Tell();
 	HeaderArchive.Serialize(ExportMapArchive.GetWriterData(), ExportMapArchive.Tell());
+	// raw arrays of 4-byte aligned items
 	PackageSummary->ExportBundleEntriesOffset = HeaderArchive.Tell();
 	HeaderArchive.Serialize(ExportBundleEntriesArchive.GetWriterData(), ExportBundleEntriesArchive.Tell());
 	PackageSummary->DependencyBundleHeadersOffset = HeaderArchive.Tell();
