Added Skeletal Mesh asset support to IO Store Telemetry

diff --git a/Engine/Source/Developer/IoStoreUtilities/Private/IoStoreUtilities.cpp b/Engine/Source/Developer/IoStoreUtilities/Private/IoStoreUtilities.cpp
--- a/Engine/Source/Developer/IoStoreUtilities/Private/IoStoreUtilities.cpp
+++ b/Engine/Source/Developer/IoStoreUtilities/Private/IoStoreUtilities.cpp
@@ -3228,6 +3228,7 @@ enum class EPluginGraphSizeClass : uint8
 	Texture,
 	StaticMesh,
 	SoundWave,
+	SkeletalMesh,
 	COUNT
 };
 
@@ -3236,7 +3237,8 @@ static const UTF8CHAR* PluginGraphEntryClassNames[] =
 	UTF8TEXT("all"),
 	UTF8TEXT("texture"),
 	UTF8TEXT("staticmesh"),
-	UTF8TEXT("soundwave")
+	UTF8TEXT("soundwave"),
+	UTF8TEXT("skeletalmesh")
 };
 
 static_assert( UE_ARRAY_COUNT(PluginGraphEntryClassNames) == (size_t)EPluginGraphSizeClass::COUNT, "Must have a name for each plugin graph size class!");
@@ -3325,6 +3327,7 @@ static void WritePluginMetadataJsons(const FString& InAssetRegistryFileName, TMa
 	FTopLevelAssetPath TextureCubeArrayPath(TEXT("/Script/Engine.TextureCubeArray"));
 	FTopLevelAssetPath StaticMeshPath(TEXT("/Script/Engine.StaticMesh"));
 	FTopLevelAssetPath SoundWavePath(TEXT("/Script/Engine.SoundWave"));
+	FTopLevelAssetPath SkeletalMeshPath(TEXT("/Script/Engine.SkeletalMesh"));
 
 	double AssetPackageMapStart = FPlatformTime::Seconds();
 	const TMap<FName, const FAssetPackageData*> AssetPackageMap = AssetRegistry.GetAssetPackageDataMap();
@@ -3383,6 +3386,10 @@ static void WritePluginMetadataJsons(const FString& InAssetRegistryFileName, TMa
 					{
 						PluginEntry->ExclusiveSizes[(uint8)EPluginGraphSizeClass::SoundWave].Add(PackageSizes);
 					}
+					else if (AssetData->AssetClassPath == SkeletalMeshPath)
+					{
+						PluginEntry->ExclusiveSizes[(uint8)EPluginGraphSizeClass::SkeletalMesh].Add(PackageSizes);
+					}
 				}
 			}
 			else
