VFS - return error if any chunk fails to write to disk

#rnx

diff --git a/Engine/Source/Developer/IoStoreUtilities/Private/IoStoreLooseFiles.cpp b/Engine/Source/Developer/IoStoreUtilities/Private/IoStoreLooseFiles.cpp
--- a/Engine/Source/Developer/IoStoreUtilities/Private/IoStoreLooseFiles.cpp
+++ b/Engine/Source/Developer/IoStoreUtilities/Private/IoStoreLooseFiles.cpp
@@ -211,7 +211,7 @@ public:
 			if (PendingWrite->ErrorText.IsEmpty() == false)
 			{
 				UE_LOG(LogIoStore, Display, TEXT("Failed to write on demand chunk '%s', reason '%s'"), *LexToString(PendingWrite->ChunkId), *PendingWrite->ErrorText);
-				continue;
+				return FIoStatus(EIoErrorCode::WriteError, PendingWrite->ErrorText);
 			}
 			JsonWriter->WriteObjectStart();
 			JsonWriter->WriteValue(TEXT("IoChunkId"), LexToString(PendingWrite->ChunkId));
