Zen header alignment

diff --git a/Engine/Source/Developer/IoStoreUtilities/Private/PackageStoreOptimizer.cpp b/Engine/Source/Developer/IoStoreUtilities/Private/PackageStoreOptimizer.cpp
--- a/Engine/Source/Developer/IoStoreUtilities/Private/PackageStoreOptimizer.cpp
+++ b/Engine/Source/Developer/IoStoreUtilities/Private/PackageStoreOptimizer.cpp
@@ -735,16 +735,24 @@ void FPackageStoreOptimizer::FinalizePackageHeader(FPackageStorePackage* Package
 	}
 	uint64 BulkDataMapSize = BulkDataMapAr.Tell();
 
-	uint64 SizeBeforePublicExportHashes = 
+	uint64 BulkDataPad = 0;
+	uint64 OffsetBeforeBulkDataMap =
 		sizeof(FZenPackageSummary)
 		+ VersioningInfoSize
 		+ NameMapSize
-		+ BulkDataMapSize + sizeof(int64);
+		+ sizeof(BulkDataPad);
 
-	uint64 AlignedSizeBeforePublicExportHashes = Align(SizeBeforePublicExportHashes, sizeof(uint64));
+	uint64 AlignedOffsetBeforeBulkDataMap = Align(OffsetBeforeBulkDataMap, sizeof(uint64));
+	BulkDataPad = AlignedOffsetBeforeBulkDataMap - OffsetBeforeBulkDataMap;
+
+	uint64 OffsetBeforePublicExportHashes =
+		AlignedOffsetBeforeBulkDataMap
+		+ BulkDataMapSize + sizeof(BulkDataMapSize);
+
+	uint64 AlignedOffsetBeforePublicExportHashes = Align(OffsetBeforePublicExportHashes, sizeof(uint64));
 
 	uint64 HeaderSize =
-		AlignedSizeBeforePublicExportHashes
+		AlignedOffsetBeforePublicExportHashes
 		+ ImportedPublicExportHashesSize
 		+ ImportMapSize
 		+ ExportMapSize
@@ -775,15 +783,23 @@ void FPackageStoreOptimizer::FinalizePackageHeader(FPackageStorePackage* Package
 
 	HeaderArchive.Serialize(NameMapArchive.GetWriterData(), NameMapArchive.Tell());
 
+	HeaderArchive << BulkDataPad;
+	if (BulkDataPad > 0)
+	{
+		uint8 PadBytes[sizeof(uint64)] = {};
+		HeaderArchive.Serialize(PadBytes, BulkDataPad);
+	}
+	check(HeaderArchive.Tell() == AlignedOffsetBeforeBulkDataMap);
+
 	HeaderArchive << BulkDataMapSize;
 	HeaderArchive.Serialize(BulkDataMapAr.GetWriterData(), BulkDataMapSize);
 
-	if (uint64 Pad=AlignedSizeBeforePublicExportHashes-SizeBeforePublicExportHashes; Pad > 0)
+	if (uint64 Pad=AlignedOffsetBeforePublicExportHashes-OffsetBeforePublicExportHashes; Pad > 0)
 	{
 		uint8 PadBytes[sizeof(uint64)] = {};
-		HeaderArchive.Serialize(&PadBytes[0], Pad);
+		HeaderArchive.Serialize(PadBytes, Pad);
 	}
-	check(HeaderArchive.Tell() == AlignedSizeBeforePublicExportHashes);
+	check(HeaderArchive.Tell() == AlignedOffsetBeforePublicExportHashes);
 
 	// raw arrays of 8-byte aligned items
 	PackageSummary->ImportedPublicExportHashesOffset = HeaderArchive.Tell();
diff --git a/Engine/Source/Runtime/CoreUObject/Private/Serialization/ZenPackageHeader.cpp b/Engine/Source/Runtime/CoreUObject/Private/Serialization/ZenPackageHeader.cpp
--- a/Engine/Source/Runtime/CoreUObject/Private/Serialization/ZenPackageHeader.cpp
+++ b/Engine/Source/Runtime/CoreUObject/Private/Serialization/ZenPackageHeader.cpp
@@ -29,6 +29,10 @@ FZenPackageHeader FZenPackageHeader::MakeView(FMemoryView Memory)
 	if (LocalVersioningInfo == nullptr || LocalVersioningInfo->PackageVersion >= EUnrealEngineObjectUE5Version::DATA_RESOURCES)
 	{
 		int64 BulkDataMapSize = 0;
+		uint64 BulkDataPad = 0;
+		PackageHeaderDataReader << BulkDataPad;
+		uint8 PadBytes[sizeof(uint64)] = {};
+		PackageHeaderDataReader.Serialize(PadBytes, BulkDataPad);
 		PackageHeaderDataReader << BulkDataMapSize;
 		const uint8* BulkDataMapData = PackageHeaderDataPtr + sizeof(FZenPackageSummary) + PackageHeaderDataReader.Tell();
 		PackageHeader.BulkDataMap = MakeArrayView(reinterpret_cast<const FBulkDataMapEntry*>(BulkDataMapData), BulkDataMapSize / sizeof(FBulkDataMapEntry));
