IAS - remove warning about unencrypted PAK's

#rnx

diff --git a/Engine/Source/Developer/IoStoreUtilities/Private/IoStoreUtilities.cpp b/Engine/Source/Developer/IoStoreUtilities/Private/IoStoreUtilities.cpp
--- a/Engine/Source/Developer/IoStoreUtilities/Private/IoStoreUtilities.cpp
+++ b/Engine/Source/Developer/IoStoreUtilities/Private/IoStoreUtilities.cpp
@@ -2626,12 +2626,6 @@ void InitializeContainerTargetsAndPackages(
 				ContainerTarget->TargetFiles.Emplace(MoveTemp(TargetFile));
 			}
 
-			if (EnumHasAnyFlags(ContainerTarget->ContainerFlags, EIoContainerFlags::OnDemand) &&
-				!EnumHasAnyFlags(ContainerTarget->ContainerFlags, EIoContainerFlags::Encrypted))
-			{
-				UE_LOG(LogIoStore, Warning, TEXT("Container '%s' set as on demand but is not encrypted, please revisit encryption strategy before publisihing to public CDN's"), *ContainerSource.Name.ToString());
-			}
-
 			if (bHasOptionalSegmentPackages)
 			{
 				if (ContainerSource.OptionalOutputPath.IsEmpty())
