Bumping schema version since i changed it.

diff --git a/Engine/Source/Developer/IoStoreUtilities/Private/IoStoreUtilities.cpp b/Engine/Source/Developer/IoStoreUtilities/Private/IoStoreUtilities.cpp
--- a/Engine/Source/Developer/IoStoreUtilities/Private/IoStoreUtilities.cpp
+++ b/Engine/Source/Developer/IoStoreUtilities/Private/IoStoreUtilities.cpp
@@ -3578,7 +3578,7 @@ static void WritePluginMetadataJsons(const FString& InAssetRegistryFileName, TMa
 		OutPluginMetadataJson << "{\n";
 		OutPluginMetadataJson << "\t\"name\":\"" << InName << "\",\n";
 
-		OutPluginMetadataJson << "\t\"schema_version\":2,\n";
+		OutPluginMetadataJson << "\t\"schema_version\":3,\n";
 
 		OutPluginMetadataJson << "\t\"is_root_plugin\":" << (InGraphEntry.bIsRoot ? TEXTVIEW("true") : TEXTVIEW("false")) << ",\n";
 
