Include shader sizes in plugin size reporting.
Add shader information to the cook metadata file.
Update cook metadata dependency counts to uint32.
Separate out plugin graph generation from usage.
Add "other" size classification

diff --git a/Engine/Plugins/Editor/AssetManagerEditor/Source/AssetManagerEditor/Private/TreeView/SAssetTableTreeView.cpp b/Engine/Plugins/Editor/AssetManagerEditor/Source/AssetManagerEditor/Private/TreeView/SAssetTableTreeView.cpp
--- a/Engine/Plugins/Editor/AssetManagerEditor/Source/AssetManagerEditor/Private/TreeView/SAssetTableTreeView.cpp
+++ b/Engine/Plugins/Editor/AssetManagerEditor/Source/AssetManagerEditor/Private/TreeView/SAssetTableTreeView.cpp
@@ -1933,7 +1933,7 @@ void SAssetTableTreeView::RefreshAssets()
 			}
 
 			const int32 PluginIndex = AssetTable->GetIndexForPlugin(StoredPluginName);
-			for (uint16 DependencyIndexInMetadata = PluginEntry.DependencyIndexStart; DependencyIndexInMetadata < PluginEntry.DependencyIndexEnd; DependencyIndexInMetadata++)
+			for (uint32 DependencyIndexInMetadata = PluginEntry.DependencyIndexStart; DependencyIndexInMetadata < PluginEntry.DependencyIndexEnd; DependencyIndexInMetadata++)
 			{
 				const UE::Cook::FCookMetadataPluginEntry& DependentPluginEntry = PluginHierarchy.PluginsEnabledAtCook[PluginHierarchy.PluginDependencies[DependencyIndexInMetadata]];
 				const TCHAR* StoredReferencePluginName = AssetTable->StoreStr(DependentPluginEntry.Name);
diff --git a/Engine/Source/Developer/CookMetadata/Public/CookMetadata.h b/Engine/Source/Developer/CookMetadata/Public/CookMetadata.h
--- a/Engine/Source/Developer/CookMetadata/Public/CookMetadata.h
+++ b/Engine/Source/Developer/CookMetadata/Public/CookMetadata.h
@@ -7,6 +7,7 @@
 #include "Containers/UnrealString.h"
 #include "CoreTypes.h"
 #include "Memory/MemoryFwd.h"
+#include "UObject/NameTypes.h"
 
 
 namespace UE::Cook
@@ -19,6 +20,7 @@ enum class ECookMetadataStateVersion : uint8
 	PostWritebackHash = 2,
 	FixSerialization = 3,
 	AddedCustomFields = 4,
+	AddedShaderPseudoHierarchy = 5,
 
 	// Add new versions above this.
 	VersionCount,
@@ -115,11 +117,32 @@ enum class ECookMetadataSizesPresent
 	Count
 };
 
+enum class ECookMetadataPluginType
+{
+	Normal,
+
+	// Root plugins are used to separate and classify game "modes" within a single project.
+	Root,
+
+	// For assets under /Engine
+	EnginePseudo,
+
+	// For assets under /Game
+	GamePseudo,
+
+	// When a shader is referenced by multiple plugins then it has no natural home for assigning
+	// its size. Instead we create a set of shader pseudo plugins based on the set of root plugins
+	// referencing the shader, including possibly an "Unrooted" plugin.
+	ShaderPseudo
+};
+
 /** The name and dependency information for a plugin that was enabled during cooking. */
 struct COOKMETADATA_API FCookMetadataPluginEntry
 {
 	FString Name;
 
+	ECookMetadataPluginType Type;
+
 	// These contain values pulled from the uplugin json file and hold fields that are not
 	// part of the engine FPluginDescriptor. They are for per-project values. The keys for the maps
 	// are indices in to FCookMetadataPluginHierarchy::CustomFieldNames. See the comment for CustomFieldNames.
@@ -130,13 +153,13 @@ struct COOKMETADATA_API FCookMetadataPluginEntry
 	// and this is an index into it. From there you can get a further index into PluginsEnabledAtCook
 	// to get the plugin information.
 	// Example:
-	//	for (uint16 DependencyIndex = Plugin->DependencyIndexStart; DependencyIndex < Plugin->DependencyIndexEnd; DependencyIndex++)
+	//	for (uint32 DependencyIndex = Plugin->DependencyIndexStart; DependencyIndex < Plugin->DependencyIndexEnd; DependencyIndex++)
 	//	{
 	//		const UE::Cook::FCookMetadataPluginEntry& DependentPlugin = PluginHierarchy.PluginsEnabledAtCook[PluginHierarchy.PluginDependencies[DependencyIndex]];
 	//	}
 	//
-	uint16 DependencyIndexStart = 0;
-	uint16 DependencyIndexEnd = 0;
+	uint32 DependencyIndexStart = 0;
+	uint32 DependencyIndexEnd = 0;
 
 	//
 	// Theses sizes are set during staging by unrealpak if the option in project packaging is set.
@@ -146,7 +169,7 @@ struct COOKMETADATA_API FCookMetadataPluginEntry
 	FPluginSizeInfo InclusiveSizes;
 	FPluginSizeInfo ExclusiveSizes;
 
-	uint16 DependencyCount() const { return DependencyIndexEnd - DependencyIndexStart; }
+	uint32 DependencyCount() const { return DependencyIndexEnd - DependencyIndexStart; }
 
 	friend FArchive& operator<<(FArchive& Ar, FCookMetadataPluginEntry& Entry)
 	{
@@ -210,6 +233,39 @@ struct COOKMETADATA_API FCookMetadataPluginHierarchy
 	}
 };
 
+/**
+*	After staging when sizes are written back we assign the sizes of shaders
+*	to these fake assets and we expose a dependency list for "normal" packages
+*	that use them. This can be used to determine an inclusive size for packages
+*	that also considers shader sizes.
+*/
+struct COOKMETADATA_API FCookMetadataShaderPseudoAsset
+{
+	// This is artificially generated based on the chunk the shader belongs to
+	// and its hash. It should be consistent across builds. The plugin the shader
+	// is assigned to is based on the packages that reference the shader. If the
+	// shader is only referenced by packages within a single plugin, the shader asset
+	// will be assigned to that plugin. Otherwise, it will be assigned to a shader
+	// pseudo plugin.
+	FString Name;
+
+	// Shaders are always compressed, independent of what GetSizesPresent() says.
+	uint32 CompressedSize = 0;
+};
+
+struct COOKMETADATA_API FCookMetadataShaderPseudoHierarchy
+{
+	TArray<FCookMetadataShaderPseudoAsset> ShaderAssets;
+
+	// Entries in this are indices in to ShaderAssets. Use PackageShaderDependencyMap to
+	// get a package's shaders.
+	TArray<int32> DependencyList;
+
+	// Keyed off a PackageName in the project, returns a [start, end) pair of indices
+	// into DependencyList for the shaders that package depends on.
+	TMap<FName, TPair<int32, int32>> PackageShaderDependencyMap;
+};
+
 /**
 *	Structure serialized to disk to contain non-asset related metadata about a cook. This should always
 *	exist alongside a Development Asset Registry, and to ensure that the pair is not out of sync, users
@@ -278,11 +334,16 @@ public:
 	FText GetSizesPresentAsText() const;
 	ECookMetadataSizesPresent GetSizesPresent() const { return SizesPresent; }
 	void SetSizesPresent(ECookMetadataSizesPresent InSizesPresent) { SizesPresent = InSizesPresent; }
+
+	void SetShaderPseudoHieararchy(FCookMetadataShaderPseudoHierarchy&& InHierarchy) { ShaderPseudoHierarchy = MoveTemp(InHierarchy); }
+	const FCookMetadataShaderPseudoHierarchy& GetShaderPseudoHierarchy() const { return ShaderPseudoHierarchy; }
 private:
 
 	ECookMetadataStateVersion Version = ECookMetadataStateVersion::InvalidVersion;
 	FCookMetadataPluginHierarchy PluginHierarchy;
 
+	FCookMetadataShaderPseudoHierarchy ShaderPseudoHierarchy;
+
 	uint64 AssociatedDevelopmentAssetRegistryHash = 0;
 
 	// Asset registry size writeback changes the AR, so we have a separate hash for that DevAR that this
diff --git a/Engine/Source/Developer/IoStoreUtilities/Private/IoStoreUtilities.cpp b/Engine/Source/Developer/IoStoreUtilities/Private/IoStoreUtilities.cpp
--- a/Engine/Source/Developer/IoStoreUtilities/Private/IoStoreUtilities.cpp
+++ b/Engine/Source/Developer/IoStoreUtilities/Private/IoStoreUtilities.cpp
@@ -9,6 +9,7 @@
 #include "HAL/FileManager.h"
 #include "HAL/PlatformFileManager.h"
 #include "Hash/CityHash.h"
+#include "Hash/xxhash.h"
 #include "Interfaces/ITargetPlatformManagerModule.h"
 #include "Interfaces/ITargetPlatform.h"
 #include "IO/IoDispatcher.h"
@@ -2128,7 +2129,49 @@ static bool ConvertToIoStoreShaderLibrary(
 	return true;
 }
 
-static void ProcessShaderLibraries(const FIoStoreArguments& Arguments, TArray<FContainerTargetSpec*>& ContainerTargets, TArray<FShaderInfo*> OutShaders)
+// Carries association between shaders and packages from shader processing to size assignment.
+struct FShaderAssociationInfo
+{
+	struct FShaderChunkInfo
+	{
+		enum EType 
+		{
+			// This shader is referenced by one or more packages and can assign its size to
+			// those packages
+			Package, 
+			// This shader is needed by the baseline engine in some manner and is a flat
+			// size cost.
+			Global, 
+			// This shader isn't global or referenced by packages - theoretically shouldn't
+			// exist.
+			Orphan
+		};
+
+		EType Type;
+		TArray<FName> ReferencedByPackages;
+		uint32 CompressedSize = 0;
+		UE::Cook::EPluginSizeTypes SizeType;
+
+		// If we are shared across plugins, this is our pseudo plugin name. This is generated during size assignment.
+		FString SharedPluginName;
+	};
+	
+	struct FShaderChunkInfoKey
+	{
+		FName PakChunkName;
+		FIoChunkId IoChunkId;
+
+		friend uint32 GetTypeHash(const FShaderChunkInfoKey& Key) { return GetTypeHash(Key.IoChunkId); }
+		friend bool operator == (const FShaderChunkInfoKey& LHS, const FShaderChunkInfoKey& RHS) { return LHS.PakChunkName == RHS.PakChunkName && LHS.IoChunkId == RHS.IoChunkId; }
+	};
+
+	TMap<FShaderChunkInfoKey, FShaderChunkInfo> ShaderChunkInfos;
+
+	// The list of shaders referenced by each package. This can be used to look up in to ContainerShaderInfos.
+	TMap<FName /* PackageName */, TArray<FShaderChunkInfoKey>> PackageShaderMap;
+};
+
+static void ProcessShaderLibraries(const FIoStoreArguments& Arguments, TArray<FContainerTargetSpec*>& ContainerTargets, TArray<FShaderInfo*> OutShaders, FShaderAssociationInfo& OutAssocInfo)
 {
 	IOSTORE_CPU_SCOPE(ProcessShaderLibraries);
 
@@ -2295,25 +2338,67 @@ static void ProcessShaderLibraries(const FIoStoreArguments& Arguments, TArray<FC
 						}
 					}
 				}
+
 				for (FShaderInfo* ShaderInfo : *FindContainerShaderLibraryShaders)
 				{
+					FShaderAssociationInfo::FShaderChunkInfoKey ShaderChunkInfoKey = { ContainerTarget->Name, ShaderInfo->ChunkId };
+					FShaderAssociationInfo::FShaderChunkInfo& ShaderChunkInfo = OutAssocInfo.ShaderChunkInfos.Add(ShaderChunkInfoKey);
+
 					FShaderInfo::EShaderType* ShaderType = ShaderInfo->TypeInContainer.Find(ContainerTarget);
 					check(ShaderType);
 					if (*ShaderType == FShaderInfo::Global)
 					{
 						ContainerTarget->GlobalShaders.Add(ShaderInfo);
+						ShaderChunkInfo.Type = FShaderAssociationInfo::FShaderChunkInfo::Global;
 					}
 					else if (*ShaderType == FShaderInfo::Inline)
 					{
 						ContainerTarget->InlineShaders.Add(ShaderInfo);
+
+						ShaderChunkInfo.Type = FShaderAssociationInfo::FShaderChunkInfo::Package;
+						checkf(ShaderInfo->ReferencedByPackages.Num() == 1, TEXT("Inline shader chunks must be referenced by 1 package only, but shader chunk %s is referenced by %d"),
+							*LexToString(ShaderInfo->ChunkId), ShaderInfo->ReferencedByPackages.Num());
+
+						ShaderChunkInfo.ReferencedByPackages.Add(ShaderInfo->ReferencedByPackages[FSetElementId::FromInteger(0)]->PackageName);
 					}
 					else if (ShaderInfo->ReferencedByPackages.Num() > 1)
 					{
 						ContainerTarget->SharedShaders.Add(ShaderInfo);
+
+						ShaderChunkInfo.Type = FShaderAssociationInfo::FShaderChunkInfo::Package;
+						
+						for (FCookedPackage* Package : ShaderInfo->ReferencedByPackages)
+						{
+							ShaderChunkInfo.ReferencedByPackages.Add(Package->PackageName);
+
+							TArray<FShaderAssociationInfo::FShaderChunkInfoKey>& PackageShaders = OutAssocInfo.PackageShaderMap.FindOrAdd(Package->PackageName);
+							PackageShaders.Add(ShaderChunkInfoKey);
+						}
 					}
 					else
 					{
-						// If there are unreferenced shaders they will go in here and be sorted last
+						//
+						// Note that we can get here with shaders that get split off in to another container (e.g. sm6 shaders). 
+						// Since they are in a different pakChunk they can't get inlined or shared. However, they still "belong" to
+						// the referencing packages for association purposes.
+						//
+						if (ShaderInfo->ReferencedByPackages.Num())
+						{
+							ShaderChunkInfo.Type = FShaderAssociationInfo::FShaderChunkInfo::Package;
+
+							for (FCookedPackage* Package : ShaderInfo->ReferencedByPackages)
+							{
+								ShaderChunkInfo.ReferencedByPackages.Add(Package->PackageName);
+
+								TArray<FShaderAssociationInfo::FShaderChunkInfoKey>& PackageShaders = OutAssocInfo.PackageShaderMap.FindOrAdd(Package->PackageName);
+								PackageShaders.Add(ShaderChunkInfoKey);
+							}
+						}
+						else
+						{
+							ShaderChunkInfo.Type = FShaderAssociationInfo::FShaderChunkInfo::Orphan;
+						}
+
 						ContainerTarget->UniqueShaders.Add(ShaderInfo);
 					}
 					AddShaderTargetFile(ShaderInfo);
@@ -3223,6 +3308,8 @@ enum class EPluginGraphSizeClass : uint8
 	StaticMesh,
 	SoundWave,
 	SkeletalMesh,
+	Shader,
+	Other,
 	COUNT
 };
 
@@ -3232,7 +3319,9 @@ static const UTF8CHAR* PluginGraphEntryClassNames[] =
 	UTF8TEXT("texture"),
 	UTF8TEXT("staticmesh"),
 	UTF8TEXT("soundwave"),
-	UTF8TEXT("skeletalmesh")
+	UTF8TEXT("skeletalmesh"),
+	UTF8TEXT("shader"),
+	UTF8TEXT("other")
 };
 
 static_assert( UE_ARRAY_COUNT(PluginGraphEntryClassNames) == (size_t)EPluginGraphSizeClass::COUNT, "Must have a name for each plugin graph size class!");
@@ -3252,40 +3341,358 @@ static bool SavePluginMetadata(const FString& InAssetRegistryFileName, const FSt
 struct FPluginGraphEntry
 {
 	uint16 IndexInEnabledPlugins = 0;
-	const UE::Cook::FCookMetadataPluginEntry* Self = nullptr;
-	TArray<const UE::Cook::FCookMetadataPluginEntry*> Dependencies;
-	TSet<const UE::Cook::FCookMetadataPluginEntry*> TotalDependencies;
+
+	FString Name;
+
+	TSet<FPluginGraphEntry*> DirectDependencies;
+	TSet<FPluginGraphEntry*> TotalDependencies;		
+	TSet<FPluginGraphEntry*> Roots;
+
+	// Only valid if bIsRoot
+	TSet<FPluginGraphEntry*> UniqueDependencies;
+
 	uint32 DirectRefcount = 0;
 	bool bIsRoot = false;
 
-	// The list of root plugins this plugin can trace a route from
-	TArray<const UE::Cook::FCookMetadataPluginEntry*> Roots;
-
 	static constexpr uint8 ClassCount = (uint8)EPluginGraphSizeClass::COUNT;
 	UE::Cook::FPluginSizeInfo ExclusiveSizes[ClassCount];
 	UE::Cook::FPluginSizeInfo InclusiveSizes[ClassCount];
 	UE::Cook::FPluginSizeInfo UniqueSizes[ClassCount];
 };
 
+struct FPluginGraph
+{
+	TArray<FPluginGraphEntry> Plugins;
+
+	TMap<FStringView, FPluginGraphEntry*> NameToPlugin;
+	TArray<FPluginGraphEntry*> TopologicallySortedPlugins;
+	TArray<FPluginGraphEntry*> RootPlugins;
+
+	// Plugins that can't trace a route between a plugin with bIsRoot==true and themselves.
+	TSet<FPluginGraphEntry*> UnrootedPlugins;
+};
+
+
 // Rework the hierarchy in to a graph where we have output edges resolved to pointers so we can pass to
 // library functions.
-static void GeneratePluginGraph(const UE::Cook::FCookMetadataPluginHierarchy& InPluginHierarchy, TMap<FStringView, FPluginGraphEntry>& OutPluginGraph)
+static void GeneratePluginGraph(const UE::Cook::FCookMetadataPluginHierarchy& InPluginHierarchy, FPluginGraph& OutPluginGraph)
 {
+	double GeneratePluginGraphStart = FPlatformTime::Seconds();
+
+	// Allocate up front so our pointer remains stable.
+	OutPluginGraph.Plugins.Reserve(InPluginHierarchy.PluginsEnabledAtCook.Num());
 	uint16 PluginIndex = 0;
 	for (const UE::Cook::FCookMetadataPluginEntry& Plugin : InPluginHierarchy.PluginsEnabledAtCook)
 	{
-		FPluginGraphEntry& OurEntry = OutPluginGraph.FindOrAdd(Plugin.Name);
+		FPluginGraphEntry& OurEntry = OutPluginGraph.Plugins.AddDefaulted_GetRef();
 		OurEntry.IndexInEnabledPlugins = PluginIndex;
-		OurEntry.Self = &Plugin;
+		OurEntry.Name = Plugin.Name;
 
-		for (uint16 DependencyIndex = Plugin.DependencyIndexStart; DependencyIndex < Plugin.DependencyIndexEnd; DependencyIndex++)
+		// Can store pointer since we reserved as a batch..
+		OutPluginGraph.NameToPlugin.Add(OurEntry.Name, &OurEntry);
+		PluginIndex++;
+	}
+
+	OutPluginGraph.RootPlugins.Reserve(InPluginHierarchy.RootPlugins.Num());
+	for (uint16 RootIndex : InPluginHierarchy.RootPlugins)
+	{
+		FPluginGraphEntry* Root = OutPluginGraph.NameToPlugin[InPluginHierarchy.PluginsEnabledAtCook[RootIndex].Name];
+		Root->bIsRoot = true;
+		OutPluginGraph.RootPlugins.Add(Root);
+	}
+
+	for (FPluginGraphEntry& PluginEntry : OutPluginGraph.Plugins)
+	{
+		const UE::Cook::FCookMetadataPluginEntry& Plugin = InPluginHierarchy.PluginsEnabledAtCook[PluginEntry.IndexInEnabledPlugins];
+
+		for (uint32 DependencyIndex = Plugin.DependencyIndexStart; DependencyIndex < Plugin.DependencyIndexEnd; DependencyIndex++)
 		{
 			const UE::Cook::FCookMetadataPluginEntry& DependentPlugin = InPluginHierarchy.PluginsEnabledAtCook[InPluginHierarchy.PluginDependencies[DependencyIndex]];
+			PluginEntry.DirectDependencies.Add(OutPluginGraph.NameToPlugin[DependentPlugin.Name]);
+		}
+	}
+	
+	// From here on out we can operate entirely on our own data - no cook metadata structures - 
+	// and we generate the various structures we need.
 
-			OurEntry.Dependencies.Add(&DependentPlugin);
+	// Sort the plugins topologically. This means that when we iterate linearly,
+	// we know that when we hit a plugin, we've already processed the dependencies.
+	// This takes some memory to track edges but is a depth first search
+	// and not anything quadratic or worse.	
+	double TopologicalSortStart = FPlatformTime::Seconds();
+	{
+		OutPluginGraph.TopologicallySortedPlugins.Reserve(OutPluginGraph.Plugins.Num());
+		for (FPluginGraphEntry& Plugin : OutPluginGraph.Plugins)
+		{
+			OutPluginGraph.TopologicallySortedPlugins.Add(&Plugin);
+		}
+
+		auto GetElementDependencies = [&OutPluginGraph](const FPluginGraphEntry* PluginEntry) -> const TSet<FPluginGraphEntry*>&
+		{
+			return OutPluginGraph.NameToPlugin[PluginEntry->Name]->DirectDependencies;
+		};
+
+		Algo::TopologicalSort(OutPluginGraph.TopologicallySortedPlugins, GetElementDependencies);
+	}
+
+
+	// Gather the set of all dependencies. This ends up being technically
+	// O(N^2) in the worst case. It's highly unlikely our plugin DAG will cause that, but we track the times
+	// just so we can keep an eye on it if it ends up taking measurable amounts of time.
+	double InclusiveComputeStart = FPlatformTime::Seconds();
+	for (FPluginGraphEntry* Plugin : OutPluginGraph.TopologicallySortedPlugins)
+	{
+		for (FPluginGraphEntry* Dependency : Plugin->DirectDependencies)
+		{
+			Plugin->TotalDependencies.Add(Dependency);
+			Dependency->DirectRefcount++;
+		
+			// In the worse case this is another O(N) iteration, which makes us overall O(N^2)
+			for (FPluginGraphEntry* TotalDependencyEntry : Dependency->TotalDependencies)
+			{
+				Plugin->TotalDependencies.Add(TotalDependencyEntry);
+			}
+		}
+	}
+	double InclusiveComputeEnd = FPlatformTime::Seconds();
+
+	// Generate the unique dependencies for the root plugins. This is the set of dependencies that only
+	// belong to the root plugin and not to another. These dependencies could be referred to by another 
+	// plugin within the unique set - the only requirement is that there exists no path from _another_
+	// root to the dependency.
+	for (FPluginGraphEntry* RootPlugin : OutPluginGraph.RootPlugins)
+	{
+		// Add us as a root entry for all our total dependencies so all plugins know which root they are in.
+		for (FPluginGraphEntry* Dependency : RootPlugin->TotalDependencies)
+		{
+			OutPluginGraph.NameToPlugin[Dependency->Name]->Roots.Add(RootPlugin);
+		}
+
+		// Duplicate the TotalDependencies and then remove any dependency that
+		// exists for another root.
+		RootPlugin->UniqueDependencies = RootPlugin->TotalDependencies;
+		for (FPluginGraphEntry* InnerRootPlugin : OutPluginGraph.RootPlugins)
+		{
+			if (RootPlugin == InnerRootPlugin)
+			{
+				continue;
+			}
+
+			for (FPluginGraphEntry* InnerRootDependency : InnerRootPlugin->TotalDependencies)
+			{
+				RootPlugin->UniqueDependencies.Remove(InnerRootDependency);
+			}
 		}
-		PluginIndex++;
 	}
+
+	double UniqueEnd = FPlatformTime::Seconds();
+
+	// Generate the unrooted set. These plugins can not trace a path from _any_ root plugin
+	// to themselves. For projects with no root plugins, this will be all plugins.
+	{
+		for (FPluginGraphEntry* Plugin : OutPluginGraph.TopologicallySortedPlugins)
+		{
+			OutPluginGraph.UnrootedPlugins.Add(Plugin);
+		}
+
+		for (FPluginGraphEntry* Plugin : OutPluginGraph.RootPlugins)
+		{
+			OutPluginGraph.UnrootedPlugins.Remove(Plugin);
+			for (FPluginGraphEntry* Dependency : Plugin->TotalDependencies)
+			{
+				OutPluginGraph.UnrootedPlugins.Remove(Plugin);
+			}
+		}
+	}
+	double GeneratePluginGraphEnd = FPlatformTime::Seconds();
+
+	UE_LOG(LogIoStore, Log, TEXT("Generated plugin graph with %d nodes. Times: %.02f total %.02f setup, %.02f sort %.02f inclusive %.02f unique %.02f unrooted."),
+		OutPluginGraph.TopologicallySortedPlugins.Num(),
+		GeneratePluginGraphEnd - GeneratePluginGraphStart,
+		TopologicalSortStart - GeneratePluginGraphStart,
+		InclusiveComputeStart - TopologicalSortStart,
+		InclusiveComputeEnd - InclusiveComputeStart,
+		UniqueEnd - InclusiveComputeEnd,
+		GeneratePluginGraphEnd - UniqueEnd);
+}
+
+static void InsertShadersInPluginHierarchy(UE::Cook::FCookMetadataState& InCookMetadata, FPluginGraph& InPluginGraph, FShaderAssociationInfo& InShaderAssociationInfo)
+{
+	const UE::Cook::FCookMetadataPluginHierarchy& PluginHierarchy = InCookMetadata.GetPluginHierarchy();
+
+	//
+	// Create any shader plugins we need. These are pseudo plugins that we create to hold the size information
+	// when the packages that reference a plugin cross the plugin boundary. We name them based on their dependencies
+	// so it's consistent across builds. These will exist whenever GFPs aren't placed entirely in their own pak chunk.
+	//
+	// The combinatorics are such that doing this for _all_ plugins isn't tenable. However, for product tracking we
+	// actually only care about root GFPs. So instead of gathering all of the plugins entirely, we gather all of the
+	// root plugins.
+	//
+	// The difficulty is that we don't necessarily _have_ any root plugins if the project hasn't defined any, so we 
+	// artificially stuff such plugins under "Unrooted".
+	//
+	// It should be noted that the entire point of root GFPs is to separate the data entirely - so if we have any
+	// of these pseudo plugins then there is a content bug as there exists a shared dependency between two "modes".
+	//
+	TMap<FString, TArray<FShaderAssociationInfo::FShaderChunkInfoKey>> ShaderPseudoPlugins;
+	TMap<FString, TArray<FString>> PluginDependenciesOnShaders;
+
+	TSet<FString> ShaderRootPossibles;
+	ShaderRootPossibles.Add(TEXT("Unrooted"));
+
+	for (TPair<FShaderAssociationInfo::FShaderChunkInfoKey, FShaderAssociationInfo::FShaderChunkInfo>& ShaderChunkInfo : InShaderAssociationInfo.ShaderChunkInfos)
+	{
+		// Only Normal shaders can be assigned
+		if (ShaderChunkInfo.Value.Type != FShaderAssociationInfo::FShaderChunkInfo::Package)
+		{
+			continue;
+		}
+
+		TSet<FString> ShaderPlugins;
+		for (FName PackageName : ShaderChunkInfo.Value.ReferencedByPackages)
+		{
+			FString PackageNameStr = PackageName.ToString();
+			FStringView Plugin = FPackageName::SplitPackageNameRoot(PackageNameStr, nullptr);
+			ShaderPlugins.Add(FString(Plugin));
+		}
+
+		if (ShaderPlugins.Num() == 1)
+		{
+			// We can assign the size to this plugin when the time comes - no pseudo plugin needed.
+			continue;
+		}
+
+		bool bAllReferencingPluginsAreRooted = true;
+		TSet<FString> ShaderRootPlugins;
+		for (FString& ReferencingPluginName : ShaderPlugins)
+		{
+			bool bReferencingPluginIsRooted = false;
+			FPluginGraphEntry** ReferencingPlugin = InPluginGraph.NameToPlugin.Find(ReferencingPluginName);
+			if (ReferencingPlugin)
+			{
+				for (FPluginGraphEntry* RootForReferencingPlugin : (*ReferencingPlugin)->Roots)
+				{
+					ShaderRootPlugins.Add(RootForReferencingPlugin->Name);
+					bReferencingPluginIsRooted = true;
+				}
+			}
+				
+			if (bReferencingPluginIsRooted == false)
+				bAllReferencingPluginsAreRooted = false;
+		}
+
+		if (ShaderRootPlugins.Num() == 0 || bAllReferencingPluginsAreRooted == false)
+		{
+			// Place in the unrooted list.
+			ShaderRootPlugins.Add(TEXT("Unrooted"));
+		}
+
+		TStringBuilder<256> PseudoPluginName;
+		PseudoPluginName.Append(TEXT("ShaderPlugin"));
+
+		TStringBuilder<256> NameConcatenation;
+
+		// We can't just concat the names because it'll get too long when we use the name as a filename, which is unfortunate. That being said, we
+		// only expect this to happen in degenerate cases - so if it's not too long we list the names for convenience
+		// otherwise we hash it.
+		for (FString& ReferencingPlugin : ShaderRootPlugins)
+		{
+			NameConcatenation.Append(TEXT("_"));
+			NameConcatenation.Append(ReferencingPlugin);
+		}
+
+		if (NameConcatenation.Len() > 100) // arbitrary length here just to try and avoid hitting MAX_PATH (260)
+		{
+			FXxHash64 NameHash = FXxHash64::HashBuffer(NameConcatenation.GetData(), NameConcatenation.Len());
+			uint8 HashBytes[8];
+			NameHash.ToByteArray(HashBytes);
+
+			PseudoPluginName.Append(TEXT("_"));
+			UE::String::BytesToHexLower(MakeArrayView(HashBytes), PseudoPluginName);
+		}
+		else
+		{
+			PseudoPluginName.Append(NameConcatenation);
+		}
+
+		FString PseudoPluginNameStr = PseudoPluginName.ToString();
+		for (FString& ReferencingPlugin : ShaderRootPlugins)
+		{
+			PluginDependenciesOnShaders.FindOrAdd(ReferencingPlugin).Add(PseudoPluginNameStr);
+		}
+
+		ShaderPseudoPlugins.FindOrAdd(PseudoPluginNameStr).Add(ShaderChunkInfo.Key);
+
+		ShaderChunkInfo.Value.SharedPluginName = MoveTemp(PseudoPluginNameStr);
+	}
+
+	// We have the list of pseudo plugins we need to make... we need to copy and append to the list
+	// of plugins in the cook metadata, after stripping out any previous run's pseudo plugins.
+	TArray<UE::Cook::FCookMetadataPluginEntry> PluginEntries = PluginHierarchy.PluginsEnabledAtCook;
+	for (int32 PluginIndex = 0; PluginIndex < PluginEntries.Num(); PluginIndex++)
+	{
+		if (PluginEntries[PluginIndex].Type == UE::Cook::ECookMetadataPluginType::ShaderPseudo)
+		{
+			// We can do a swap because we insert these at the end in a group so there shouldn't
+			// by anything else after us.
+			PluginEntries.RemoveAtSwap(PluginIndex);
+			PluginIndex--;
+		}
+	}
+
+	for (const TPair< FString, TArray<FShaderAssociationInfo::FShaderChunkInfoKey>>& ShaderPP : ShaderPseudoPlugins)
+	{
+		UE::Cook::FCookMetadataPluginEntry& ShaderPluginEntry = PluginEntries.AddDefaulted_GetRef();
+		ShaderPluginEntry.Name = ShaderPP.Key;
+		ShaderPluginEntry.Type = UE::Cook::ECookMetadataPluginType::ShaderPseudo;
+	}
+
+	// We have to redo the dependency tree as well to add the shaders as dependencies.
+	TMap<FString, int32> PluginNameToIndex;
+	int32 CurrentIndex = 0;
+	for (UE::Cook::FCookMetadataPluginEntry& Entry : PluginEntries)
+	{
+		PluginNameToIndex.Add(Entry.Name, CurrentIndex);
+		CurrentIndex++;
+	}
+
+	if (IntFitsIn<uint16>(PluginEntries.Num()) == false)
+	{
+		UE_LOG(LogIoStore, Warning, TEXT("Post shared shader plugin count is > 65535 (%d)  - not updating cook metadata!"), PluginEntries.Num());
+		return;
+	}
+
+	TArray<uint16> DependencyList;
+	for (UE::Cook::FCookMetadataPluginEntry& Entry : PluginEntries)
+	{
+		// Add the normal dependencies. Since we didn't reorder anything we can
+		// use the old dependency list
+		uint32 StartIndex = (uint32)DependencyList.Num();
+		for (uint32 DependencyIndex = Entry.DependencyIndexStart; DependencyIndex < Entry.DependencyIndexEnd; DependencyIndex++)
+		{
+			const UE::Cook::FCookMetadataPluginEntry& DependentPlugin = PluginHierarchy.PluginsEnabledAtCook[PluginHierarchy.PluginDependencies[DependencyIndex]];
+			DependencyList.Add(PluginNameToIndex[DependentPlugin.Name]);
+		}
+
+		// However we also need to check for shaders
+		TArray<FString>* DependenciesOnShader = PluginDependenciesOnShaders.Find(Entry.Name);
+		if (DependenciesOnShader)
+		{
+			for (FString& ShaderPluginName : (*DependenciesOnShader))
+			{
+				DependencyList.Add(PluginNameToIndex[ShaderPluginName]);
+			}
+		}
+
+		Entry.DependencyIndexStart = StartIndex;
+		Entry.DependencyIndexEnd = (uint32)DependencyList.Num();
+	}
+
+	// Now blast the old one away and replace.
+	UE::Cook::FCookMetadataPluginHierarchy& MutablePluginHierarchy = InCookMetadata.GetMutablePluginHierarchy();
+	MutablePluginHierarchy.PluginDependencies = MoveTemp(DependencyList);
+	MutablePluginHierarchy.PluginsEnabledAtCook = PluginEntries;
 }
 
 /**
@@ -3296,22 +3703,29 @@ static void GeneratePluginGraph(const UE::Cook::FCookMetadataPluginHierarchy& In
 *	to a plugin has either /Engine or /Game as it's top level path and will be assigned to a pseudo
 *	plugin.
 */
-static void WritePluginMetadataJsons(const FString& InAssetRegistryFileName, TMap<FPackageId, TArray<FIoStoreChunkSource, TInlineAllocator<2>>>& PackageToChunks, FAssetRegistryState& AssetRegistry, UE::Cook::FCookMetadataState& CookMetadata)
+static void UpdatePluginMetadataAndWriteJsons(
+	const FString& InAssetRegistryFileName, 
+	TMap<FPackageId, TArray<FIoStoreChunkSource, TInlineAllocator<2>>>& PackageToChunks, 
+	FAssetRegistryState& AssetRegistry, 
+	UE::Cook::FCookMetadataState& CookMetadata, 
+	FShaderAssociationInfo* InShaderAssociationInfo)
 {
 	double WritePluginStart = FPlatformTime::Seconds();
 
-	//
-	// Using the name of the packages, assign chunk sizes to their associated plugin.
-	//
-	const UE::Cook::FCookMetadataPluginHierarchy& PluginHierarchy = CookMetadata.GetPluginHierarchy();
-	TMap<FStringView, FPluginGraphEntry> PluginGraph;
-	GeneratePluginGraph(PluginHierarchy, PluginGraph);
+	FPluginGraph PluginGraph;
+	GeneratePluginGraph(CookMetadata.GetPluginHierarchy(), PluginGraph);
 
-	for (uint16 RootIndex : PluginHierarchy.RootPlugins)
+	if (InShaderAssociationInfo)
 	{
-		PluginGraph[PluginHierarchy.PluginsEnabledAtCook[RootIndex].Name].bIsRoot = true;
+		InsertShadersInPluginHierarchy(CookMetadata, PluginGraph, *InShaderAssociationInfo);
+
+		// Generate the graph aggain after we've inserted the new "plugins".
+		PluginGraph = FPluginGraph();
+		GeneratePluginGraph(CookMetadata.GetPluginHierarchy(), PluginGraph);
 	}
 
+	double GeneratePluginGraphEnd = FPlatformTime::Seconds();
+
 	TSet<FString> LoggedPluginNames;
 
 	FTopLevelAssetPath Texture2DPath(TEXT("/Script/Engine.Texture2D"));
@@ -3323,6 +3737,134 @@ static void WritePluginMetadataJsons(const FString& InAssetRegistryFileName, TMa
 	FTopLevelAssetPath SoundWavePath(TEXT("/Script/Engine.SoundWave"));
 	FTopLevelAssetPath SkeletalMeshPath(TEXT("/Script/Engine.SkeletalMesh"));
 
+	
+	if (InShaderAssociationInfo)
+	{
+		//
+		// Create a bunch of "Assets" that we can iterate over in the same manner as a normal
+		// asset and find its containing "plugin" for the purposes of assigning its size.
+		//
+		TArray<UE::Cook::FCookMetadataShaderPseudoAsset> ShaderPseudoAssets;
+
+		TMap<FName, TArray<int32>> PackageDependencyMap;
+
+		// All package shaders should now either be able to be assigned to:
+		// 1. a single package (i.e. plugin)
+		// 2. shared between packages in 1 plugin (i.e. a plugin)
+		// 3. shared between packages across plugins (i.e. assignable to a pseudo plugin we created earlier).
+		//
+		uint64 CrossPluginShaderSize = 0;
+		uint64 SinglePluginShaderSize = 0;
+		uint64 InlineShaderSize = 0;
+		uint64 GlobalShaderSize = 0;
+		uint64 OrphanShaderSize = 0;
+		for (TPair<FShaderAssociationInfo::FShaderChunkInfoKey, FShaderAssociationInfo::FShaderChunkInfo>& ShaderChunkInfo : InShaderAssociationInfo->ShaderChunkInfos)
+		{
+			if (ShaderChunkInfo.Value.Type == FShaderAssociationInfo::FShaderChunkInfo::Orphan)
+			{
+				OrphanShaderSize += ShaderChunkInfo.Value.CompressedSize;
+				continue;
+			}
+			if (ShaderChunkInfo.Value.Type == FShaderAssociationInfo::FShaderChunkInfo::Global)
+			{
+				GlobalShaderSize += ShaderChunkInfo.Value.CompressedSize;
+				continue;
+			}
+
+			check(ShaderChunkInfo.Value.Type == FShaderAssociationInfo::FShaderChunkInfo::Package);
+
+			TStringBuilder<128> PackageName;
+			if (ShaderChunkInfo.Value.SharedPluginName.Len())
+			{
+				// We know we belong to this plugin.
+				PackageName.Append(ShaderChunkInfo.Value.SharedPluginName);
+				CrossPluginShaderSize += ShaderChunkInfo.Value.CompressedSize;
+			}
+			else
+			{
+				// We know all the plugin prefixes are the same for all referrers
+				if (ShaderChunkInfo.Value.ReferencedByPackages.Num() == 1)
+				{
+					InlineShaderSize += ShaderChunkInfo.Value.CompressedSize;
+				}
+				else
+				{
+					SinglePluginShaderSize += ShaderChunkInfo.Value.CompressedSize;
+				}
+
+				FString PackageNameStr = ShaderChunkInfo.Value.ReferencedByPackages[0].ToString();
+				FStringView Plugin = FPackageName::SplitPackageNameRoot(PackageNameStr, nullptr);
+				PackageName.Append(Plugin);
+			}
+
+			FPluginGraphEntry** PluginEntryPtr = PluginGraph.NameToPlugin.Find(PackageName.ToString());
+			if (PluginEntryPtr)
+			{
+				PluginEntryPtr[0]->ExclusiveSizes[(uint8)EPluginGraphSizeClass::All][ShaderChunkInfo.Value.SizeType] += ShaderChunkInfo.Value.CompressedSize;
+				PluginEntryPtr[0]->ExclusiveSizes[(uint8)EPluginGraphSizeClass::Shader][ShaderChunkInfo.Value.SizeType] += ShaderChunkInfo.Value.CompressedSize;
+			}
+			else
+			{
+				FString AllocatedPluginName(PackageName.ToString());
+				bool bAlreadyLogged = false;
+				LoggedPluginNames.Add(AllocatedPluginName, &bAlreadyLogged);
+				if (bAlreadyLogged == false)
+				{
+					UE_LOG(LogIoStore, Warning, TEXT("Plugin for shader not found: %s"), *AllocatedPluginName);
+				}
+			}
+
+
+			// What to name our package? Needs to be unique. We just concat everything so a human
+			// can trace where it came from.
+			PackageName.Append(TEXT("/ShaderPseudoAsset_"));
+
+			PackageName.Append(ShaderChunkInfo.Key.PakChunkName.ToString());
+			PackageName.Append(TEXT("_"));
+
+			// shader hash (iochunkid)
+			UE::String::BytesToHexLower(MakeArrayView(ShaderChunkInfo.Key.IoChunkId.GetData(), sizeof(FIoChunkId)), PackageName);
+
+			ShaderPseudoAssets.Add({PackageName.ToString(), ShaderChunkInfo.Value.CompressedSize});
+
+			// Track who depends on this shader by index.
+			for (FName ReferencingPackage : ShaderChunkInfo.Value.ReferencedByPackages)
+			{
+				TArray<int32>& PackageDependencies = PackageDependencyMap.FindOrAdd(ReferencingPackage);
+				PackageDependencies.Add(ShaderPseudoAssets.Num());
+			}
+		}
+
+		TMap<FName, TPair<int32, int32>> FinalizedDependencyMap;
+
+		// Now convert the package dependency map in to array ranges.
+		TArray<int32> DependencyByIndex;
+		for (TPair<FName, TArray<int32>>& Dependencies : PackageDependencyMap)
+		{
+			TPair<int32, int32>& Entry = FinalizedDependencyMap.Add(Dependencies.Key);
+			Entry.Key = DependencyByIndex.Num();
+			DependencyByIndex.Append(Dependencies.Value);
+			Entry.Value = DependencyByIndex.Num();
+		}
+
+		// Move the new info over to the cook metadata.
+		UE::Cook::FCookMetadataShaderPseudoHierarchy PSH;
+		PSH.ShaderAssets = MoveTemp(ShaderPseudoAssets);
+		PSH.PackageShaderDependencyMap = MoveTemp(FinalizedDependencyMap);
+		PSH.DependencyList = MoveTemp(DependencyByIndex);
+		CookMetadata.SetShaderPseudoHieararchy(MoveTemp(PSH));
+
+		double TotalShaderSize = (double)(InlineShaderSize + CrossPluginShaderSize + SinglePluginShaderSize + GlobalShaderSize + OrphanShaderSize);
+		UE_LOG(LogIoStore, Display, TEXT("Shader total sizes: %s single package (%.0f%%), %s single root GFP (%.0f%%), %s cross root GFP (%.0f%%), %s global (%.0f%%), %s orphan (%.0f%%) - %s assigned (%.0f%%)"),
+			*NumberString(InlineShaderSize), 100.0 * InlineShaderSize / TotalShaderSize,
+			*NumberString(SinglePluginShaderSize), 100.0 * SinglePluginShaderSize / TotalShaderSize,
+			*NumberString(CrossPluginShaderSize), 100.0 * CrossPluginShaderSize / TotalShaderSize,
+			*NumberString(GlobalShaderSize), 100.0 * GlobalShaderSize / TotalShaderSize,
+			*NumberString(OrphanShaderSize), 100.0 * OrphanShaderSize / TotalShaderSize,
+			*NumberString(InlineShaderSize + SinglePluginShaderSize + CrossPluginShaderSize), 100.0 * (InlineShaderSize + SinglePluginShaderSize + CrossPluginShaderSize) / TotalShaderSize
+		);
+	}
+
 	double AssetPackageMapStart = FPlatformTime::Seconds();
 	const TMap<FName, const FAssetPackageData*> AssetPackageMap = AssetRegistry.GetAssetPackageDataMap();
 	for (const TPair<FName, const FAssetPackageData*>& AssetPackage : AssetPackageMap)
@@ -3356,9 +3898,10 @@ static void WritePluginMetadataJsons(const FString& InAssetRegistryFileName, TMa
 			FStringView PackageName(PackageNameStr);
 
 			FStringView PluginName = FPackageName::SplitPackageNameRoot(PackageName, nullptr);
-			FPluginGraphEntry* PluginEntry = PluginGraph.Find(PluginName);
-			if (PluginEntry)
+			FPluginGraphEntry** PluginEntryPtr = PluginGraph.NameToPlugin.Find(PluginName);
+			if (PluginEntryPtr)
 			{
+				FPluginGraphEntry* PluginEntry = *PluginEntryPtr;
 				PluginEntry->ExclusiveSizes[(uint8)EPluginGraphSizeClass::All].Add(PackageSizes);
 
 				// If we have asset class info and it's a top contender, track it also.
@@ -3384,6 +3927,11 @@ static void WritePluginMetadataJsons(const FString& InAssetRegistryFileName, TMa
 					{
 						PluginEntry->ExclusiveSizes[(uint8)EPluginGraphSizeClass::SkeletalMesh].Add(PackageSizes);
 					}
+					// Note that we can't get shaders here so we don't need to handle ::Shader.
+					else
+					{
+						PluginEntry->ExclusiveSizes[(uint8)EPluginGraphSizeClass::Other].Add(PackageSizes);
+					}
 				}
 			}
 			else
@@ -3399,180 +3947,49 @@ static void WritePluginMetadataJsons(const FString& InAssetRegistryFileName, TMa
 		}
 	}
 
-	//
-	// Generate the inclusive size for each plugin.
-	//
-	{
-		// This is so re-staging the same cook is consistent.
-		UE::Cook::FCookMetadataPluginHierarchy& MutablePluginHierarchy = CookMetadata.GetMutablePluginHierarchy();
-		for (UE::Cook::FCookMetadataPluginEntry& Plugin : MutablePluginHierarchy.PluginsEnabledAtCook)
-		{
-			Plugin.InclusiveSizes.Zero();
-			Plugin.ExclusiveSizes.Zero();
-		}
-	}
-
-	// Sort the plugins topologically. This means that when we iterate linearly,
-	// we know that when we hit a plugin, we've already processed the dependencies.
-	// This takes some memory to track edges but is a depth first search
-	// and not anything quadratic or worse.
-	double TopologicalSortStart = FPlatformTime::Seconds();
-	TArray<const UE::Cook::FCookMetadataPluginEntry*> SortedList;
+	// Inclusive is the sum of us plus all our dependencies.
+	for (FPluginGraphEntry* PluginEntry : PluginGraph.TopologicallySortedPlugins)
 	{
-		for (const UE::Cook::FCookMetadataPluginEntry& Plugin : PluginHierarchy.PluginsEnabledAtCook)
-		{
-			SortedList.Add(&Plugin);
-		}
-
-		auto GetElementDependencies = [&PluginGraph](const UE::Cook::FCookMetadataPluginEntry* PluginEntry) -> const TArray<const UE::Cook::FCookMetadataPluginEntry*>&
-		{
-			return PluginGraph[PluginEntry->Name].Dependencies;
-		};
-
-		Algo::TopologicalSort(SortedList, GetElementDependencies);
-
-		// Make sure the topological sort worked correctly.
-#if DO_CHECK
-		TSet<uint16> Test;
-		for (const UE::Cook::FCookMetadataPluginEntry* Check : SortedList)
-		{
-			Test.Add(PluginGraph[Check->Name].IndexInEnabledPlugins);
-
-			for (uint16 DependencyIndex = Check->DependencyIndexStart; DependencyIndex < Check->DependencyIndexEnd; DependencyIndex++)
-			{
-				const UE::Cook::FCookMetadataPluginEntry& DependentPlugin = PluginHierarchy.PluginsEnabledAtCook[PluginHierarchy.PluginDependencies[DependencyIndex]];
-
-				check (Test.Contains(PluginHierarchy.PluginDependencies[DependencyIndex]));
-			}
-		}
-#endif
-	}
-
-	// For inclusive sizes we need to gather the set of all dependencies. This ends up being technically
-	// O(N^2) in the worst case. It's highly unlikely our plugin DAG will cause that, but we track the times
-	// just so we can keep an eye on it if it ends up taking measurable amounts of time.
-	double InclusiveComputeStart = FPlatformTime::Seconds();
-	for (const UE::Cook::FCookMetadataPluginEntry* Plugin : SortedList)
-	{
-		FPluginGraphEntry& PluginEntry = PluginGraph[Plugin->Name];
-
 		for (uint8 ClassIndex = 0; ClassIndex < FPluginGraphEntry::ClassCount; ClassIndex++)
 		{
-			PluginEntry.InclusiveSizes[ClassIndex] = PluginEntry.ExclusiveSizes[ClassIndex];
+			PluginEntry->InclusiveSizes[ClassIndex] = PluginEntry->ExclusiveSizes[ClassIndex];
 		}
 
-		for (uint16 DependencyIndex = Plugin->DependencyIndexStart; DependencyIndex < Plugin->DependencyIndexEnd; DependencyIndex++)
+		for (FPluginGraphEntry* Dependency : PluginEntry->TotalDependencies)
 		{
-			const UE::Cook::FCookMetadataPluginEntry& DependentPlugin = PluginHierarchy.PluginsEnabledAtCook[PluginHierarchy.PluginDependencies[DependencyIndex]];
-
-			bool bAlreadyInSet = false;
-			PluginEntry.TotalDependencies.Add(&DependentPlugin, &bAlreadyInSet);
-
-			if (bAlreadyInSet == false)
-			{
-				const FPluginGraphEntry* DependencyGraphEntry = &PluginGraph[DependentPlugin.Name];
-				for (uint8 ClassIndex = 0; ClassIndex < FPluginGraphEntry::ClassCount; ClassIndex++)
-				{
-					PluginEntry.InclusiveSizes[ClassIndex].Add(DependencyGraphEntry->ExclusiveSizes[ClassIndex]);
-				}
-				
-			}
-			
-			FPluginGraphEntry& DependentEntry = PluginGraph[DependentPlugin.Name];
-			DependentEntry.DirectRefcount++;
-
-			// In the worse case this is another O(N) iteration, which makes us overall O(N^2)
-			for (const UE::Cook::FCookMetadataPluginEntry* TotalDependencyEntry : DependentEntry.TotalDependencies)
+			for (uint8 ClassIndex = 0; ClassIndex < FPluginGraphEntry::ClassCount; ClassIndex++)
 			{
-				PluginEntry.TotalDependencies.Add(TotalDependencyEntry, &bAlreadyInSet);
-				if (bAlreadyInSet == false)
-				{
-					const FPluginGraphEntry* DependencyGraphEntry = &PluginGraph[TotalDependencyEntry->Name];
-					for (uint8 ClassIndex = 0; ClassIndex < FPluginGraphEntry::ClassCount; ClassIndex++)
-					{
-						PluginEntry.InclusiveSizes[ClassIndex].Add(DependencyGraphEntry->ExclusiveSizes[ClassIndex]);
-					}
-				}
-			}
+				PluginEntry->InclusiveSizes[ClassIndex].Add(Dependency->ExclusiveSizes[ClassIndex]);
+			}			
 		}
 	}
-	double InclusiveComputeEnd = FPlatformTime::Seconds();
 
 	// Now we need to find the unique size for each root plugin. This is the size of dependencies that only
 	// belong to the root plugin and not to another. Conceptually this is the "assuming all other roots are
-	// installed, this is the size cost to add this plugin to the install". These dependencies could be referred
-	// to by another plugin within the unique set - the only requirement is that there exists no path from _another_
-	// root to the dependency.
-	for (uint16 RootIndex : PluginHierarchy.RootPlugins)
+	// installed, this is the size cost to add this plugin to the install".
+	for (FPluginGraphEntry* RootPlugin : PluginGraph.RootPlugins)
 	{
-		FPluginGraphEntry& RootPluginEntry = PluginGraph[PluginHierarchy.PluginsEnabledAtCook[RootIndex].Name];
-
-		// Add us as a root entry for all our total dependencies so all plugins know which root they are in.
-		for (const UE::Cook::FCookMetadataPluginEntry* Dependency : RootPluginEntry.TotalDependencies)
-		{
-			PluginGraph[Dependency->Name].Roots.Add(RootPluginEntry.Self);
-		}
-
-		// Duplicate the TotalDependencies and then remove any dependency that
-		// exists for another root.
-		TSet<const UE::Cook::FCookMetadataPluginEntry*> UniqueDependencies = RootPluginEntry.TotalDependencies;
-
-		for (uint16 InnerRootIndex : PluginHierarchy.RootPlugins)
-		{
-			if (RootIndex == InnerRootIndex)
-			{
-				continue;
-			}
-
-			FPluginGraphEntry& InnerRootPluginEntry = PluginGraph[PluginHierarchy.PluginsEnabledAtCook[InnerRootIndex].Name];
-			for (const UE::Cook::FCookMetadataPluginEntry* InnerRootDependency : InnerRootPluginEntry.TotalDependencies)
-			{
-				UniqueDependencies.Remove(InnerRootDependency);
-			}
-		}
-
-		// Sum the exclusive size of the unique set to get the unique size.
-		for (uint8 ClassIndex = 0; ClassIndex < FPluginGraphEntry::ClassCount; ClassIndex++)
-		{
-			RootPluginEntry.UniqueSizes[ClassIndex].Zero();
-		}
-
-		for (const UE::Cook::FCookMetadataPluginEntry* UniqueDependency : UniqueDependencies)
+		for (FPluginGraphEntry* UniqueDependency : RootPlugin->UniqueDependencies)
 		{
 			for (uint8 ClassIndex = 0; ClassIndex < FPluginGraphEntry::ClassCount; ClassIndex++)
 			{
-				RootPluginEntry.UniqueSizes[ClassIndex].Add(PluginGraph[UniqueDependency->Name].ExclusiveSizes[ClassIndex]);
+				RootPlugin->UniqueSizes[ClassIndex].Add(UniqueDependency->ExclusiveSizes[ClassIndex]);
 			}
 		}
 	}
 	
 	// Find the total size of all plugins that aren't rooted in the root set.
 	UE::Cook::FPluginSizeInfo UnrootedTotal;
-	TSet<const UE::Cook::FCookMetadataPluginEntry*> UnrootedPlugins;
+	for (FPluginGraphEntry* Plugin : PluginGraph.UnrootedPlugins)
 	{		
-		for (const UE::Cook::FCookMetadataPluginEntry* Plugin : SortedList)
-		{
-			UnrootedPlugins.Add(Plugin);
-		}
-
-		for (uint16 RootIndex : PluginHierarchy.RootPlugins)
-		{
-			UnrootedPlugins.Remove(&PluginHierarchy.PluginsEnabledAtCook[RootIndex]);
+		UnrootedTotal.Add(Plugin->ExclusiveSizes[(uint8)EPluginGraphSizeClass::All]);
+	}
 
-			FPluginGraphEntry& PluginEntry = PluginGraph[PluginHierarchy.PluginsEnabledAtCook[RootIndex].Name];
-			for (const UE::Cook::FCookMetadataPluginEntry* Plugin : PluginEntry.TotalDependencies)
-			{
-				UnrootedPlugins.Remove(Plugin);
-			}
-		}
+	double WriteBegin = FPlatformTime::Seconds();
 
-		for (const UE::Cook::FCookMetadataPluginEntry* Plugin : UnrootedPlugins)
-		{
-			UnrootedTotal.Add(PluginGraph[Plugin->Name].ExclusiveSizes[(uint8)EPluginGraphSizeClass::All]);
-		}
-	}
+	UE::Cook::FCookMetadataPluginHierarchy& MutablePluginHierarchy = CookMetadata.GetMutablePluginHierarchy();
 
-	auto GeneratePluginJson = [&PluginHierarchy](TUtf8StringBuilder<4096>& OutPluginMetadataJson, FStringView InName, const FPluginGraphEntry& InGraphEntry, EPluginGraphSizeClass InSizeClass)
+	auto GeneratePluginJson = [&MutablePluginHierarchy](TUtf8StringBuilder<4096>& OutPluginMetadataJson, FStringView InName, const FPluginGraphEntry& InGraphEntry, EPluginGraphSizeClass InSizeClass)
 	{
 		OutPluginMetadataJson.Reset();
 		OutPluginMetadataJson << "{\n";
@@ -3605,16 +4022,17 @@ static void WritePluginMetadataJsons(const FString& InAssetRegistryFileName, TMa
 		OutPluginMetadataJson << "\t\"direct_refcount\":" << InGraphEntry.DirectRefcount << ",\n";
 
 		// pass through any custom fields that were added to the cook metadata.
-		if (InGraphEntry.Self)
+		if (InGraphEntry.IndexInEnabledPlugins != TNumericLimits<uint16>::Max())
 		{
-			for (const TPair<uint8, bool>& BoolValue : InGraphEntry.Self->CustomBoolFields)
+			const UE::Cook::FCookMetadataPluginEntry& CookMetadataEntry = MutablePluginHierarchy.PluginsEnabledAtCook[InGraphEntry.IndexInEnabledPlugins];
+			for (const TPair<uint8, bool>& BoolValue : CookMetadataEntry.CustomBoolFields)
 			{
-				const FString& FieldName = PluginHierarchy.CustomFieldNames[BoolValue.Key];
-				OutPluginMetadataJson << "\t\"" << FieldName << (BoolValue.Value ? "\":true,\n" : "\":false,\n");				
+				const FString& FieldName = MutablePluginHierarchy.CustomFieldNames[BoolValue.Key];
+				OutPluginMetadataJson << "\t\"" << FieldName << (BoolValue.Value ? "\":true,\n" : "\":false,\n");
 			}
-			for (const TPair<uint8, FString>& StringValue : InGraphEntry.Self->CustomStringFields)
+			for (const TPair<uint8, FString>& StringValue : CookMetadataEntry.CustomStringFields)
 			{
-				const FString& FieldName = PluginHierarchy.CustomFieldNames[StringValue.Key];
+				const FString& FieldName = MutablePluginHierarchy.CustomFieldNames[StringValue.Key];
 				OutPluginMetadataJson << "\t\"" << FieldName << "\":\"" << StringValue.Value << "\",\n";
 			}
 		}
@@ -3622,14 +4040,16 @@ static void WritePluginMetadataJsons(const FString& InAssetRegistryFileName, TMa
 		{
 			OutPluginMetadataJson << "\t\"roots\":[";
 
-			for (int32 RootIndex = 0; RootIndex < InGraphEntry.Roots.Num(); RootIndex++)
+			int32 RootIndex = 0;
+			for (FPluginGraphEntry* Root : InGraphEntry.Roots)
 			{
-				const UE::Cook::FCookMetadataPluginEntry* Root = InGraphEntry.Roots[RootIndex];
+				
 				OutPluginMetadataJson << "\"" << Root->Name << "\"";
 				if (RootIndex + 1 < InGraphEntry.Roots.Num())
 				{
 					OutPluginMetadataJson << ",";
 				}
+				RootIndex++;
 			}
 
 			OutPluginMetadataJson << "]\n";
@@ -3649,11 +4069,18 @@ static void WritePluginMetadataJsons(const FString& InAssetRegistryFileName, TMa
 	TUtf8StringBuilder<4096> Csv;
 	Csv.Append("name,asset_sizes_class,exclusive_installed,exclusive_optional,exclusive_ias,inclusive_installed,inclusive_optional,inclusive_ias,unique_installed,unique_optional,unique_ias,direct_refcount,total_dependency_count\n");
 
-	// We want to write the sizes back to the cook metadata, so we need a non-const version.
-	UE::Cook::FCookMetadataPluginHierarchy& MutablePluginHierarchy = CookMetadata.GetMutablePluginHierarchy();
+
+	// This is so re-staging the same cook is consistent.	
+	for (UE::Cook::FCookMetadataPluginEntry& Plugin : MutablePluginHierarchy.PluginsEnabledAtCook)
+	{
+		Plugin.InclusiveSizes.Zero();
+		Plugin.ExclusiveSizes.Zero();
+	}
+
+	uint32 JsonWrittenCount = 0;
 	for (UE::Cook::FCookMetadataPluginEntry& Plugin : MutablePluginHierarchy.PluginsEnabledAtCook)
 	{
-		const FPluginGraphEntry& PluginEntry = PluginGraph[Plugin.Name];
+		const FPluginGraphEntry& PluginEntry = *PluginGraph.NameToPlugin[Plugin.Name];
 		if (PluginEntry.InclusiveSizes[(uint8)EPluginGraphSizeClass::All].TotalSize() == 0)
 		{
 			continue;
@@ -3664,6 +4091,7 @@ static void WritePluginMetadataJsons(const FString& InAssetRegistryFileName, TMa
 
 		for (uint8 ClassIndex = 0; ClassIndex < FPluginGraphEntry::ClassCount; ClassIndex++)
 		{
+			JsonWrittenCount++;
 			GeneratePluginJson(PluginMetadataJson, Plugin.Name, PluginEntry, (EPluginGraphSizeClass)ClassIndex);
 			SavePluginMetadata(InAssetRegistryFileName, Plugin.Name, PluginMetadataJson, (EPluginGraphSizeClass)ClassIndex);
 
@@ -3690,10 +4118,11 @@ static void WritePluginMetadataJsons(const FString& InAssetRegistryFileName, TMa
 		OrphanedEntry.InclusiveSizes[(uint8)EPluginGraphSizeClass::All] = UnrootedTotal;
 		GeneratePluginJson(PluginMetadataJson, TEXT("OrphanedPlugins"), OrphanedEntry, EPluginGraphSizeClass::All);
 		SavePluginMetadata(InAssetRegistryFileName, TEXT("OrphanedPlugins"), PluginMetadataJson, EPluginGraphSizeClass::All);
+		JsonWrittenCount++;
 
 		Csv.Appendf("OrphanedPlugins,all,0,0,0,%llu,%llu,%llu,0,%u\n",
 			UnrootedTotal[UE::Cook::EPluginSizeTypes::Installed], UnrootedTotal[UE::Cook::EPluginSizeTypes::Optional], UnrootedTotal[UE::Cook::EPluginSizeTypes::Streaming],
-			UnrootedPlugins.Num());
+			PluginGraph.UnrootedPlugins.Num());
 	}
 
 	{
@@ -3706,13 +4135,17 @@ static void WritePluginMetadataJsons(const FString& InAssetRegistryFileName, TMa
 	}
 
 	double WritePluginEnd = FPlatformTime::Seconds();
-	UE_LOG(LogIoStore, Display, TEXT("Wrote plugin size jsons/csv in %.2f seconds [inclusive computation: %.2fs; plugin graph: %.2fs; topological sort: %.2fs; package mapping: %.2fs]"), 
-		WritePluginEnd - WritePluginStart, InclusiveComputeEnd - InclusiveComputeStart, AssetPackageMapStart - WritePluginStart, 
-		InclusiveComputeStart - TopologicalSortStart, TopologicalSortStart - AssetPackageMapStart);
+	UE_LOG(LogIoStore, Display, TEXT("Wrote %s plugin size jsons/csv in %.2f seconds [graph %.2f shaders %.2f sizes %.2f writes %.2f]"), 
+		*NumberString(JsonWrittenCount),
+		WritePluginEnd - WritePluginStart, 
+		GeneratePluginGraphEnd - WritePluginStart, 
+		AssetPackageMapStart - GeneratePluginGraphEnd,
+		WriteBegin - AssetPackageMapStart,
+		WritePluginEnd - WriteBegin);
 }
 
 
-static void AddChunkInfoToAssetRegistry(TMap<FPackageId, TArray<FIoStoreChunkSource, TInlineAllocator<2>>>&& PackageToChunks, FAssetRegistryState& AssetRegistry, uint64 TotalCompressedSize)
+static void AddChunkInfoToAssetRegistry(TMap<FPackageId, TArray<FIoStoreChunkSource, TInlineAllocator<2>>>&& PackageToChunks, FAssetRegistryState& AssetRegistry, const FShaderAssociationInfo* ShaderSizeInfo, uint64 InUnassignableShaderCodeBytes, uint64 InAssignableShaderCodeBytes, uint64 TotalCompressedSize)
 {
 	//
 	// The asset registry has the chunks associate with each package, so we can just iterate the
@@ -3734,6 +4167,8 @@ static void AddChunkInfoToAssetRegistry(TMap<FPackageId, TArray<FIoStoreChunkSou
 			continue;
 		}
 
+		FPackageId PackageId = FPackageId::FromName(AssetPackage.Key);
+
 		const FAssetData* AssetData = UE::AssetRegistry::GetMostImportantAsset(AssetRegistry.GetAssetsByPackageName(AssetPackage.Key), UE::AssetRegistry::EGetMostImportantAssetFlags::IgnoreSkipClasses);
 		if (AssetData == nullptr)
 		{
@@ -3741,7 +4176,7 @@ static void AddChunkInfoToAssetRegistry(TMap<FPackageId, TArray<FIoStoreChunkSou
 			continue;
 		}
 
-		const TArray<FIoStoreChunkSource, TInlineAllocator<2>>* PackageChunks = PackageToChunks.Find(FPackageId::FromName(AssetPackage.Key));
+		const TArray<FIoStoreChunkSource, TInlineAllocator<2>>* PackageChunks = PackageToChunks.Find(PackageId);
 		if (PackageChunks == nullptr)
 		{
 			// This happens when the package has been stripped by UAT prior to staging by e.g. PakDenyList.
@@ -3769,7 +4204,7 @@ static void AddChunkInfoToAssetRegistry(TMap<FPackageId, TArray<FIoStoreChunkSou
 
 		// We assign a package's chunks to a single asset, remove it from the list so that
 		// at the end we can track how many chunks don't get assigned.
-		PackageToChunks.Remove(FPackageId::FromName(AssetPackage.Key));
+		PackageToChunks.Remove(PackageId);
 
 		UpdatedAssetCount++;
 		AssetsCompressedSize += PackageCompressedSize.TotalSize();
@@ -3784,6 +4219,10 @@ static void AddChunkInfoToAssetRegistry(TMap<FPackageId, TArray<FIoStoreChunkSou
 			RemainingByType[(uint8)Info.ChunkInfo.ChunkType] += Info.ChunkInfo.CompressedSize;
 		}
 	}
+	
+	// Shaders aren't in the PackageToChunks map, but we want to numbers reported to include them.
+	RemainingByType[(uint8)EIoChunkType::ShaderCode] += InUnassignableShaderCodeBytes;
+	AssetsCompressedSize += InAssignableShaderCodeBytes;
 
 	double PercentAssets = 1.0f;
 	if (TotalCompressedSize != 0)
@@ -3916,7 +4355,7 @@ static int32 DoAssetRegistryWritebackAfterStage(const FString& InAssetRegistryFi
 		});
 	}
 
-	AddChunkInfoToAssetRegistry(MoveTemp(PackageToChunks), AssetRegistry, TotalCompressedSize);
+	AddChunkInfoToAssetRegistry(MoveTemp(PackageToChunks), AssetRegistry, nullptr, 0, 0, TotalCompressedSize);
 
 	return SaveAssetRegistry(InAssetRegistryFileName, AssetRegistry, nullptr) ? 0 : 1;
 }
@@ -4039,13 +4478,21 @@ static ECookMetadataFiles FindAndLoadMetadataFiles(
 	return ResultFiles;
 }
 
+struct FIoStoreWriterInfo
+{
+	UE::Cook::EPluginSizeTypes SizeType;
+	FName PakChunkName;
+};
+
 static bool DoAssetRegistryWritebackDuringStage(
 	EAssetRegistryWritebackMethod InMethod, 
 	bool bInWritePluginMetadata,
 	const FString& InCookedDir, 
 	bool bInCompressionEnabled,
 	TArray<TSharedPtr<IIoStoreWriter>>& InIoStoreWriters, 
-	TArray<UE::Cook::EPluginSizeTypes>& IoStoreWriterSizeClassifications)
+	TArray<FIoStoreWriterInfo>& InIoStoreWriterInfos,
+	FShaderAssociationInfo& InShaderAssociationInfo
+)
 {
 	// This version called during container creation.
 	TRACE_CPUPROFILER_EVENT_SCOPE(UpdateAssetRegistryWithSizeInfo);
@@ -4070,17 +4517,52 @@ static bool DoAssetRegistryWritebackDuringStage(
 	//
 	// We want to separate out the sizes based on where they go in the end product.
 	//
-	UE::Cook::FPluginSizeInfo ProductSize;	
+	uint64 UnassignableShaderCodeBytes = 0;
+	uint64 AssignableShaderCodeBytes = 0;
+	UE::Cook::FPluginSizeInfo ProductSize;
 	TMap<FPackageId, TArray<FIoStoreChunkSource, TInlineAllocator<2>>> PackageToChunks;
 	{
 		int32 IoStoreWriterIndex = 0;
 		for (TSharedPtr<IIoStoreWriter> IoStoreWriter : InIoStoreWriters)
 		{
-			IoStoreWriter->EnumerateChunks([&PackageToChunks, IoStoreWriterClassification = IoStoreWriterSizeClassifications[IoStoreWriterIndex], &ProductSize](const FIoStoreTocChunkInfo& ChunkInfo)
-			{
+			IoStoreWriter->EnumerateChunks(
+				[&PackageToChunks, 
+				 IoStoreWriterInfo = InIoStoreWriterInfos[IoStoreWriterIndex],
+				 &ProductSize, 
+				 &InShaderAssociationInfo,
+				 &UnassignableShaderCodeBytes,
+				 &AssignableShaderCodeBytes
+				 ](const FIoStoreTocChunkInfo& ChunkInfo)
+			{			
+				ProductSize[IoStoreWriterInfo.SizeType] += ChunkInfo.CompressedSize;
+
+				// Shader code chunks don't have the package in their chunk id, so we have to use other data to look
+				// it up and find it.
 				FPackageId PackageId = FPackageId::FromValue(*(int64*)(ChunkInfo.Id.GetData()));
-				PackageToChunks.FindOrAdd(PackageId).Add({ChunkInfo, IoStoreWriterClassification });
-				ProductSize[IoStoreWriterClassification] += ChunkInfo.CompressedSize;
+				if (ChunkInfo.ChunkType == EIoChunkType::ShaderCode)
+				{
+					// Update size info for the shader.
+					FShaderAssociationInfo::FShaderChunkInfo* ShaderChunkInfo = InShaderAssociationInfo.ShaderChunkInfos.Find({IoStoreWriterInfo.PakChunkName, ChunkInfo.Id});
+					ShaderChunkInfo->CompressedSize = ChunkInfo.CompressedSize;
+					ShaderChunkInfo->SizeType = IoStoreWriterInfo.SizeType;
+
+					// Shaders don't put their package in their chunk - they are just a hash. We have the list of them
+					// already in the shader association so we don't bother adding here. However some can't get assigned to
+					// anything and so we track that size for reporting.
+					if (ShaderChunkInfo->Type == FShaderAssociationInfo::FShaderChunkInfo::Global ||
+						ShaderChunkInfo->Type == FShaderAssociationInfo::FShaderChunkInfo::Orphan)
+					{
+						UnassignableShaderCodeBytes += ChunkInfo.CompressedSize;
+					}
+					else
+					{
+						AssignableShaderCodeBytes += ChunkInfo.CompressedSize;
+					}
+				}
+				else
+				{
+					PackageToChunks.FindOrAdd(PackageId).Add({ChunkInfo, IoStoreWriterInfo.SizeType});
+				}
 				return true;
 			});
 
@@ -4088,15 +4570,13 @@ static bool DoAssetRegistryWritebackDuringStage(
 		}
 	}
 
-	uint64 UpdatedDevArHash = 0;
 	if (bInWritePluginMetadata)
 	{
-		WritePluginMetadataJsons(AssetRegistryFileName, PackageToChunks, AssetRegistry, CookMetadata);
+		UpdatePluginMetadataAndWriteJsons(AssetRegistryFileName, PackageToChunks, AssetRegistry, CookMetadata, &InShaderAssociationInfo);
 	}
 
-	AddChunkInfoToAssetRegistry(MoveTemp(PackageToChunks), AssetRegistry, ProductSize.TotalSize());
-	
-	
+	AddChunkInfoToAssetRegistry(MoveTemp(PackageToChunks), AssetRegistry, &InShaderAssociationInfo, UnassignableShaderCodeBytes, AssignableShaderCodeBytes, ProductSize.TotalSize());
+	uint64 UpdatedDevArHash = 0;
 	switch (InMethod)
 	{
 	case EAssetRegistryWritebackMethod::OriginalFile:
@@ -4468,7 +4948,7 @@ int32 CreateTarget(const FIoStoreArguments& Arguments, const FIoStoreWriterSetti
 	}
 	TArray<FString> OnDemandContainers;
 	TArray<TSharedPtr<IIoStoreWriter>> IoStoreWriters;
-	TArray<UE::Cook::EPluginSizeTypes> IoStoreWriterSizeClassifications;
+	TArray<FIoStoreWriterInfo> IoStoreWriterInfos;
 	TSharedPtr<IIoStoreWriter> GlobalIoStoreWriter;
 	{
 		IOSTORE_CPU_SCOPE(InitializeWriters);
@@ -4483,7 +4963,7 @@ int32 CreateTarget(const FIoStoreArguments& Arguments, const FIoStoreWriterSetti
 			}
 			GlobalIoStoreWriter = IoStoreWriterContext->CreateContainer(*Arguments.GlobalContainerPath, GlobalContainerSettings);
 			IoStoreWriters.Add(GlobalIoStoreWriter);
-			IoStoreWriterSizeClassifications.Add(UE::Cook::EPluginSizeTypes::Installed);
+			IoStoreWriterInfos.Add({UE::Cook::EPluginSizeTypes::Installed, "global"});
 		}
 		for (FContainerTargetSpec* ContainerTarget : ContainerTargets)
 		{
@@ -4500,7 +4980,7 @@ int32 CreateTarget(const FIoStoreArguments& Arguments, const FIoStoreWriterSetti
 				FLooseFilesWriterSettings WriterSettings;
 				WriterSettings.TargetRootPath = ContainerTarget->StageLooseFileRootPath;
 				ContainerTarget->IoStoreWriter = MakeLooseFilesIoStoreWriter(WriterSettings);
-				IoStoreWriterSizeClassifications.Add(UE::Cook::EPluginSizeTypes::Streaming); // LooseFiles currently end up as a streamed source.
+				IoStoreWriterInfos.Add({UE::Cook::EPluginSizeTypes::Streaming, ContainerTarget->Name}); // LooseFiles currently end up as a streamed source.
 				IoStoreWriters.Add(ContainerTarget->IoStoreWriter);
 			}
 			else
@@ -4538,7 +5018,7 @@ int32 CreateTarget(const FIoStoreArguments& Arguments, const FIoStoreWriterSetti
 
 				if (EnumHasAnyFlags(ContainerTarget->ContainerFlags, EIoContainerFlags::OnDemand))
 				{
-					IoStoreWriterSizeClassifications.Add(UE::Cook::EPluginSizeTypes::Streaming);
+					IoStoreWriterInfos.Add({UE::Cook::EPluginSizeTypes::Streaming, ContainerTarget->Name});
 				}
 				else
 				{
@@ -4552,11 +5032,11 @@ int32 CreateTarget(const FIoStoreArguments& Arguments, const FIoStoreWriterSetti
 					}
 					if (BaseFileName.EndsWith(TEXT("optional")))
 					{
-						IoStoreWriterSizeClassifications.Add(UE::Cook::EPluginSizeTypes::Optional);
+						IoStoreWriterInfos.Add({UE::Cook::EPluginSizeTypes::Optional, ContainerTarget->Name});
 					}
 					else
 					{
-						IoStoreWriterSizeClassifications.Add(UE::Cook::EPluginSizeTypes::Installed);
+						IoStoreWriterInfos.Add({UE::Cook::EPluginSizeTypes::Installed, ContainerTarget->Name});
 					}
 				}
 
@@ -4565,7 +5045,7 @@ int32 CreateTarget(const FIoStoreArguments& Arguments, const FIoStoreWriterSetti
 					ContainerTarget->OptionalSegmentIoStoreWriter = IoStoreWriterContext->CreateContainer(*ContainerTarget->OptionalSegmentOutputPath, ContainerSettings);
 					ContainerTarget->OptionalSegmentIoStoreWriter->SetReferenceChunkDatabase(ChunkDatabase);
 					IoStoreWriters.Add(ContainerTarget->OptionalSegmentIoStoreWriter);
-					IoStoreWriterSizeClassifications.Add(UE::Cook::EPluginSizeTypes::OptionalSegment);
+					IoStoreWriterInfos.Add({UE::Cook::EPluginSizeTypes::OptionalSegment, ContainerTarget->Name});
 				}
 				if (EnumHasAnyFlags(ContainerTarget->ContainerFlags, EIoContainerFlags::OnDemand))
 				{
@@ -4609,7 +5089,8 @@ int32 CreateTarget(const FIoStoreArguments& Arguments, const FIoStoreWriterSetti
 
 	UE_LOG(LogIoStore, Display, TEXT("Processing shader libraries, compressing with Oodle %s, level %d (%s)"), FOodleDataCompression::ECompressorToString(Arguments.ShaderOodleCompressor), (int32)Arguments.ShaderOodleLevel, FOodleDataCompression::ECompressionLevelToString(Arguments.ShaderOodleLevel));
 	TArray<FShaderInfo*> Shaders;
-	ProcessShaderLibraries(Arguments, ContainerTargets, Shaders);
+	FShaderAssociationInfo ShaderAssocInfo;
+	ProcessShaderLibraries(Arguments, ContainerTargets, Shaders, ShaderAssocInfo);
 
 	auto AppendTargetFileChunk = [&WriteRequestManager](FContainerTargetSpec* ContainerTarget, const FContainerTargetFile& TargetFile)
 	{
@@ -4859,7 +5340,14 @@ int32 CreateTarget(const FIoStoreArguments& Arguments, const FIoStoreWriterSetti
 
 	if (Arguments.WriteBackMetadataToAssetRegistry != EAssetRegistryWritebackMethod::Disabled)
 	{
-		DoAssetRegistryWritebackDuringStage(Arguments.WriteBackMetadataToAssetRegistry, Arguments.bWritePluginSizeSummaryJsons, Arguments.CookedDir, GeneralIoWriterSettings.CompressionMethod != NAME_None, IoStoreWriters, IoStoreWriterSizeClassifications);
+		DoAssetRegistryWritebackDuringStage(
+			Arguments.WriteBackMetadataToAssetRegistry, 
+			Arguments.bWritePluginSizeSummaryJsons, 
+			Arguments.CookedDir, 
+			GeneralIoWriterSettings.CompressionMethod != NAME_None, 
+			IoStoreWriters, 
+			IoStoreWriterInfos,
+			ShaderAssocInfo);
 	}
 
 	TArray<FIoStoreWriterResult> IoStoreWriterResults;
@@ -8785,3 +9273,4 @@ bool PrimeEndPoint(FStringView IoStoreOnDemandIniPath)
 		return false;
 	}
 }
+
diff --git a/Engine/Source/Editor/UnrealEd/Private/CookOnTheFlyServer.cpp b/Engine/Source/Editor/UnrealEd/Private/CookOnTheFlyServer.cpp
--- a/Engine/Source/Editor/UnrealEd/Private/CookOnTheFlyServer.cpp
+++ b/Engine/Source/Editor/UnrealEd/Private/CookOnTheFlyServer.cpp
@@ -9521,9 +9521,17 @@ void UCookOnTheFlyServer::WriteCookMetadata(const ITargetPlatform* InTargetPlatf
 
 		// Add the /Engine and /Game pseudo plugins. These are placeholders for holding size information when unrealpak runs.
 		const uint16 EnginePluginIndex = (uint16)PluginsToAdd.Num();
-		PluginsToAdd.AddDefaulted_GetRef().Name = TEXT("Engine");
+		{
+			UE::Cook::FCookMetadataPluginEntry& EnginePseudoPlugin = PluginsToAdd.AddDefaulted_GetRef();
+			EnginePseudoPlugin.Name = TEXT("Engine");
+			EnginePseudoPlugin.Type = UE::Cook::ECookMetadataPluginType::EnginePseudo;
+		}
 		const uint16 GamePluginIndex = (uint16)PluginsToAdd.Num();
-		PluginsToAdd.AddDefaulted_GetRef().Name = TEXT("Game");
+		{
+			UE::Cook::FCookMetadataPluginEntry& GamePseudoPlugin = PluginsToAdd.AddDefaulted_GetRef();
+			GamePseudoPlugin.Name = TEXT("Game");
+			GamePseudoPlugin.Type = UE::Cook::ECookMetadataPluginType::GamePseudo;
+		}
 
 		// Construct the dependency list.
 		TArray<uint16> RootPlugins;
@@ -9531,16 +9539,18 @@ void UCookOnTheFlyServer::WriteCookMetadata(const ITargetPlatform* InTargetPlatf
 		{
 			uint16 SelfIndex = IndexForPlugin[EnabledPlugin->GetName()];
 			UE::Cook::FCookMetadataPluginEntry& Entry = PluginsToAdd[SelfIndex];
+			Entry.Type = UE::Cook::ECookMetadataPluginType::Normal;
 
 			// We detect if this would overflow below and cancel the write - so while this could store
 			// bogus data, it won't get saved.
-			Entry.DependencyIndexStart = (uint16)PluginChildArray.Num();
+			Entry.DependencyIndexStart = (uint32)PluginChildArray.Num();
 
 			const FPluginDescriptor& Descriptor = EnabledPlugin->GetDescriptor();
 
 			// Root plugins are sealed && no code
 			if (Descriptor.bIsSealed && Descriptor.bNoCode)
 			{
+				Entry.Type = UE::Cook::ECookMetadataPluginType::Root;
 				RootPlugins.Add(SelfIndex);
 			}
 
@@ -9677,36 +9687,30 @@ void UCookOnTheFlyServer::WriteCookMetadata(const ITargetPlatform* InTargetPlatf
 				PluginChildArray.Add(GamePluginIndex);
 			}
 
-			Entry.DependencyIndexEnd = PluginChildArray.Num();
+			Entry.DependencyIndexEnd = (uint32)PluginChildArray.Num();
 		}
 
 		// Also ensure Game depends on Engine.
-		PluginsToAdd[GamePluginIndex].DependencyIndexStart = PluginChildArray.Num();
+		PluginsToAdd[GamePluginIndex].DependencyIndexStart = (uint32)PluginChildArray.Num();
 		PluginChildArray.Add(EnginePluginIndex);
-		PluginsToAdd[GamePluginIndex].DependencyIndexEnd = PluginChildArray.Num();
+		PluginsToAdd[GamePluginIndex].DependencyIndexEnd = (uint32)PluginChildArray.Num();
 
-		if (IntFitsIn<uint16>(PluginChildArray.Num()) == false)
-		{
-			UE_LOG(LogCook, Warning, TEXT("Number of child plugins exceeds 64k, unable to write cook metadata file (count = %d)"), PluginChildArray.Num());
-		}
-		else
-		{
-			UE::Cook::FCookMetadataState MetadataState;
-			UE::Cook::FCookMetadataPluginHierarchy PluginHierarchy;
+		UE::Cook::FCookMetadataState MetadataState;
+		UE::Cook::FCookMetadataPluginHierarchy PluginHierarchy;
 
-			PluginHierarchy.PluginsEnabledAtCook = MoveTemp(PluginsToAdd);
-			PluginHierarchy.PluginDependencies = MoveTemp(PluginChildArray);
-			PluginHierarchy.RootPlugins = MoveTemp(RootPlugins);
-			PluginHierarchy.CustomFieldNames = MoveTemp(CustomFieldNames);
+		PluginHierarchy.PluginsEnabledAtCook = MoveTemp(PluginsToAdd);
+		PluginHierarchy.PluginDependencies = MoveTemp(PluginChildArray);
+		PluginHierarchy.RootPlugins = MoveTemp(RootPlugins);
+		PluginHierarchy.CustomFieldNames = MoveTemp(CustomFieldNames);
 
-			MetadataState.SetPluginHierarchyInfo(MoveTemp(PluginHierarchy));
-			MetadataState.SetAssociatedDevelopmentAssetRegistryHash(InDevelopmentAssetRegistryHash);
+		MetadataState.SetPluginHierarchyInfo(MoveTemp(PluginHierarchy));
+		MetadataState.SetAssociatedDevelopmentAssetRegistryHash(InDevelopmentAssetRegistryHash);
 
-			MetadataState.SetPlatformAndBuildVersion(PlatformNameString, FApp::GetBuildVersion());
-			MetadataState.SetHordeJobId(FPlatformMisc::GetEnvironmentVariable(TEXT("UE_HORDE_JOBID")));
+		MetadataState.SetPlatformAndBuildVersion(PlatformNameString, FApp::GetBuildVersion());
+		MetadataState.SetHordeJobId(FPlatformMisc::GetEnvironmentVariable(TEXT("UE_HORDE_JOBID")));
 
 			MetadataState.SaveToFile(GetCookedCookMetadataFilename(PlatformNameString));
-		}
+
 	}
 }
 
