ZenStore structures: expose structures to interpret zenstore headers for use in -diffonly cooks.
PackageStoreOptimizer: Expose a static version of FindScriptObject and functions for calculating PublicExportHash.
FNameMap: Add a foreach iterator to provide a list of all data in the struct.
ZenPackageHeader: Move from private to internal, add a constructor that returns an error string rather than fatally asserting.
#rnx

diff --git a/Engine/Source/Developer/IoStoreUtilities/Internal/PackageStoreOptimizer.h b/Engine/Source/Developer/IoStoreUtilities/Internal/PackageStoreOptimizer.h
--- a/Engine/Source/Developer/IoStoreUtilities/Internal/PackageStoreOptimizer.h
+++ b/Engine/Source/Developer/IoStoreUtilities/Internal/PackageStoreOptimizer.h
@@ -2,6 +2,7 @@
 
 #pragma once
 
+#include "Containers/StringFwd.h"
 #include "Serialization/AsyncLoading2.h"
 #include "IO/PackageStore.h"
 #include "IO/IoDispatcher.h"
@@ -205,6 +206,15 @@ private:
 class FPackageStoreOptimizer
 {
 public:
+	struct FScriptObjectData
+	{
+		FName ObjectName;
+		FString FullName;
+		FPackageObjectIndex GlobalIndex;
+		FPackageObjectIndex OuterIndex;
+		FPackageObjectIndex CDOClassIndex;
+	};
+
 	uint64 GetTotalScriptObjectCount() const
 	{
 		return TotalScriptObjectCount;
@@ -220,15 +230,11 @@ public:
 	IOSTOREUTILITIES_API FIoBuffer CreateScriptObjectsBuffer() const;
 	void LoadScriptObjectsBuffer(const FIoBuffer& ScriptObjectsBuffer);
 
+	IOSTOREUTILITIES_API static void FindScriptObjects(TMap<FPackageObjectIndex, FScriptObjectData>& OutScriptObjectsMap);
+	IOSTOREUTILITIES_API static void AppendPathForPublicExportHash(UObject* Object, FStringBuilderBase& OutPath);
+	IOSTOREUTILITIES_API static bool TryGetPublicExportHash(FStringView PackageRelativeExportPath, uint64& OutPublicExportHash);
+
 private:
-	struct FScriptObjectData
-	{
-		FName ObjectName;
-		FString FullName;
-		FPackageObjectIndex GlobalIndex;
-		FPackageObjectIndex OuterIndex;
-		FPackageObjectIndex CDOClassIndex;
-	};
 
 	struct FCookedHeaderData
 	{
@@ -253,7 +259,8 @@ private:
 	TArray<FPackageStorePackage::FExportGraphNode*> SortExportGraphNodesInLoadOrder(FPackageStorePackage* Package, FExportGraphEdges& Edges) const;
 	void CreateExportBundle(FPackageStorePackage* Package) const;
 	void FinalizePackageHeader(FPackageStorePackage* Package) const;
-	void FindScriptObjectsRecursive(FPackageObjectIndex OuterIndex, UObject* Object);
+	static void FindScriptObjectsRecursive(TMap<FPackageObjectIndex, FScriptObjectData>& OutScriptObjectsMap,
+		FPackageObjectIndex OuterIndex, UObject* Object);
 	void FindScriptObjects();
 
 	TMap<FPackageObjectIndex, FScriptObjectData> ScriptObjectsMap;
diff --git a/Engine/Source/Developer/IoStoreUtilities/Private/PackageStoreOptimizer.cpp b/Engine/Source/Developer/IoStoreUtilities/Private/PackageStoreOptimizer.cpp
--- a/Engine/Source/Developer/IoStoreUtilities/Private/PackageStoreOptimizer.cpp
+++ b/Engine/Source/Developer/IoStoreUtilities/Private/PackageStoreOptimizer.cpp
@@ -13,6 +13,7 @@
 #include "Serialization/MemoryReader.h"
 #include "UObject/Package.h"
 #include "Misc/SecureHash.h"
+#include "Misc/StringBuilder.h"
 #include "Serialization/LargeMemoryReader.h"
 
 DEFINE_LOG_CATEGORY_STATIC(LogPackageStoreOptimizer, Log, All);
@@ -267,6 +268,67 @@ void FPackageStoreOptimizer::ResolveImport(FPackageStorePackage::FUnresolvedImpo
 	}
 }
 
+void FPackageStoreOptimizer::AppendPathForPublicExportHash(UObject* Object, FStringBuilderBase& OutPath)
+{
+	if (!Object)
+	{
+		return;
+	}
+	if (!Object->GetOuter())
+	{
+		// Outermost package objects do not have a PublicExportHash Name; they use 0 for PublicExportHash
+		// Write nothing for them.
+		return;
+	}
+
+	constexpr int32 DirectorySeparatorLen = 1;
+	const TCHAR DirectorySeparatorChar = '/';
+
+	TArray<FName, TInlineAllocator<10>> PathNames;
+	int32 PathNameLen = 0;
+	for (UObject* Iter = Object; Iter->GetOuter(); Iter = Object->GetOuter())
+	{
+		FName PathName = Iter->GetFName();
+		PathNames.Add(PathName);
+		PathNameLen += PathName.GetStringLength() + DirectorySeparatorLen;
+	}
+	if (PathNameLen <= 1)
+	{
+		// We should be writing at least 2 characters: /<LeafName>. Write nothing if we find an empty LeafName.
+		return;
+	}
+
+	int32 InitialLength = OutPath.Len();
+	for (FName PathName : ReverseIterate(PathNames))
+	{
+		OutPath << DirectorySeparatorChar << PathName;
+	}
+
+	// Down-case the characters we wrote; public export hash is based on lower-case
+	TCHAR* OutPathData = OutPath.GetData();
+	int32 NewLength = OutPath.Len();
+	for (int32 N = InitialLength; N < NewLength; ++N)
+	{
+		OutPathData[N] = FChar::ToLower(OutPathData[N]);
+	}
+}
+
+bool FPackageStoreOptimizer::TryGetPublicExportHash(FStringView PackageRelativeExportPath, uint64& OutPublicExportHash)
+{
+	OutPublicExportHash = 0;
+
+	// PackageRelativeExportHash should have been generated by GetNameForPublicExportHash or an equivalent function
+	if (PackageRelativeExportPath.Len() == 0 || // PublicExportHash is not defined for the UPackage object
+		PackageRelativeExportPath.Len() == 1 || // An object with an empty leaf name is an error, but we ignore it and treat it the same as a package
+		PackageRelativeExportPath[0] != '/') // Invalid string; GetNameForPublicExportHash starts every relative path with /
+	{
+		return false;
+	}
+
+	OutPublicExportHash = GetPublicExportHash(PackageRelativeExportPath);
+	return true;
+}
+
 uint64 FPackageStoreOptimizer::GetPublicExportHash(FStringView PackageRelativeExportPath)
 {
 	check(PackageRelativeExportPath.Len() > 1);
@@ -830,7 +892,8 @@ FIoBuffer FPackageStoreOptimizer::CreatePackageBuffer(const FPackageStorePackage
 	return BundleBuffer;
 }
 
-void FPackageStoreOptimizer::FindScriptObjectsRecursive(FPackageObjectIndex OuterIndex, UObject* Object)
+void FPackageStoreOptimizer::FindScriptObjectsRecursive(
+	TMap<FPackageObjectIndex, FScriptObjectData>& OutScriptObjectsMap, FPackageObjectIndex OuterIndex, UObject* Object)
 {
 	if (!Object->HasAllFlags(RF_Public))
 	{
@@ -841,7 +904,7 @@ void FPackageStoreOptimizer::FindScriptObjectsRecursive(FPackageObjectIndex Oute
 	FString OuterFullName;
 	FPackageObjectIndex OuterCDOClassIndex;
 	{
-		const FScriptObjectData* Outer = ScriptObjectsMap.Find(OuterIndex);
+		const FScriptObjectData* Outer = OutScriptObjectsMap.Find(OuterIndex);
 		check(Outer);
 		OuterFullName = Outer->FullName;
 		OuterCDOClassIndex = Outer->CDOClassIndex;
@@ -856,7 +919,7 @@ void FPackageStoreOptimizer::FindScriptObjectsRecursive(FPackageObjectIndex Oute
 	TempFullName.ToLowerInline();
 	FPackageObjectIndex GlobalImportIndex = FPackageObjectIndex::FromScriptPath(TempFullName);
 
-	FScriptObjectData* ScriptImport = ScriptObjectsMap.Find(GlobalImportIndex);
+	FScriptObjectData* ScriptImport = OutScriptObjectsMap.Find(GlobalImportIndex);
 	if (ScriptImport)
 	{
 		UE_LOG(LogPackageStoreOptimizer, Fatal, TEXT("Import name hash collision \"%s\" and \"%s"), *TempFullName, *ScriptImport->FullName);
@@ -878,7 +941,7 @@ void FPackageStoreOptimizer::FindScriptObjectsRecursive(FPackageObjectIndex Oute
 		}
 	}
 
-	ScriptImport = &ScriptObjectsMap.Add(GlobalImportIndex);
+	ScriptImport = &OutScriptObjectsMap.Add(GlobalImportIndex);
 	ScriptImport->GlobalIndex = GlobalImportIndex;
 	ScriptImport->FullName = MoveTemp(TempFullName);
 	ScriptImport->OuterIndex = OuterIndex;
@@ -889,11 +952,17 @@ void FPackageStoreOptimizer::FindScriptObjectsRecursive(FPackageObjectIndex Oute
 	GetObjectsWithOuter(Object, InnerObjects, /*bIncludeNestedObjects*/false);
 	for (UObject* InnerObject : InnerObjects)
 	{
-		FindScriptObjectsRecursive(GlobalImportIndex, InnerObject);
+		FindScriptObjectsRecursive(OutScriptObjectsMap, GlobalImportIndex, InnerObject);
 	}
 };
 
 void FPackageStoreOptimizer::FindScriptObjects()
+{
+	FindScriptObjects(ScriptObjectsMap);
+	TotalScriptObjectCount = ScriptObjectsMap.Num();
+}
+
+void FPackageStoreOptimizer::FindScriptObjects(TMap<FPackageObjectIndex, FScriptObjectData>& OutScriptObjectsMap)
 {
 	TRACE_CPUPROFILER_EVENT_SCOPE(FindScriptObjects);
 	TArray<UPackage*> ScriptPackages;
@@ -908,10 +977,10 @@ void FPackageStoreOptimizer::FindScriptObjects()
 		FullName.ToLowerInline();
 		FPackageObjectIndex GlobalImportIndex = FPackageObjectIndex::FromScriptPath(FullName);
 
-		FScriptObjectData* ScriptImport = ScriptObjectsMap.Find(GlobalImportIndex);
+		FScriptObjectData* ScriptImport = OutScriptObjectsMap.Find(GlobalImportIndex);
 		checkf(!ScriptImport, TEXT("Import name hash collision \"%s\" and \"%s"), *FullName, *ScriptImport->FullName);
 
-		ScriptImport = &ScriptObjectsMap.Add(GlobalImportIndex);
+		ScriptImport = &OutScriptObjectsMap.Add(GlobalImportIndex);
 		ScriptImport->GlobalIndex = GlobalImportIndex;
 		ScriptImport->FullName = FullName;
 		ScriptImport->OuterIndex = FPackageObjectIndex();
@@ -921,11 +990,9 @@ void FPackageStoreOptimizer::FindScriptObjects()
 		GetObjectsWithOuter(Package, InnerObjects, /*bIncludeNestedObjects*/false);
 		for (UObject* InnerObject : InnerObjects)
 		{
-			FindScriptObjectsRecursive(GlobalImportIndex, InnerObject);
+			FindScriptObjectsRecursive(OutScriptObjectsMap, GlobalImportIndex, InnerObject);
 		}
 	}
-
-	TotalScriptObjectCount = ScriptObjectsMap.Num();
 }
 
 FIoBuffer FPackageStoreOptimizer::CreateScriptObjectsBuffer() const
diff --git a/Engine/Source/Runtime/Core/Public/Serialization/MappedName.h b/Engine/Source/Runtime/Core/Public/Serialization/MappedName.h
--- a/Engine/Source/Runtime/Core/Public/Serialization/MappedName.h
+++ b/Engine/Source/Runtime/Core/Public/Serialization/MappedName.h
@@ -152,6 +152,10 @@ public:
 		return NameToMinimalName(GetName(MappedName));
 	}
 
+	using RangedForConstIteratorType = TArray<FDisplayNameEntryId>::RangedForConstIteratorType;
+	RangedForConstIteratorType begin() const { return NameEntries.begin(); }
+	RangedForConstIteratorType end() const { return NameEntries.end(); }
+
 private:
 	TArray<FDisplayNameEntryId> NameEntries;
 	FMappedName::EType NameMapType = FMappedName::EType::Global;
diff --git a/Engine/Source/Runtime/CoreUObject/Internal/Serialization/ZenPackageHeader.h b/Engine/Source/Runtime/CoreUObject/Internal/Serialization/ZenPackageHeader.h
new file mode 100644
--- /dev/null
+++ b/Engine/Source/Runtime/CoreUObject/Internal/Serialization/ZenPackageHeader.h
@@ -0,0 +1,36 @@
+// Copyright Epic Games, Inc. All Rights Reserved.
+
+#pragma once
+
+#include "Containers/Array.h"
+#include "Containers/ArrayView.h"
+#include "IO/PackageId.h"
+#include "Memory/MemoryFwd.h"
+#include "Serialization/AsyncLoading2.h"
+#include "Serialization/MappedName.h"
+#include "UObject/NameTypes.h"
+
+struct FZenPackageHeader
+{
+	uint32 CookedHeaderSize = 0;
+	uint32 ExportCount = 0; // Need to keep this count around after ExportMap is cleared
+	TOptional<FZenPackageVersioningInfo> VersioningInfo;
+	FNameMap NameMap;
+	FName PackageName;
+
+	// Backed by IoBuffer
+	const FZenPackageSummary* PackageSummary = nullptr;
+	TArrayView<const uint64> ImportedPublicExportHashes;
+	TArrayView<const FPackageObjectIndex> ImportMap;
+	TArrayView<const FExportMapEntry> ExportMap;
+	TArrayView<const FBulkDataMapEntry> BulkDataMap;
+	TArrayView<const FExportBundleEntry> ExportBundleEntries;
+	TArrayView<const FDependencyBundleHeader> DependencyBundleHeaders;
+	TArrayView<const FDependencyBundleEntry> DependencyBundleEntries;
+
+	TArray<FName> ImportedPackageNames;
+
+	COREUOBJECT_API static FZenPackageHeader MakeView(FMemoryView Memory);
+	COREUOBJECT_API static FZenPackageHeader MakeView(FMemoryView Memory, FString& OutError);
+	COREUOBJECT_API void Reset();
+};
diff --git a/Engine/Source/Runtime/CoreUObject/Private/Serialization/AsyncLoading2.cpp b/Engine/Source/Runtime/CoreUObject/Private/Serialization/AsyncLoading2.cpp
--- a/Engine/Source/Runtime/CoreUObject/Private/Serialization/AsyncLoading2.cpp
+++ b/Engine/Source/Runtime/CoreUObject/Private/Serialization/AsyncLoading2.cpp
@@ -77,7 +77,7 @@
 #include "Containers/SpscQueue.h"
 #include "IO/IoPriorityQueue.h"
 #include "UObject/CoreRedirects.h"
-#include "ZenPackageHeader.h"
+#include "Serialization/ZenPackageHeader.h"
 
 #include <atomic>
 
diff --git a/Engine/Source/Runtime/CoreUObject/Private/Serialization/ZenPackageHeader.cpp b/Engine/Source/Runtime/CoreUObject/Private/Serialization/ZenPackageHeader.cpp
--- a/Engine/Source/Runtime/CoreUObject/Private/Serialization/ZenPackageHeader.cpp
+++ b/Engine/Source/Runtime/CoreUObject/Private/Serialization/ZenPackageHeader.cpp
@@ -1,11 +1,24 @@
 // Copyright Epic Games, Inc. All Rights Reserved.
 
-#include "ZenPackageHeader.h"
-#include "Serialization/MemoryReader.h"
+#include "Serialization/ZenPackageHeader.h"
 
+#include "Serialization/MemoryReader.h"
 
 FZenPackageHeader FZenPackageHeader::MakeView(FMemoryView Memory)
 {
+	FString Error;
+	FZenPackageHeader Result = MakeView(Memory, Error);
+	if (!Error.IsEmpty())
+	{
+		UE_LOG(LogStreaming, Fatal, TEXT("%s"), *Error);
+	}
+	return Result;
+}
+
+FZenPackageHeader FZenPackageHeader::MakeView(FMemoryView Memory, FString& OutError)
+{
+	OutError.Reset();
+
 	FZenPackageHeader PackageHeader;
 	const uint8* PackageHeaderDataPtr = reinterpret_cast<const uint8*>(Memory.GetData());
 	const FZenPackageSummary* PackageSummary = reinterpret_cast<const FZenPackageSummary*>(PackageHeaderDataPtr);
@@ -55,7 +68,7 @@ FZenPackageHeader FZenPackageHeader::MakeView(FMemoryView Memory)
 
 	if (ExportBundleEntriesCount != PackageHeader.ExportCount * FExportBundleEntry::ExportCommandType_Count)
 	{
-		UE_LOG(LogStreaming, Fatal, TEXT("Corrupt Zen header in package %s"), *PackageHeader.PackageName.ToString());
+		OutError = FString::Printf(TEXT("Corrupt Zen header in package %s"), *PackageHeader.PackageName.ToString());
 		return PackageHeader;
 	}
 
diff --git a/Engine/Source/Runtime/CoreUObject/Private/Serialization/ZenPackageHeader.h b/Engine/Source/Runtime/CoreUObject/Private/Serialization/ZenPackageHeader.h
deleted file mode 100644
