Best guess as to the cause of the pak crash. No local repro, and this should probably be done anyway.

diff --git a/Engine/Source/Developer/IoStoreUtilities/Private/IoStoreUtilities.cpp b/Engine/Source/Developer/IoStoreUtilities/Private/IoStoreUtilities.cpp
--- a/Engine/Source/Developer/IoStoreUtilities/Private/IoStoreUtilities.cpp
+++ b/Engine/Source/Developer/IoStoreUtilities/Private/IoStoreUtilities.cpp
@@ -3672,7 +3672,15 @@ static void InsertShadersInPluginHierarchy(UE::Cook::FCookMetadataState& InCookM
 		for (uint32 DependencyIndex = Entry.DependencyIndexStart; DependencyIndex < Entry.DependencyIndexEnd; DependencyIndex++)
 		{
 			const UE::Cook::FCookMetadataPluginEntry& DependentPlugin = PluginHierarchy.PluginsEnabledAtCook[PluginHierarchy.PluginDependencies[DependencyIndex]];
-			DependencyList.Add(PluginNameToIndex[DependentPlugin.Name]);
+			int32* PluginIndex = PluginNameToIndex.Find(DependentPlugin.Name);
+			if (PluginIndex)
+			{
+				DependencyList.Add(*PluginIndex);
+			}
+			else
+			{
+				UE_LOG(LogIoStore, Warning, TEXT("Couldn't find plugin %s when re-adding dependencies."), *DependentPlugin.Name);
+			}
 		}
 
 		// However we also need to check for shaders
@@ -3681,7 +3689,15 @@ static void InsertShadersInPluginHierarchy(UE::Cook::FCookMetadataState& InCookM
 		{
 			for (FString& ShaderPluginName : (*DependenciesOnShader))
 			{
-				DependencyList.Add(PluginNameToIndex[ShaderPluginName]);
+				int32* PluginIndex = PluginNameToIndex.Find(ShaderPluginName);
+				if (PluginIndex)
+				{
+					DependencyList.Add(*PluginIndex);
+				}
+				else
+				{
+					UE_LOG(LogIoStore, Warning, TEXT("Couldn't find shader pseudo plugin %s when adding dependencies."), *ShaderPluginName);
+				}
 			}
 		}
 
