Added Level asset type to the IO Store telemetry

diff --git a/Engine/Source/Developer/IoStoreUtilities/Private/IoStoreUtilities.cpp b/Engine/Source/Developer/IoStoreUtilities/Private/IoStoreUtilities.cpp
--- a/Engine/Source/Developer/IoStoreUtilities/Private/IoStoreUtilities.cpp
+++ b/Engine/Source/Developer/IoStoreUtilities/Private/IoStoreUtilities.cpp
@@ -3309,6 +3309,7 @@ enum class EPluginGraphSizeClass : uint8
 	SoundWave,
 	SkeletalMesh,
 	Shader,
+	Level, 
 	Other,
 	COUNT
 };
@@ -3321,6 +3322,7 @@ static const UTF8CHAR* PluginGraphEntryClassNames[] =
 	UTF8TEXT("soundwave"),
 	UTF8TEXT("skeletalmesh"),
 	UTF8TEXT("shader"),
+	UTF8TEXT("level"),
 	UTF8TEXT("other")
 };
 
@@ -3752,6 +3754,7 @@ static void UpdatePluginMetadataAndWriteJsons(
 	FTopLevelAssetPath StaticMeshPath(TEXT("/Script/Engine.StaticMesh"));
 	FTopLevelAssetPath SoundWavePath(TEXT("/Script/Engine.SoundWave"));
 	FTopLevelAssetPath SkeletalMeshPath(TEXT("/Script/Engine.SkeletalMesh"));
+	FTopLevelAssetPath LevelPath(TEXT("/Script/Engine.Level"));
 
 	
 	if (InShaderAssociationInfo)
@@ -3943,6 +3946,10 @@ static void UpdatePluginMetadataAndWriteJsons(
 					{
 						PluginEntry->ExclusiveSizes[(uint8)EPluginGraphSizeClass::SkeletalMesh].Add(PackageSizes);
 					}
+					else if (AssetData->AssetClassPath == LevelPath)
+					{
+						PluginEntry->ExclusiveSizes[(uint8)EPluginGraphSizeClass::Level].Add(PackageSizes);
+					}
 					// Note that we can't get shaders here so we don't need to handle ::Shader.
 					else
 					{
