Added Level, NiagaraSystem, MaterialInstance, Blueprint and AnimationSequence to asset telemetry tracking
Oprtimised plugin event data to filter out records where there is a zero inclusive size for a given asset type.

diff --git a/Engine/Source/Developer/IoStoreUtilities/Private/IoStoreUtilities.cpp b/Engine/Source/Developer/IoStoreUtilities/Private/IoStoreUtilities.cpp
--- a/Engine/Source/Developer/IoStoreUtilities/Private/IoStoreUtilities.cpp
+++ b/Engine/Source/Developer/IoStoreUtilities/Private/IoStoreUtilities.cpp
@@ -3310,6 +3310,10 @@ enum class EPluginGraphSizeClass : uint8
 	SkeletalMesh,
 	Shader,
 	Level, 
+	Animation,
+	Niagara,
+	Material,
+	Blueprint,
 	Other,
 	COUNT
 };
@@ -3323,6 +3327,10 @@ static const UTF8CHAR* PluginGraphEntryClassNames[] =
 	UTF8TEXT("skeletalmesh"),
 	UTF8TEXT("shader"),
 	UTF8TEXT("level"),
+	UTF8TEXT("animation"),
+	UTF8TEXT("niagara"),
+	UTF8TEXT("material"),
+	UTF8TEXT("blueprint"),
 	UTF8TEXT("other")
 };
 
@@ -3767,8 +3775,11 @@ static void UpdatePluginMetadataAndWriteJsons(
 	FTopLevelAssetPath StaticMeshPath(TEXT("/Script/Engine.StaticMesh"));
 	FTopLevelAssetPath SoundWavePath(TEXT("/Script/Engine.SoundWave"));
 	FTopLevelAssetPath SkeletalMeshPath(TEXT("/Script/Engine.SkeletalMesh"));
-	FTopLevelAssetPath LevelPath(TEXT("/Script/Engine.Level"));
-
+	FTopLevelAssetPath LevelPath(TEXT("/Script/Engine.World"));
+	FTopLevelAssetPath BlueprintPath(TEXT("/Script/Engine.BlueprintGeneratedClass"));
+	FTopLevelAssetPath AnimationSequencePath(TEXT("/Script/Engine.AnimSequence"));
+	FTopLevelAssetPath NiagaraSystemPath(TEXT("/Script/Niagara.NiagaraSystem"));
+	FTopLevelAssetPath MaterialInstancePath(TEXT("/Script/Engine.MaterialInstanceConstant"));
 	
 	if (InShaderAssociationInfo)
 	{
@@ -3959,10 +3970,26 @@ static void UpdatePluginMetadataAndWriteJsons(
 					{
 						PluginEntry->ExclusiveSizes[(uint8)EPluginGraphSizeClass::SkeletalMesh].Add(PackageSizes);
 					}
+					else if (AssetData->AssetClassPath == AnimationSequencePath)
+					{
+						PluginEntry->ExclusiveSizes[(uint8)EPluginGraphSizeClass::Animation].Add(PackageSizes);
+					}
+					else if (AssetData->AssetClassPath == NiagaraSystemPath)
+					{
+						PluginEntry->ExclusiveSizes[(uint8)EPluginGraphSizeClass::Niagara].Add(PackageSizes);
+					}
+					else if (AssetData->AssetClassPath == MaterialInstancePath)
+					{
+						PluginEntry->ExclusiveSizes[(uint8)EPluginGraphSizeClass::Material].Add(PackageSizes);
+					}
 					else if (AssetData->AssetClassPath == LevelPath)
 					{
 						PluginEntry->ExclusiveSizes[(uint8)EPluginGraphSizeClass::Level].Add(PackageSizes);
 					}
+					else if (AssetData->AssetClassPath == BlueprintPath)
+					{
+						PluginEntry->ExclusiveSizes[(uint8)EPluginGraphSizeClass::Blueprint].Add(PackageSizes);
+					}
 					// Note that we can't get shaders here so we don't need to handle ::Shader.
 					else
 					{
@@ -4028,15 +4055,21 @@ static void UpdatePluginMetadataAndWriteJsons(
 	auto GeneratePluginJson = [&MutablePluginHierarchy](TUtf8StringBuilder<4096>& OutPluginMetadataJson, FStringView InName, const FPluginGraphEntry& InGraphEntry, EPluginGraphSizeClass InSizeClass)
 	{
 		OutPluginMetadataJson.Reset();
+
+		uint8 SizeClass = (uint8)InSizeClass;
+
+		if (InGraphEntry.InclusiveSizes[SizeClass].TotalSize()== 0)
+		{
+			// Asset type or its dependencies do not contribute to the record.
+			return;
+		}
+		
 		OutPluginMetadataJson << "{\n";
 		OutPluginMetadataJson << "\t\"name\":\"" << InName << "\",\n";
 
 		OutPluginMetadataJson << "\t\"schema_version\":3,\n";
 
 		OutPluginMetadataJson << "\t\"is_root_plugin\":" << (InGraphEntry.bIsRoot ? TEXTVIEW("true") : TEXTVIEW("false")) << ",\n";
-
-		uint8 SizeClass = (uint8)InSizeClass;
-
 		OutPluginMetadataJson << "\t\"asset_sizes_class\":\"" << PluginGraphEntryClassNames[SizeClass] << "\",\n";
 
 		OutPluginMetadataJson << "\t\"exclusive_installed\":" << InGraphEntry.ExclusiveSizes[SizeClass][UE::Cook::EPluginSizeTypes::Installed] << ",\n";
@@ -4128,20 +4161,25 @@ static void UpdatePluginMetadataAndWriteJsons(
 		for (uint8 ClassIndex = 0; ClassIndex < FPluginGraphEntry::ClassCount; ClassIndex++)
 		{
 			JsonWrittenCount++;
+
 			GeneratePluginJson(PluginMetadataJson, Plugin.Name, PluginEntry, (EPluginGraphSizeClass)ClassIndex);
-			SavePluginMetadata(InAssetRegistryFileName, Plugin.Name, PluginMetadataJson, (EPluginGraphSizeClass)ClassIndex);
-
-			Csv.Appendf("%ls,%s,%llu,%llu,%llu,%llu,%llu,%llu,%llu,%llu,%llu,%u,%u\n", *Plugin.Name, PluginGraphEntryClassNames[ClassIndex],
-				PluginEntry.ExclusiveSizes[ClassIndex][UE::Cook::EPluginSizeTypes::Installed],
-				PluginEntry.ExclusiveSizes[ClassIndex][UE::Cook::EPluginSizeTypes::Optional],
-				PluginEntry.ExclusiveSizes[ClassIndex][UE::Cook::EPluginSizeTypes::Streaming],
-				PluginEntry.InclusiveSizes[ClassIndex][UE::Cook::EPluginSizeTypes::Installed],
-				PluginEntry.InclusiveSizes[ClassIndex][UE::Cook::EPluginSizeTypes::Optional],
-				PluginEntry.InclusiveSizes[ClassIndex][UE::Cook::EPluginSizeTypes::Streaming],
-				PluginEntry.UniqueSizes[ClassIndex][UE::Cook::EPluginSizeTypes::Installed],
-				PluginEntry.UniqueSizes[ClassIndex][UE::Cook::EPluginSizeTypes::Optional],
-				PluginEntry.UniqueSizes[ClassIndex][UE::Cook::EPluginSizeTypes::Streaming],
-				PluginEntry.DirectRefcount,PluginEntry.TotalDependencies.Num());
+
+			if (PluginMetadataJson.Len()>0)
+			{
+				SavePluginMetadata(InAssetRegistryFileName, Plugin.Name, PluginMetadataJson, (EPluginGraphSizeClass)ClassIndex);
+
+				Csv.Appendf("%ls,%s,%llu,%llu,%llu,%llu,%llu,%llu,%llu,%llu,%llu,%u,%u\n", *Plugin.Name, PluginGraphEntryClassNames[ClassIndex],
+					PluginEntry.ExclusiveSizes[ClassIndex][UE::Cook::EPluginSizeTypes::Installed],
+					PluginEntry.ExclusiveSizes[ClassIndex][UE::Cook::EPluginSizeTypes::Optional],
+					PluginEntry.ExclusiveSizes[ClassIndex][UE::Cook::EPluginSizeTypes::Streaming],
+					PluginEntry.InclusiveSizes[ClassIndex][UE::Cook::EPluginSizeTypes::Installed],
+					PluginEntry.InclusiveSizes[ClassIndex][UE::Cook::EPluginSizeTypes::Optional],
+					PluginEntry.InclusiveSizes[ClassIndex][UE::Cook::EPluginSizeTypes::Streaming],
+					PluginEntry.UniqueSizes[ClassIndex][UE::Cook::EPluginSizeTypes::Installed],
+					PluginEntry.UniqueSizes[ClassIndex][UE::Cook::EPluginSizeTypes::Optional],
+					PluginEntry.UniqueSizes[ClassIndex][UE::Cook::EPluginSizeTypes::Streaming],
+					PluginEntry.DirectRefcount, PluginEntry.TotalDependencies.Num());
+			}
 		}
 	}
 
