Cooker: ZenStore -diffonly: Fix typo in AppendPathForPublicExportHash that caused an infinite loop.
#rnx

diff --git a/Engine/Source/Developer/IoStoreUtilities/Private/PackageStoreOptimizer.cpp b/Engine/Source/Developer/IoStoreUtilities/Private/PackageStoreOptimizer.cpp
--- a/Engine/Source/Developer/IoStoreUtilities/Private/PackageStoreOptimizer.cpp
+++ b/Engine/Source/Developer/IoStoreUtilities/Private/PackageStoreOptimizer.cpp
@@ -287,7 +287,7 @@ void FPackageStoreOptimizer::AppendPathForPublicExportHash(UObject* Object, FStr
 
 	TArray<FName, TInlineAllocator<10>> PathNames;
 	int32 PathNameLen = 0;
-	for (UObject* Iter = Object; Iter->GetOuter(); Iter = Object->GetOuter())
+	for (UObject* Iter = Object; Iter->GetOuter(); Iter = Iter->GetOuter())
 	{
 		FName PathName = Iter->GetFName();
 		PathNames.Add(PathName);
