Added support for Geometry adn Virtual Texture Builder types to telemetry

diff --git a/Engine/Source/Developer/IoStoreUtilities/Private/IoStoreUtilities.cpp b/Engine/Source/Developer/IoStoreUtilities/Private/IoStoreUtilities.cpp
--- a/Engine/Source/Developer/IoStoreUtilities/Private/IoStoreUtilities.cpp
+++ b/Engine/Source/Developer/IoStoreUtilities/Private/IoStoreUtilities.cpp
@@ -3314,6 +3314,7 @@ enum class EPluginGraphSizeClass : uint8
 	Niagara,
 	Material,
 	Blueprint,
+	Geometry,
 	Other,
 	COUNT
 };
@@ -3331,6 +3332,7 @@ static const UTF8CHAR* PluginGraphEntryClassNames[] =
 	UTF8TEXT("niagara"),
 	UTF8TEXT("material"),
 	UTF8TEXT("blueprint"),
+	UTF8TEXT("geometry"),
 	UTF8TEXT("other")
 };
 
@@ -3772,12 +3774,14 @@ static void UpdatePluginMetadataAndWriteJsons(
 	FTopLevelAssetPath Texture3DPath(TEXT("/Script/Engine.Texture3D"));
 	FTopLevelAssetPath TextureCubePath(TEXT("/Script/Engine.TextureCube"));
 	FTopLevelAssetPath TextureCubeArrayPath(TEXT("/Script/Engine.TextureCubeArray"));
+	FTopLevelAssetPath VirtualTextureBuilderPath(TEXT("/Script/Engine.VirtualTextureBuilder"));
 	FTopLevelAssetPath StaticMeshPath(TEXT("/Script/Engine.StaticMesh"));
 	FTopLevelAssetPath SoundWavePath(TEXT("/Script/Engine.SoundWave"));
 	FTopLevelAssetPath SkeletalMeshPath(TEXT("/Script/Engine.SkeletalMesh"));
 	FTopLevelAssetPath LevelPath(TEXT("/Script/Engine.World"));
 	FTopLevelAssetPath BlueprintPath(TEXT("/Script/Engine.BlueprintGeneratedClass"));
 	FTopLevelAssetPath AnimationSequencePath(TEXT("/Script/Engine.AnimSequence"));
+	FTopLevelAssetPath GeometryCollectionPath(TEXT("/Script/GeometryCollectionEngine.GeometryCollection"));
 	FTopLevelAssetPath NiagaraSystemPath(TEXT("/Script/Niagara.NiagaraSystem"));
 	FTopLevelAssetPath MaterialInstancePath(TEXT("/Script/Engine.MaterialInstanceConstant"));
 	
@@ -3954,7 +3958,8 @@ static void UpdatePluginMetadataAndWriteJsons(
 						AssetData->AssetClassPath == Texture3DPath ||
 						AssetData->AssetClassPath == TextureCubePath ||
 						AssetData->AssetClassPath == TextureCubeArrayPath ||
-						AssetData->AssetClassPath == Texture2DArrayPath)
+						AssetData->AssetClassPath == Texture2DArrayPath || 
+						AssetData->AssetClassPath == VirtualTextureBuilderPath)
 					{
 						PluginEntry->ExclusiveSizes[(uint8)EPluginGraphSizeClass::Texture].Add(PackageSizes);
 					}
@@ -3990,6 +3995,11 @@ static void UpdatePluginMetadataAndWriteJsons(
 					{
 						PluginEntry->ExclusiveSizes[(uint8)EPluginGraphSizeClass::Blueprint].Add(PackageSizes);
 					}
+					else if (AssetData->AssetClassPath == GeometryCollectionPath)
+					{
+						PluginEntry->ExclusiveSizes[(uint8)EPluginGraphSizeClass::Geometry].Add(PackageSizes);
+					}
+					
 					// Note that we can't get shaders here so we don't need to handle ::Shader.
 					else
 					{
