Pack the plugin jsons in to one file to avoid nuking network file protocols.

diff --git a/Engine/Source/Developer/IoStoreUtilities/Private/IoStoreUtilities.cpp b/Engine/Source/Developer/IoStoreUtilities/Private/IoStoreUtilities.cpp
--- a/Engine/Source/Developer/IoStoreUtilities/Private/IoStoreUtilities.cpp
+++ b/Engine/Source/Developer/IoStoreUtilities/Private/IoStoreUtilities.cpp
@@ -3338,18 +3338,6 @@ static const UTF8CHAR* PluginGraphEntryClassNames[] =
 
 static_assert( UE_ARRAY_COUNT(PluginGraphEntryClassNames) == (size_t)EPluginGraphSizeClass::COUNT, "Must have a name for each plugin graph size class!");
 
-static bool SavePluginMetadata(const FString& InAssetRegistryFileName, const FString& PluginName, TUtf8StringBuilder<4096>& InPluginMetadataJson, EPluginGraphSizeClass InAssetClass)
-{
-	FString PluginMetadataFilename = FPaths::GetPath(InAssetRegistryFileName) / TEXT("PluginJsons") / PluginName + TEXT("_") + PluginGraphEntryClassNames[(uint8)InAssetClass] + TEXT(".json");
-	if (WriteUtf8StringView(InPluginMetadataJson.ToView(), PluginMetadataFilename) == false)
-	{
-		UE_LOG(LogIoStore, Error, TEXT("Unable to write plugin metadata file: %s"), *PluginMetadataFilename);
-		return false;
-	}
-	return true;
-}
-
-
 struct FPluginGraphEntry
 {
 	uint16 IndexInEnabledPlugins = 0;
@@ -4142,7 +4130,6 @@ static void UpdatePluginMetadataAndWriteJsons(
 	//
 	// Generate the plugin_summary jsons.	
 	//
-	TUtf8StringBuilder<4096> PluginMetadataJson;
 
 	// Also write a csv for easier browsing in spreadsheets.
 	TUtf8StringBuilder<4096> Csv;
@@ -4155,8 +4142,25 @@ static void UpdatePluginMetadataAndWriteJsons(
 		Plugin.InclusiveSizes.Zero();
 		Plugin.ExclusiveSizes.Zero();
 	}
+	
+	// Instead of writing out a ton of small event files we concat them all. There are a lot in a mature project.
+	// This will be megabytes.
+	TArray<UTF8CHAR> PluginMetadataFullJson;
+	uint32 PluginsAddedToFullJson = 0;
+	PluginMetadataFullJson.Append(UTF8TEXTVIEW("{ \"PluginSizeInfos\": ["));
+
+	auto AddPluginJsonToFull = [&PluginMetadataFullJson, &PluginsAddedToFullJson](TUtf8StringBuilder<4096>& InJsonToAdd)
+	{
+		if (PluginsAddedToFullJson)
+		{
+			PluginMetadataFullJson.Add(UTF8TEXT(','));
+		}
+		PluginsAddedToFullJson++;
+		PluginMetadataFullJson.Append(InJsonToAdd.GetData(), InJsonToAdd.Len());
+	};
+
+	TUtf8StringBuilder<4096> PluginMetadataJson;
 
-	uint32 JsonWrittenCount = 0;
 	for (UE::Cook::FCookMetadataPluginEntry& Plugin : MutablePluginHierarchy.PluginsEnabledAtCook)
 	{
 		const FPluginGraphEntry& PluginEntry = *PluginGraph.NameToPlugin[Plugin.Name];
@@ -4170,14 +4174,12 @@ static void UpdatePluginMetadataAndWriteJsons(
 
 		for (uint8 ClassIndex = 0; ClassIndex < FPluginGraphEntry::ClassCount; ClassIndex++)
 		{
-			JsonWrittenCount++;
 
 			GeneratePluginJson(PluginMetadataJson, Plugin.Name, PluginEntry, (EPluginGraphSizeClass)ClassIndex);
 
 			if (PluginMetadataJson.Len()>0)
 			{
-				SavePluginMetadata(InAssetRegistryFileName, Plugin.Name, PluginMetadataJson, (EPluginGraphSizeClass)ClassIndex);
-
+				AddPluginJsonToFull(PluginMetadataJson);
 				Csv.Appendf("%ls,%s,%llu,%llu,%llu,%llu,%llu,%llu,%llu,%llu,%llu,%u,%u\n", *Plugin.Name, PluginGraphEntryClassNames[ClassIndex],
 					PluginEntry.ExclusiveSizes[ClassIndex][UE::Cook::EPluginSizeTypes::Installed],
 					PluginEntry.ExclusiveSizes[ClassIndex][UE::Cook::EPluginSizeTypes::Optional],
@@ -4201,14 +4203,26 @@ static void UpdatePluginMetadataAndWriteJsons(
 		OrphanedEntry.DirectRefcount = 0;
 		OrphanedEntry.InclusiveSizes[(uint8)EPluginGraphSizeClass::All] = UnrootedTotal;
 		GeneratePluginJson(PluginMetadataJson, TEXT("OrphanedPlugins"), OrphanedEntry, EPluginGraphSizeClass::All);
-		SavePluginMetadata(InAssetRegistryFileName, TEXT("OrphanedPlugins"), PluginMetadataJson, EPluginGraphSizeClass::All);
-		JsonWrittenCount++;
+		if (PluginMetadataJson.Len() > 0)
+		{
+			AddPluginJsonToFull(PluginMetadataJson);
+		}
 
 		Csv.Appendf("OrphanedPlugins,all,0,0,0,%llu,%llu,%llu,0,%u\n",
 			UnrootedTotal[UE::Cook::EPluginSizeTypes::Installed], UnrootedTotal[UE::Cook::EPluginSizeTypes::Optional], UnrootedTotal[UE::Cook::EPluginSizeTypes::Streaming],
 			PluginGraph.UnrootedPlugins.Num());
 	}
 
+	PluginMetadataFullJson.Append(UTF8TEXTVIEW("]}"));
+	{
+		FString JsonFilename = FPaths::GetPath(InAssetRegistryFileName) / TEXT("plugin_size_jsons.json");
+		if (WriteUtf8StringView(MakeStringView(PluginMetadataFullJson), JsonFilename) == false)
+		{
+			UE_LOG(LogIoStore, Error, TEXT("Unable to write plugin json file: %s"), *JsonFilename);
+			return;
+		}
+	}
+
 	{
 		FString CsvFilename = FPaths::GetPath(InAssetRegistryFileName) / TEXT("plugin_sizes.csv");
 		if (WriteUtf8StringView(Csv.ToView(), CsvFilename) == false)
@@ -4219,8 +4233,7 @@ static void UpdatePluginMetadataAndWriteJsons(
 	}
 
 	double WritePluginEnd = FPlatformTime::Seconds();
-	UE_LOG(LogIoStore, Display, TEXT("Wrote %s plugin size jsons/csv in %.2f seconds [graph %.2f shaders %.2f sizes %.2f writes %.2f]"), 
-		*NumberString(JsonWrittenCount),
+	UE_LOG(LogIoStore, Display, TEXT("Wrote plugin size jsons/csv in %.2f seconds [graph %.2f shaders %.2f sizes %.2f writes %.2f]"), 
 		WritePluginEnd - WritePluginStart, 
 		GeneratePluginGraphEnd - WritePluginStart, 
 		AssetPackageMapStart - GeneratePluginGraphEnd,
