Fix for build warning with missing shader plugin dependencies.

diff --git a/Engine/Source/Developer/IoStoreUtilities/Private/IoStoreUtilities.cpp b/Engine/Source/Developer/IoStoreUtilities/Private/IoStoreUtilities.cpp
--- a/Engine/Source/Developer/IoStoreUtilities/Private/IoStoreUtilities.cpp
+++ b/Engine/Source/Developer/IoStoreUtilities/Private/IoStoreUtilities.cpp
@@ -3676,6 +3676,13 @@ static void InsertShadersInPluginHierarchy(UE::Cook::FCookMetadataState& InCookM
 		for (uint32 DependencyIndex = Entry.DependencyIndexStart; DependencyIndex < Entry.DependencyIndexEnd; DependencyIndex++)
 		{
 			const UE::Cook::FCookMetadataPluginEntry& DependentPlugin = PluginHierarchy.PluginsEnabledAtCook[PluginHierarchy.PluginDependencies[DependencyIndex]];
+			if (DependentPlugin.Type == UE::Cook::ECookMetadataPluginType::ShaderPseudo)
+			{
+				// If we are rerunning stage then we might already have added shader plugins
+				// as dependencies, which we've removed so they won't exist in the lookup
+				// (and we want to redo them anyway).
+				continue;
+			}
 			int32* PluginIndex = PluginNameToIndex.Find(DependentPlugin.Name);
 			if (PluginIndex)
 			{
