Change default port for autolaunched ZenServer local service to 8558.  When connecting to an existing service (not autolaunched), the default port remains what it was before, but will change in the future.

diff --git a/Engine/Source/Developer/IoStoreUtilities/Internal/ZenStoreHttpClient.h b/Engine/Source/Developer/IoStoreUtilities/Internal/ZenStoreHttpClient.h
--- a/Engine/Source/Developer/IoStoreUtilities/Internal/ZenStoreHttpClient.h
+++ b/Engine/Source/Developer/IoStoreUtilities/Internal/ZenStoreHttpClient.h
@@ -54,9 +54,9 @@ public:
 	const TCHAR* GetHostName() const { return ZenService.GetInstance().GetHostName(); }
 	uint16 GetPort() const { return ZenService.GetInstance().GetPort(); }
 	const UE::Zen::FZenServiceInstance& GetZenServiceInstance() const { return ZenService.GetInstance(); }
-#else // Default to localhost:1337 for platforms where Zen wouldn't be supported yet
+#else // Default to localhost:8558 for platforms where Zen wouldn't be supported yet
 	const TCHAR* GetHostName() const { return TEXT("localhost"); }
-	uint16 GetPort() const { return 1337; }
+	uint16 GetPort() const { return 8558; }
 #endif
 
 	TFuture<TIoStatusOr<FCbObject>> GetOplog();
diff --git a/Engine/Source/Developer/Zen/Public/Experimental/ZenServerInterface.h b/Engine/Source/Developer/Zen/Public/Experimental/ZenServerInterface.h
--- a/Engine/Source/Developer/Zen/Public/Experimental/ZenServerInterface.h
+++ b/Engine/Source/Developer/Zen/Public/Experimental/ZenServerInterface.h
@@ -39,7 +39,7 @@ struct FServiceAutoLaunchSettings
 {
 	FString DataPath;
 	FString ExtraArgs;
-	uint16 DesiredPort = 1337;
+	uint16 DesiredPort = 8558;
 	bool bShowConsole = false;
 	bool bLimitProcessLifetime = false;
 	bool bSendUnattendedBugReports = false;
diff --git a/Engine/Source/Editor/UnrealEd/Private/TargetDomain/TargetDomainUtils.cpp b/Engine/Source/Editor/UnrealEd/Private/TargetDomain/TargetDomainUtils.cpp
--- a/Engine/Source/Editor/UnrealEd/Private/TargetDomain/TargetDomainUtils.cpp
+++ b/Engine/Source/Editor/UnrealEd/Private/TargetDomain/TargetDomainUtils.cpp
@@ -427,9 +427,9 @@ TArray<const UTF8CHAR*> FEditorDomainOplog::ReservedOplogKeys;
 
 FEditorDomainOplog::FEditorDomainOplog()
 #if UE_WITH_ZEN
-: HttpClient(TEXT("localhost"), UE::Zen::FZenServiceInstance::GetAutoLaunchedPort() > 0 ? UE::Zen::FZenServiceInstance::GetAutoLaunchedPort() : 1337)
+: HttpClient(TEXT("localhost"), UE::Zen::FZenServiceInstance::GetAutoLaunchedPort() > 0 ? UE::Zen::FZenServiceInstance::GetAutoLaunchedPort() : 8558)
 #else
-: HttpClient(TEXT("localhost"), 1337)
+: HttpClient(TEXT("localhost"), 8558)
 #endif
 {
 	StaticInit();
diff --git a/Engine/Source/Programs/AutomationTool/BuildGraph/Tasks/ZenExportSnapshotTask.cs b/Engine/Source/Programs/AutomationTool/BuildGraph/Tasks/ZenExportSnapshotTask.cs
--- a/Engine/Source/Programs/AutomationTool/BuildGraph/Tasks/ZenExportSnapshotTask.cs
+++ b/Engine/Source/Programs/AutomationTool/BuildGraph/Tasks/ZenExportSnapshotTask.cs
@@ -371,7 +371,7 @@ namespace AutomationTool.Tasks
 					ExportSourceData NewExportSource = new ExportSourceData();
 					NewExportSource.IsLocalHost = ZenServerObject["islocalhost"].AsBool();
 					NewExportSource.HostName = ZenServerObject["hostname"].AsString("localhost");
-					NewExportSource.HostPort = ZenServerObject["hostport"].AsInt16(1337);
+					NewExportSource.HostPort = ZenServerObject["hostport"].AsInt16(8558);
 					NewExportSource.ProjectId = ZenServerObject["projectid"].AsString();
 					NewExportSource.OplogId = ZenServerObject["oplogid"].AsString();
 					NewExportSource.TargetPlatform = Platform;
diff --git a/Engine/Source/Programs/AutomationTool/BuildGraph/Tasks/ZenImportOplogTask.cs b/Engine/Source/Programs/AutomationTool/BuildGraph/Tasks/ZenImportOplogTask.cs
--- a/Engine/Source/Programs/AutomationTool/BuildGraph/Tasks/ZenImportOplogTask.cs
+++ b/Engine/Source/Programs/AutomationTool/BuildGraph/Tasks/ZenImportOplogTask.cs
@@ -74,7 +74,7 @@ namespace AutomationTool.Tasks
 		/// The host port for the zen server we are importing from
 		/// </summary>
 		[TaskParameter(Optional = true)]
-		public string HostPort = "1337";
+		public string HostPort = "8558";
 		
 
 		/// <summary>
@@ -227,7 +227,7 @@ namespace AutomationTool.Tasks
 
 			Writer.WriteBool("islocalhost", true);
 			Writer.WriteString("hostname", "localhost");
-			Writer.WriteInteger("hostport", 1337);
+			Writer.WriteInteger("hostport", 8558);
 			Writer.WriteString("projectid", Parameters.ProjectName);
 			Writer.WriteString("oplogid", Parameters.OplogName);
 
diff --git a/Engine/Source/Programs/AutomationTool/Scripts/CopyBuildToStagingDirectory.Automation.cs b/Engine/Source/Programs/AutomationTool/Scripts/CopyBuildToStagingDirectory.Automation.cs
--- a/Engine/Source/Programs/AutomationTool/Scripts/CopyBuildToStagingDirectory.Automation.cs
+++ b/Engine/Source/Programs/AutomationTool/Scripts/CopyBuildToStagingDirectory.Automation.cs
@@ -787,7 +787,7 @@ namespace AutomationScripts
 				{
 					Logger.LogInformation("Reading oplog from Zen...");
 					string Host = "localhost";
-					int Port = 1337;
+					int Port = 8558;
 					string ProjectId = ZenServerObject["projectid"].AsString();
 					string OplogId = ZenServerObject["oplogid"].AsString();
 					bool bAutoLaunch = false;
diff --git a/Engine/Source/Programs/AutomationTool/Scripts/RunProjectCommand.Automation.cs b/Engine/Source/Programs/AutomationTool/Scripts/RunProjectCommand.Automation.cs
--- a/Engine/Source/Programs/AutomationTool/Scripts/RunProjectCommand.Automation.cs
+++ b/Engine/Source/Programs/AutomationTool/Scripts/RunProjectCommand.Automation.cs
@@ -836,7 +836,7 @@ namespace AutomationScripts
 			}
 
 			string ProjectId = ProjectUtils.GetProjectPathId(SC.RawProjectPath);
-			ushort ZenHostPort = 1337;
+			ushort ZenHostPort = 8558;
 			List<string> HostAddresses = null;
 			if (ZenStoreMarkerFile != null)
 			{
diff --git a/Engine/Source/Programs/OidcToken/TokenService.cs b/Engine/Source/Programs/OidcToken/TokenService.cs
--- a/Engine/Source/Programs/OidcToken/TokenService.cs
+++ b/Engine/Source/Programs/OidcToken/TokenService.cs
@@ -293,7 +293,7 @@ namespace OidcToken
 		/// <summary>
 		/// The url to the zen server that we should share a token with (if Zen flag is set)
 		/// </summary>
-		public string ZenUrl { get; set; } = "http://localhost:1337";
+		public string ZenUrl { get; set; } = "http://localhost:8558";
 		
 		/// <summary>
 		/// Path to the game root directory
diff --git a/Engine/Source/Runtime/StorageServerClient/Private/StorageServerPlatformFile.h b/Engine/Source/Runtime/StorageServerClient/Private/StorageServerPlatformFile.h
--- a/Engine/Source/Runtime/StorageServerClient/Private/StorageServerPlatformFile.h
+++ b/Engine/Source/Runtime/StorageServerClient/Private/StorageServerPlatformFile.h
@@ -124,7 +124,7 @@ private:
 	FString ServerProject;
 	FString ServerPlatform;
 	mutable TArray<FString> HostAddrs;
-	mutable uint16 HostPort = 1337;
+	mutable uint16 HostPort = 8558;
 };
 
 #endif
