Shader dependencies incorrectly were added as base-1.

diff --git a/Engine/Source/Developer/IoStoreUtilities/Private/IoStoreUtilities.cpp b/Engine/Source/Developer/IoStoreUtilities/Private/IoStoreUtilities.cpp
--- a/Engine/Source/Developer/IoStoreUtilities/Private/IoStoreUtilities.cpp
+++ b/Engine/Source/Developer/IoStoreUtilities/Private/IoStoreUtilities.cpp
@@ -3873,7 +3873,7 @@ static void UpdatePluginMetadataAndWriteJsons(
 			for (FName ReferencingPackage : ShaderChunkInfo.Value.ReferencedByPackages)
 			{
 				TArray<int32>& PackageDependencies = PackageDependencyMap.FindOrAdd(ReferencingPackage);
-				PackageDependencies.Add(ShaderPseudoAssets.Num());
+				PackageDependencies.Add(ShaderPseudoAssets.Num() - 1);
 			}
 		}
 
