Adding asset class count to plugin analytics.

diff --git a/Engine/Source/Developer/IoStoreUtilities/Private/IoStoreUtilities.cpp b/Engine/Source/Developer/IoStoreUtilities/Private/IoStoreUtilities.cpp
--- a/Engine/Source/Developer/IoStoreUtilities/Private/IoStoreUtilities.cpp
+++ b/Engine/Source/Developer/IoStoreUtilities/Private/IoStoreUtilities.cpp
@@ -3345,7 +3345,7 @@ struct FPluginGraphEntry
 	FString Name;
 
 	TSet<FPluginGraphEntry*> DirectDependencies;
-	TSet<FPluginGraphEntry*> TotalDependencies;		
+	TSet<FPluginGraphEntry*> TotalDependencies;
 	TSet<FPluginGraphEntry*> Roots;
 
 	// Only valid if bIsRoot
@@ -3358,6 +3358,10 @@ struct FPluginGraphEntry
 	UE::Cook::FPluginSizeInfo ExclusiveSizes[ClassCount];
 	UE::Cook::FPluginSizeInfo InclusiveSizes[ClassCount];
 	UE::Cook::FPluginSizeInfo UniqueSizes[ClassCount];
+
+	uint64 ExclusiveCounts[ClassCount] = {};
+	uint64 InclusiveCounts[ClassCount] = {};
+	uint64 UniqueCounts[ClassCount] = {};
 };
 
 struct FPluginGraph
@@ -3949,6 +3953,7 @@ static void UpdatePluginMetadataAndWriteJsons(
 				// If we have asset class info and it's a top contender, track it also.
 				if (AssetData != nullptr)
 				{
+					EPluginGraphSizeClass AssetSizeClass = EPluginGraphSizeClass::Other;
 					if (AssetData->AssetClassPath == Texture2DPath ||
 						AssetData->AssetClassPath == Texture3DPath ||
 						AssetData->AssetClassPath == TextureCubePath ||
@@ -3956,50 +3961,49 @@ static void UpdatePluginMetadataAndWriteJsons(
 						AssetData->AssetClassPath == Texture2DArrayPath || 
 						AssetData->AssetClassPath == VirtualTextureBuilderPath)
 					{
-						PluginEntry->ExclusiveSizes[(uint8)EPluginGraphSizeClass::Texture].Add(PackageSizes);
+						AssetSizeClass = EPluginGraphSizeClass::Texture;
 					}
 					else if (AssetData->AssetClassPath == StaticMeshPath)
 					{
-						PluginEntry->ExclusiveSizes[(uint8)EPluginGraphSizeClass::StaticMesh].Add(PackageSizes);
+						AssetSizeClass = EPluginGraphSizeClass::StaticMesh;
 					}
 					else if (AssetData->AssetClassPath == SoundWavePath)
 					{
-						PluginEntry->ExclusiveSizes[(uint8)EPluginGraphSizeClass::SoundWave].Add(PackageSizes);
+						AssetSizeClass = EPluginGraphSizeClass::SoundWave;
 					}
 					else if (AssetData->AssetClassPath == SkeletalMeshPath)
 					{
-						PluginEntry->ExclusiveSizes[(uint8)EPluginGraphSizeClass::SkeletalMesh].Add(PackageSizes);
+						AssetSizeClass = EPluginGraphSizeClass::SkeletalMesh;
 					}
 					else if (AssetData->AssetClassPath == AnimationSequencePath)
 					{
-						PluginEntry->ExclusiveSizes[(uint8)EPluginGraphSizeClass::Animation].Add(PackageSizes);
+						AssetSizeClass = EPluginGraphSizeClass::Animation;
 					}
 					else if (AssetData->AssetClassPath == NiagaraSystemPath)
 					{
-						PluginEntry->ExclusiveSizes[(uint8)EPluginGraphSizeClass::Niagara].Add(PackageSizes);
+						AssetSizeClass = EPluginGraphSizeClass::Niagara;
 					}
 					else if (AssetData->AssetClassPath == MaterialInstancePath)
 					{
-						PluginEntry->ExclusiveSizes[(uint8)EPluginGraphSizeClass::Material].Add(PackageSizes);
+						AssetSizeClass = EPluginGraphSizeClass::Material;
 					}
 					else if (AssetData->AssetClassPath == LevelPath)
 					{
-						PluginEntry->ExclusiveSizes[(uint8)EPluginGraphSizeClass::Level].Add(PackageSizes);
+						AssetSizeClass = EPluginGraphSizeClass::Level;
 					}
 					else if (AssetData->AssetClassPath == BlueprintPath)
 					{
-						PluginEntry->ExclusiveSizes[(uint8)EPluginGraphSizeClass::Blueprint].Add(PackageSizes);
+						AssetSizeClass = EPluginGraphSizeClass::Blueprint;
 					}
 					else if (AssetData->AssetClassPath == GeometryCollectionPath)
 					{
-						PluginEntry->ExclusiveSizes[(uint8)EPluginGraphSizeClass::Geometry].Add(PackageSizes);
+						AssetSizeClass = EPluginGraphSizeClass::Geometry;
 					}
 					
 					// Note that we can't get shaders here so we don't need to handle ::Shader.
-					else
-					{
-						PluginEntry->ExclusiveSizes[(uint8)EPluginGraphSizeClass::Other].Add(PackageSizes);
-					}
+
+					PluginEntry->ExclusiveSizes[(uint8)AssetSizeClass].Add(PackageSizes);
+					PluginEntry->ExclusiveCounts[(uint8)AssetSizeClass]++;
 				}
 			}
 			else
@@ -4021,6 +4025,7 @@ static void UpdatePluginMetadataAndWriteJsons(
 		for (uint8 ClassIndex = 0; ClassIndex < FPluginGraphEntry::ClassCount; ClassIndex++)
 		{
 			PluginEntry->InclusiveSizes[ClassIndex] = PluginEntry->ExclusiveSizes[ClassIndex];
+			PluginEntry->InclusiveCounts[ClassIndex] = PluginEntry->ExclusiveCounts[ClassIndex];
 		}
 
 		for (FPluginGraphEntry* Dependency : PluginEntry->TotalDependencies)
@@ -4028,6 +4033,7 @@ static void UpdatePluginMetadataAndWriteJsons(
 			for (uint8 ClassIndex = 0; ClassIndex < FPluginGraphEntry::ClassCount; ClassIndex++)
 			{
 				PluginEntry->InclusiveSizes[ClassIndex].Add(Dependency->ExclusiveSizes[ClassIndex]);
+				PluginEntry->InclusiveCounts[ClassIndex] += Dependency->ExclusiveCounts[ClassIndex];
 			}			
 		}
 	}
@@ -4042,6 +4048,7 @@ static void UpdatePluginMetadataAndWriteJsons(
 			for (uint8 ClassIndex = 0; ClassIndex < FPluginGraphEntry::ClassCount; ClassIndex++)
 			{
 				RootPlugin->UniqueSizes[ClassIndex].Add(UniqueDependency->ExclusiveSizes[ClassIndex]);
+				RootPlugin->UniqueCounts[ClassIndex] += UniqueDependency->ExclusiveCounts[ClassIndex];
 			}
 		}
 	}
@@ -4072,10 +4079,11 @@ static void UpdatePluginMetadataAndWriteJsons(
 		OutPluginMetadataJson << "{\n";
 		OutPluginMetadataJson << "\t\"name\":\"" << InName << "\",\n";
 
-		OutPluginMetadataJson << "\t\"schema_version\":3,\n";
+		OutPluginMetadataJson << "\t\"schema_version\":4,\n";
 
 		OutPluginMetadataJson << "\t\"is_root_plugin\":" << (InGraphEntry.bIsRoot ? TEXTVIEW("true") : TEXTVIEW("false")) << ",\n";
 		OutPluginMetadataJson << "\t\"asset_sizes_class\":\"" << PluginGraphEntryClassNames[SizeClass] << "\",\n";
+		OutPluginMetadataJson << "\t\"exclusive_asset_class_count\":" << InGraphEntry.ExclusiveCounts[SizeClass] << ",\n";
 
 		OutPluginMetadataJson << "\t\"exclusive_installed\":" << InGraphEntry.ExclusiveSizes[SizeClass][UE::Cook::EPluginSizeTypes::Installed] << ",\n";
 		OutPluginMetadataJson << "\t\"exclusive_optional\":" << InGraphEntry.ExclusiveSizes[SizeClass][UE::Cook::EPluginSizeTypes::Optional] << ",\n";
