IAS - separate out trailling host path when writing the configuraiton .ini file

diff --git a/Engine/Source/Developer/IoStoreUtilities/Private/IoStoreUtilities.cpp b/Engine/Source/Developer/IoStoreUtilities/Private/IoStoreUtilities.cpp
--- a/Engine/Source/Developer/IoStoreUtilities/Private/IoStoreUtilities.cpp
+++ b/Engine/Source/Developer/IoStoreUtilities/Private/IoStoreUtilities.cpp
@@ -9252,10 +9252,46 @@ bool UploadIoStoreContainerFiles(const UE::IO::IAS::FIoStoreUploadParams& Upload
 		if (FParse::Value(FCommandLine::Get(), TEXT("DistributionUrl="), DistributionUrl))
 		{
 			Sb << TEXT("DistributionUrl=\"") << DistributionUrl << TEXT("\"\r\n");
+			Sb << TEXT("TocPath=\"") << UploadResult.TocPath << TEXT("\"\r\n");
 		}
 		else
 		{
-			Sb << TEXT("ServiceUrl=\"") << UploadParams.ServiceUrl << TEXT("\"\r\n");
+			// The configuration file should specify a service URL without any trailing
+			// host path, i.e. http://{host:port}/{host-path}. Add the trailing path
+			// to the TOC path to form the complete path the TOC from the host, i.e
+			// TocPath={host-path}/{bucket}/{bucket-prefix}/{toc-hash}.uchunktoc
+
+			FStringView ServiceUrl = UploadParams.ServiceUrl;
+			FStringView TocPrefx;
+			{
+				int32 Sep = INDEX_NONE;
+				int32 Count = 0;
+				for (int32 Idx = 0, Len = ServiceUrl.Len(); Idx < Len; ++Idx)
+				{
+					if (ServiceUrl[Idx] == TEXT('/'))
+					{
+						if (++Count == 3)
+						{
+							Sep = Idx;
+							break;
+						}
+					}
+				}
+
+				if (Sep != INDEX_NONE)
+				{
+					TocPrefx = ServiceUrl.RightChop(Sep + 1);
+					ServiceUrl.LeftInline(Sep);
+				}
+			}
+
+			Sb << TEXT("ServiceUrl=\"") << ServiceUrl << TEXT("\"\r\n");
+			Sb << TEXT("TocPath=\"");
+			if (TocPrefx.IsEmpty() == false)
+			{
+				Sb << TocPrefx << TEXT('/');
+			}
+			Sb << UploadParams.Bucket / UploadResult.TocPath << TEXT("\"\r\n");
 		}
 
 		// Temporary solution to get replays working with encrypted on demand content
@@ -9284,16 +9320,6 @@ bool UploadIoStoreContainerFiles(const UE::IO::IAS::FIoStoreUploadParams& Upload
 			}
 		}
 
-		if (DistributionUrl.IsEmpty())
-		{
-			// Append the bucket name when using a local/custom service
-			Sb << TEXT("TocPath=\"") << UploadParams.Bucket / UploadResult.TocPath << TEXT("\"\r\n");
-		}
-		else
-		{
-			Sb << TEXT("TocPath=\"") << UploadResult.TocPath << TEXT("\"\r\n");
-		}
-
 		UE_LOG(LogIoStore, Display, TEXT("Saving on demand config file '%s'"), *ConfigFilePath);
 		if (FFileHelper::SaveStringToFile(Sb.ToString(), *ConfigFilePath) == false)
 		{
