UnrealPak - expose -List option for dumping container file info to csv

diff --git a/Engine/Source/Developer/IoStoreUtilities/Private/IoStoreUtilities.cpp b/Engine/Source/Developer/IoStoreUtilities/Private/IoStoreUtilities.cpp
--- a/Engine/Source/Developer/IoStoreUtilities/Private/IoStoreUtilities.cpp
+++ b/Engine/Source/Developer/IoStoreUtilities/Private/IoStoreUtilities.cpp
@@ -5959,6 +5959,28 @@ int32 ListContainer(
 	return 0;
 }
 
+bool ListIoStoreContainer(const TCHAR* CmdLine)
+{
+	FKeyChain KeyChain;
+	LoadKeyChain(CmdLine, KeyChain);
+
+	FString ContainerPathOrWildcard;
+	if (!FParse::Value(FCommandLine::Get(), TEXT("ListContainer="), ContainerPathOrWildcard))
+	{
+		UE_LOG(LogIoStore, Error, TEXT("Missing argument -ListContainer=<ContainerFileOrWildCard>"));
+		return false;
+	}
+
+	FString CsvPath;
+	if (!FParse::Value(FCommandLine::Get(), TEXT("csv="), CsvPath))
+	{
+		UE_LOG(LogIoStore, Error, TEXT("Missing argument -Csv=<Path>"));
+		return false;
+	}
+
+	return ListContainer(KeyChain, ContainerPathOrWildcard, CsvPath) == 0;
+}
+
 bool LegacyListIoStoreContainer(
 	const TCHAR* InContainerFilename,
 	int64 InSizeFilter,
diff --git a/Engine/Source/Developer/IoStoreUtilities/Public/IoStoreUtilities.h b/Engine/Source/Developer/IoStoreUtilities/Public/IoStoreUtilities.h
--- a/Engine/Source/Developer/IoStoreUtilities/Public/IoStoreUtilities.h
+++ b/Engine/Source/Developer/IoStoreUtilities/Public/IoStoreUtilities.h
@@ -13,17 +13,21 @@ struct FGuid;
 struct FKeyChain;
 class FIoChunkId;
 
-IOSTOREUTILITIES_API int32 CreateIoStoreContainerFiles(const TCHAR* CmdLine);
+#define UE_API IOSTOREUTILITIES_API
 
-IOSTOREUTILITIES_API bool DumpIoStoreContainerInfo(const TCHAR* InContainerFilename, const FKeyChain& InKeyChain);
+UE_API int32 CreateIoStoreContainerFiles(const TCHAR* CmdLine);
 
-IOSTOREUTILITIES_API bool LegacyListIoStoreContainer(
+UE_API bool DumpIoStoreContainerInfo(const TCHAR* InContainerFilename, const FKeyChain& InKeyChain);
+
+UE_API bool LegacyListIoStoreContainer(
 	const TCHAR* InContainerFilename,
 	int64 InSizeFilter,
 	const FString& InCSVFilename,
 	const FKeyChain& InKeyChain);
 
-IOSTOREUTILITIES_API bool LegacyDiffIoStoreContainers(
+UE_API bool ListIoStoreContainer(const TCHAR* CmdLine);
+
+UE_API bool LegacyDiffIoStoreContainers(
 	const TCHAR* InContainerFilename1,
 	const TCHAR* InContainerFilename2,
 	bool bInLogUniques1,
@@ -31,7 +35,7 @@ IOSTOREUTILITIES_API bool LegacyDiffIoStoreContainers(
 	const FKeyChain& InKeyChain1,
 	const FKeyChain* InKeyChain2 = nullptr);
 
-IOSTOREUTILITIES_API bool ExtractFilesFromIoStoreContainer(
+UE_API bool ExtractFilesFromIoStoreContainer(
 	const TCHAR* InContainerFilename,
 	const TCHAR* InDestPath,
 	const FKeyChain& InKeyChain,
@@ -40,7 +44,7 @@ IOSTOREUTILITIES_API bool ExtractFilesFromIoStoreContainer(
 	TArray<FGuid>* OutUsedEncryptionKeys,
 	bool* bOutIsSigned);
 
-IOSTOREUTILITIES_API bool ProcessFilesFromIoStoreContainer(
+UE_API bool ProcessFilesFromIoStoreContainer(
 	const TCHAR* InContainerFilename,
 	const TCHAR* InDestPath,
 	const FKeyChain& InKeyChain,
@@ -51,17 +55,19 @@ IOSTOREUTILITIES_API bool ProcessFilesFromIoStoreContainer(
 	bool* bOutIsSigned,
 	int32 MaxConcurrentReaders);
 
-IOSTOREUTILITIES_API bool SignIoStoreContainer(const TCHAR* InContainerFilename, const FRSAKeyHandle InSigningKey);
+UE_API bool SignIoStoreContainer(const TCHAR* InContainerFilename, const FRSAKeyHandle InSigningKey);
 
-IOSTOREUTILITIES_API bool UploadIoStoreContainerFiles(const TCHAR* ContainerPathOrWildcard); 
+UE_API bool UploadIoStoreContainerFiles(const TCHAR* ContainerPathOrWildcard); 
 
-IOSTOREUTILITIES_API bool DownloadIoStoreContainerFiles(const TCHAR* TocPath);
+UE_API bool DownloadIoStoreContainerFiles(const TCHAR* TocPath);
 
-IOSTOREUTILITIES_API bool ListOnDemandTocs();
+UE_API bool ListOnDemandTocs();
 
 /**
  * Attempts to connect to the endpoint in the provided ini file, download all
  * available tocs and then download all chunks in those tocs. This will
  * pre-warm the CDN and allow testing closer to the end users experience.
  */
-IOSTOREUTILITIES_API bool PrimeEndPoint(FStringView IoStoreOnDemandIniPath);
+UE_API bool PrimeEndPoint(FStringView IoStoreOnDemandIniPath);
+
+#undef UE_API
diff --git a/Engine/Source/Developer/PakFileUtilities/Private/PakFileUtilities.cpp b/Engine/Source/Developer/PakFileUtilities/Private/PakFileUtilities.cpp
--- a/Engine/Source/Developer/PakFileUtilities/Private/PakFileUtilities.cpp
+++ b/Engine/Source/Developer/PakFileUtilities/Private/PakFileUtilities.cpp
@@ -5478,6 +5478,11 @@ bool ExecuteUnrealPak(const TCHAR* CmdLine)
 			{
 				return ListOnDemandTocs();
 			}
+
+			if (FParse::Value(CmdLine, TEXT("-ListContainer="), IoStoreArg))
+			{
+				return ListIoStoreContainer(CmdLine);
+			}
 		}
 	}
 
