UnrealPak - dump bulk meta data to Json or Csv

Writes a package bulk data instances with offset, size and flags to the specified destination file.

Usage:
UnrealPak.exe -ListContainerBulkData=<Path/To/PakFolder/OrContainerFile> -cryptokeys=<Path> -Out=<Filename[.json|.csv]>

diff --git a/Engine/Source/Developer/IoStoreUtilities/Private/IoStoreUtilities.cpp b/Engine/Source/Developer/IoStoreUtilities/Private/IoStoreUtilities.cpp
--- a/Engine/Source/Developer/IoStoreUtilities/Private/IoStoreUtilities.cpp
+++ b/Engine/Source/Developer/IoStoreUtilities/Private/IoStoreUtilities.cpp
@@ -26,8 +26,10 @@
 #include "Misc/WildcardString.h"
 #include "Modules/ModuleManager.h"
 #include "Serialization/Archive.h"
+#include "Serialization/BulkData.h"
 #include "Serialization/JsonReader.h"
 #include "Serialization/JsonSerializer.h"
+#include "Serialization/JsonWriter.h"
 #include "Serialization/BufferWriter.h"
 #include "Serialization/LargeMemoryWriter.h"
 #include "Serialization/MemoryReader.h"
@@ -75,6 +77,7 @@
 #include "Misc/PathViews.h"
 #include "HAL/FileManagerGeneric.h"
 #include "Serialization/CompactBinarySerialization.h"
+#include "Serialization/ZenPackageHeader.h"
 
 IMPLEMENT_MODULE(FDefaultModuleImpl, IoStoreUtilities);
 
@@ -5981,6 +5984,237 @@ bool ListIoStoreContainer(const TCHAR* CmdLine)
 	return ListContainer(KeyChain, ContainerPathOrWildcard, CsvPath) == 0;
 }
 
+bool ListContainerBulkData(
+	const FKeyChain& KeyChain,
+	const FString& ContainerPathOrWildcard,
+	const FString& OutFile)
+{
+	struct FPackageData
+	{
+		FPackageId Id;
+		FString Filename;
+		TArray<FBulkDataMapEntry> BulkDataMap;
+	};
+
+	struct FContainerData
+	{
+		FString Name;
+		FString Path;
+		TArray<FPackageData> Packages;
+	};
+
+	TArray<FContainerData> Containers;
+	TArray<FString> ContainerFilePaths;
+
+	if (IFileManager::Get().FileExists(*ContainerPathOrWildcard))
+	{
+		ContainerFilePaths.Add(ContainerPathOrWildcard);
+	}
+	else if (IFileManager::Get().DirectoryExists(*ContainerPathOrWildcard))
+	{
+		FString Directory = ContainerPathOrWildcard;
+		FPaths::NormalizeDirectoryName(Directory);
+
+		TArray<FString> FoundContainerFiles;
+		IFileManager::Get().FindFiles(FoundContainerFiles, *(Directory / TEXT("*.utoc")), true, false);
+
+		for (const FString& Filename : FoundContainerFiles)
+		{
+			ContainerFilePaths.Emplace(Directory / Filename);
+		}
+	}
+	else
+	{
+		FString Directory = FPaths::GetPath(ContainerPathOrWildcard);
+		FPaths::NormalizeDirectoryName(Directory);
+
+		TArray<FString> FoundContainerFiles;
+		IFileManager::Get().FindFiles(FoundContainerFiles, *ContainerPathOrWildcard, true, false);
+
+		for (const FString& Filename : FoundContainerFiles)
+		{
+			ContainerFilePaths.Emplace(Directory / Filename);
+		}
+	}
+
+	if (ContainerFilePaths.Num() == 0)
+	{
+		UE_LOG(LogIoStore, Error, TEXT("Container '%s' doesn't exist and no container matches wildcard."), *ContainerPathOrWildcard);
+		return false;
+	}
+
+	for (const FString& ContainerFilePath : ContainerFilePaths)
+	{
+		FString ContainerName = FPaths::GetBaseFilename(ContainerFilePath);
+		if (ContainerName == TEXT("global"))
+		{
+			continue;
+		}
+
+		TUniquePtr<FIoStoreReader> Reader = CreateIoStoreReader(*ContainerFilePath, KeyChain);
+		if (!Reader.IsValid())
+		{
+			UE_LOG(LogIoStore, Error, TEXT("Failed to read container '%s'"), *ContainerFilePath);
+			continue;
+		}
+
+		TMap<FIoChunkId, FString> FilenameByChunkId;
+		Reader->GetDirectoryIndexReader().IterateDirectoryIndex(FIoDirectoryIndexHandle::RootDirectory(), TEXT(""),
+			[&FilenameByChunkId, &Reader](FStringView Filename, uint32 TocEntryIndex) -> bool
+			{
+				TIoStatusOr<FIoStoreTocChunkInfo> ChunkInfo = Reader->GetChunkInfo(TocEntryIndex);
+				if (ChunkInfo.IsOk())
+				{
+					FilenameByChunkId.Add(ChunkInfo.ValueOrDie().Id, FString(Filename));
+				}
+				return true;
+			});
+
+		UE_LOG(LogIoStore, Display, TEXT("Listing bulk data in container '%s'"), *ContainerFilePath);
+		FIoChunkId ChunkId = CreateIoChunkId(Reader->GetContainerId().Value(), 0, EIoChunkType::ContainerHeader);
+		TIoStatusOr<FIoBuffer> Status = Reader->Read(ChunkId, FIoReadOptions());
+
+		if (!Status.IsOk())
+		{
+			UE_LOG(LogIoStore, Display, TEXT("Failed to read container header '%s', reason '%s'"),
+				*ContainerFilePath, *Status.Status().ToString());
+			continue;
+		}
+
+		FIoContainerHeader ContainerHeader;
+		{
+			FIoBuffer Chunk = Status.ValueOrDie();
+			FMemoryReaderView Ar(MakeArrayView(Chunk.Data(), Chunk.GetSize()));
+			Ar << ContainerHeader;
+		}
+
+		FContainerData& Container = Containers.AddDefaulted_GetRef();
+		Container.Path = ContainerFilePath;
+		Container.Name = MoveTemp(ContainerName);
+
+		for (const FPackageId& PackageId : ContainerHeader.PackageIds)
+		{
+			ChunkId = CreatePackageDataChunkId(PackageId);
+			Status = Reader->Read(ChunkId, FIoReadOptions());
+			if (!Status.IsOk())
+			{
+				UE_LOG(LogIoStore, Display, TEXT("Failed to package data"));
+				continue;
+			}
+
+			FIoBuffer Chunk = Status.ValueOrDie();
+			FZenPackageHeader PkgHeader = FZenPackageHeader::MakeView(Chunk.GetView());
+			FPackageData& Pkg = Container.Packages.AddDefaulted_GetRef();
+			Pkg.Id = PackageId;
+			Pkg.BulkDataMap = PkgHeader.BulkDataMap;
+			if (FString* Filename = FilenameByChunkId.Find(ChunkId))
+			{
+				Pkg.Filename = *Filename;
+			}
+		}
+	}
+
+	const FString Ext = FPaths::GetExtension(OutFile);
+	if (Ext == TEXT("json"))
+	{
+		using FWriter = TSharedPtr<TJsonWriter<TCHAR, TPrettyJsonPrintPolicy<TCHAR>>>; 
+		using FWriterFactory = TJsonWriterFactory<TCHAR, TPrettyJsonPrintPolicy<TCHAR>>;
+
+		FString Json;
+		FWriter Writer = FWriterFactory::Create(&Json);
+		Writer->WriteArrayStart();
+		
+		TStringBuilder<512> Sb;
+		for (const FContainerData& Container: Containers)
+		{
+			Writer->WriteObjectStart();
+			Writer->WriteValue(TEXT("Container"), Container.Name);
+			Writer->WriteArrayStart(TEXT("Packages"));
+			for (const FPackageData& Pkg : Container.Packages)
+			{
+				if (Pkg.BulkDataMap.IsEmpty())
+				{
+					continue;
+				}
+
+				Writer->WriteObjectStart();
+				Writer->WriteValue(TEXT("PackageId"), FString::Printf(TEXT("0x%llX"), Pkg.Id.Value()));
+				Writer->WriteValue(TEXT("Filename"), Pkg.Filename);
+				Writer->WriteArrayStart(TEXT("BulkData"));
+				for (const FBulkDataMapEntry& Entry : Pkg.BulkDataMap)
+				{
+					Sb.Reset();
+					LexToString(static_cast<EBulkDataFlags>(Entry.Flags), Sb);
+					Writer->WriteObjectStart();
+					Writer->WriteValue(TEXT("Offset"), Entry.SerialOffset);
+					Writer->WriteValue(TEXT("Size"), Entry.SerialSize);
+					Writer->WriteValue(TEXT("Flags"), Sb.ToString());
+					Writer->WriteObjectEnd();
+				}
+				Writer->WriteArrayEnd();
+				Writer->WriteObjectEnd();
+			}
+			Writer->WriteArrayEnd();
+			Writer->WriteObjectEnd();
+		}
+
+		Writer->WriteArrayEnd();
+		Writer->Close();
+
+		UE_LOG(LogIoStore, Display, TEXT("Saving '%s'"), *OutFile);
+		if (!FFileHelper::SaveStringToFile(Json, *OutFile))
+		{
+			return false;
+		}
+	}
+	else
+	{
+		TUniquePtr<FArchive> CsvAr(IFileManager::Get().CreateFileWriter(*OutFile));
+		CsvAr->Logf(TEXT("Container,Filename,PackageId,Offset,Size,Flags"));
+
+		TStringBuilder<512> Sb;
+		for (const FContainerData& Container: Containers)
+		{
+			for (const FPackageData& Pkg : Container.Packages)
+			{
+				for (const FBulkDataMapEntry& Entry : Pkg.BulkDataMap)
+				{
+					Sb.Reset();
+					LexToString(static_cast<EBulkDataFlags>(Entry.Flags), Sb);
+					CsvAr->Logf(TEXT("%s,%s,0x%llX,%lld,%lld,%s"),
+						*Container.Name, *Pkg.Filename, Pkg.Id.Value(), Entry.SerialOffset, Entry.SerialSize, Sb.ToString());
+				}
+			}
+		}
+
+		UE_LOG(LogIoStore, Display, TEXT("Saving '%s'"), *OutFile);
+	}
+
+	return true;
+}
+
+bool ListIoStoreContainerBulkData(const TCHAR* CmdLine)
+{
+	FKeyChain KeyChain;
+	LoadKeyChain(CmdLine, KeyChain);
+
+	FString ContainerPathOrWildcard;
+	if (!FParse::Value(FCommandLine::Get(), TEXT("ListContainerBulkData="), ContainerPathOrWildcard))
+	{
+		UE_LOG(LogIoStore, Error, TEXT("Missing argument -ListContainerBulkData=<ContainerFileOrWildCard>"));
+		return false;
+	}
+
+	FString OutFile;
+	if (!FParse::Value(FCommandLine::Get(), TEXT("Out="), OutFile))
+	{
+		UE_LOG(LogIoStore, Error, TEXT("Missing argument -Out=<Path.[json|csv]>"));
+		return false;
+	}
+
+	return ListContainerBulkData(KeyChain, ContainerPathOrWildcard, OutFile) == 0;
+}
+
 bool LegacyListIoStoreContainer(
 	const TCHAR* InContainerFilename,
 	int64 InSizeFilter,
diff --git a/Engine/Source/Developer/IoStoreUtilities/Public/IoStoreUtilities.h b/Engine/Source/Developer/IoStoreUtilities/Public/IoStoreUtilities.h
--- a/Engine/Source/Developer/IoStoreUtilities/Public/IoStoreUtilities.h
+++ b/Engine/Source/Developer/IoStoreUtilities/Public/IoStoreUtilities.h
@@ -26,6 +26,7 @@ UE_API bool LegacyListIoStoreContainer(
 	const FKeyChain& InKeyChain);
 
 UE_API bool ListIoStoreContainer(const TCHAR* CmdLine);
+UE_API bool ListIoStoreContainerBulkData(const TCHAR* CmdLine);
 
 UE_API bool LegacyDiffIoStoreContainers(
 	const TCHAR* InContainerFilename1,
diff --git a/Engine/Source/Developer/PakFileUtilities/Private/PakFileUtilities.cpp b/Engine/Source/Developer/PakFileUtilities/Private/PakFileUtilities.cpp
--- a/Engine/Source/Developer/PakFileUtilities/Private/PakFileUtilities.cpp
+++ b/Engine/Source/Developer/PakFileUtilities/Private/PakFileUtilities.cpp
@@ -5478,6 +5478,11 @@ bool ExecuteUnrealPak(const TCHAR* CmdLine)
 			{
 				return ListIoStoreContainer(CmdLine);
 			}
+
+			if (FParse::Value(CmdLine, TEXT("-ListContainerBulkData="), IoStoreArg))
+			{
+				return ListIoStoreContainerBulkData(CmdLine);
+			}
 		}
 	}
 
diff --git a/Engine/Source/Runtime/CoreUObject/Private/Serialization/BulkData.cpp b/Engine/Source/Runtime/CoreUObject/Private/Serialization/BulkData.cpp
--- a/Engine/Source/Runtime/CoreUObject/Private/Serialization/BulkData.cpp
+++ b/Engine/Source/Runtime/CoreUObject/Private/Serialization/BulkData.cpp
@@ -30,6 +30,61 @@
 
 //////////////////////////////////////////////////////////////////////////////
 
+FStringBuilderBase& LexToString(EBulkDataFlags Flags, FStringBuilderBase& Sb)
+{
+	#define TEST_AND_ADD_FLAG(Sb, Flags, Contains)\
+	{\
+		if ((uint32(Flags) & uint32(Contains)) == uint32(Contains))\
+		{\
+			if (Sb.Len())\
+			{\
+				Sb.Append(TEXT("|"));\
+			}\
+			Sb.Append(TEXT(#Contains));\
+		}\
+	}
+
+	if (uint32(Flags) == BULKDATA_None)
+	{
+		Sb.Append("BULKDATA_None");
+		return Sb;
+	}
+
+	TEST_AND_ADD_FLAG(Sb, Flags, BULKDATA_PayloadAtEndOfFile);
+	TEST_AND_ADD_FLAG(Sb, Flags, BULKDATA_SerializeCompressedZLIB);
+	TEST_AND_ADD_FLAG(Sb, Flags, BULKDATA_ForceSingleElementSerialization);
+	TEST_AND_ADD_FLAG(Sb, Flags, BULKDATA_SingleUse);
+	TEST_AND_ADD_FLAG(Sb, Flags, BULKDATA_ForceInlinePayload);
+	TEST_AND_ADD_FLAG(Sb, Flags, BULKDATA_SerializeCompressed);
+	TEST_AND_ADD_FLAG(Sb, Flags, BULKDATA_PayloadInSeperateFile);
+	TEST_AND_ADD_FLAG(Sb, Flags, BULKDATA_Force_NOT_InlinePayload);
+	TEST_AND_ADD_FLAG(Sb, Flags, BULKDATA_OptionalPayload);
+	TEST_AND_ADD_FLAG(Sb, Flags, BULKDATA_MemoryMappedPayload);
+	TEST_AND_ADD_FLAG(Sb, Flags, BULKDATA_Size64Bit);
+	TEST_AND_ADD_FLAG(Sb, Flags, BULKDATA_DuplicateNonOptionalPayload);
+	TEST_AND_ADD_FLAG(Sb, Flags, BULKDATA_NoOffsetFixUp);
+	TEST_AND_ADD_FLAG(Sb, Flags, BULKDATA_WorkspaceDomainPayload);
+	TEST_AND_ADD_FLAG(Sb, Flags, BULKDATA_LazyLoadable);
+	TEST_AND_ADD_FLAG(Sb, Flags, BULKDATA_UsesIoDispatcher);
+	TEST_AND_ADD_FLAG(Sb, Flags, BULKDATA_DataIsMemoryMapped);
+	TEST_AND_ADD_FLAG(Sb, Flags, BULKDATA_HasAsyncReadPending);
+	TEST_AND_ADD_FLAG(Sb, Flags, BULKDATA_AlwaysAllowDiscard);
+	PRAGMA_DISABLE_DEPRECATION_WARNINGS
+	TEST_AND_ADD_FLAG(Sb, Flags, BULKDATA_BadDataVersion);
+	TEST_AND_ADD_FLAG(Sb, Flags, BULKDATA_ForceStreamPayload);
+	TEST_AND_ADD_FLAG(Sb, Flags, BULKDATA_SerializeCompressedBitWindow);
+	TEST_AND_ADD_FLAG(Sb, Flags, BULKDATA_Unused);
+	PRAGMA_ENABLE_DEPRECATION_WARNINGS
+
+	return Sb;
+}
+
+FString LexToString(EBulkDataFlags Flags)
+{
+	TStringBuilder<256> Sb;
+	return LexToString(Flags, Sb).ToString();
+}
+
 const FIoFilenameHash FALLBACK_IO_FILENAME_HASH = INVALID_IO_FILENAME_HASH - 1;
 
 FIoFilenameHash MakeIoFilenameHash(const FPackagePath& PackagePath)
diff --git a/Engine/Source/Runtime/CoreUObject/Public/Serialization/BulkData.h b/Engine/Source/Runtime/CoreUObject/Public/Serialization/BulkData.h
--- a/Engine/Source/Runtime/CoreUObject/Public/Serialization/BulkData.h
+++ b/Engine/Source/Runtime/CoreUObject/Public/Serialization/BulkData.h
@@ -9,6 +9,7 @@
 #include "Containers/Map.h"
 #include "Containers/SortedMap.h"
 #include "Containers/StringView.h"
+#include "Containers/StringFwd.h"
 #include "Containers/UnrealString.h"
 #include "CoreMinimal.h"
 #include "CoreTypes.h"
@@ -173,6 +174,9 @@ enum EBulkDataFlags : uint32
 	BULKDATA_AlwaysAllowDiscard = 1 << 28,
 };
 
+COREUOBJECT_API FStringBuilderBase& LexToString(EBulkDataFlags Flags, FStringBuilderBase& Sb);
+COREUOBJECT_API FString LexToString(EBulkDataFlags Flags);
+
 /**
  * Allows FArchive to serialize EBulkDataFlags, this will not be required once EBulkDataFlags is promoted
  * to be a enum class.
