Make use of ability to get all chunk data in batch from zenserver 0.2.36 release.

diff --git a/Engine/Source/Developer/IoStoreUtilities/Internal/ZenStoreHttpClient.h b/Engine/Source/Developer/IoStoreUtilities/Internal/ZenStoreHttpClient.h
--- a/Engine/Source/Developer/IoStoreUtilities/Internal/ZenStoreHttpClient.h
+++ b/Engine/Source/Developer/IoStoreUtilities/Internal/ZenStoreHttpClient.h
@@ -61,6 +61,7 @@ public:
 
 	TFuture<TIoStatusOr<FCbObject>> GetOplog();
 	TFuture<TIoStatusOr<FCbObject>> GetFiles();
+	TFuture<TIoStatusOr<FCbObject>> GetChunkInfos();
 
 	static const UTF8CHAR* FindOrAddAttachmentId(FUtf8StringView AttachmentText);
 	static const UTF8CHAR* FindAttachmentId(FUtf8StringView AttachmentText);
diff --git a/Engine/Source/Developer/IoStoreUtilities/Private/IoStoreUtilities.cpp b/Engine/Source/Developer/IoStoreUtilities/Private/IoStoreUtilities.cpp
--- a/Engine/Source/Developer/IoStoreUtilities/Private/IoStoreUtilities.cpp
+++ b/Engine/Source/Developer/IoStoreUtilities/Private/IoStoreUtilities.cpp
@@ -835,6 +835,11 @@ public:
 		return ZenStoreClient.IsValid();
 	}
 
+	TFuture<TIoStatusOr<FCbObject>> GetChunkInfos()
+	{
+		return ZenStoreClient->GetChunkInfos();
+	}
+
 	TIoStatusOr<uint64> GetChunkSize(const FIoChunkId& ChunkId)
 	{
 		return ZenStoreClient->GetChunkSize(ChunkId);
@@ -2550,7 +2555,43 @@ void InitializeContainerTargetsAndPackages(
 		return true;
 	};
 
+	struct FChunkListItem
+	{
+		FIoHash RawHash;
+		uint64 RawSize;
+	};
+
+	TMap<FIoChunkId, FChunkListItem> ChunkList;
+	if (Arguments.PackageStore->HasZenStoreClient())
+	{
+		double StartChunkInfoTime = FPlatformTime::Seconds();
+
+		TIoStatusOr<FCbObject> Chunks = Arguments.PackageStore->GetChunkInfos().Get();
+		if (!Chunks.IsOk())
+		{
+			UE_LOG(LogIoStore, Error, TEXT("Failed to retrieve chunk list"));
+			return;
+		}
+
+		FCbObject ChunksObj = Chunks.ConsumeValueOrDie();
+		for (FCbField& ChunkEntry : ChunksObj["chunkinfos"])
+		{
+			FCbObject ChunkObj = ChunkEntry.AsObject();
+			FIoChunkId ChunkId;
+			if (!LoadFromCompactBinary(ChunkObj["id"], ChunkId))
+			{
+				UE_LOG(LogIoStore, Warning, TEXT("Received invalid chunk id, skipping."));
+				continue;
+			}
+			ChunkList.Add(ChunkId, { ChunkObj["rawhash"].AsHash(), ChunkObj["rawsize"].AsUInt64() });
+		}
+
+		UE_LOG(LogIoStore, Display, TEXT("Fetched '%d' chunk infos in %f seconds"), ChunkList.Num(), FPlatformTime::Seconds() - StartChunkInfoTime);
+
+	}
+
 	auto CreateTargetFileFromZen = [
+		&ChunkList,
 		&Arguments,
 		&Packages,
 		&PackageNameMap,
@@ -2581,13 +2622,13 @@ void InitializeContainerTargetsAndPackages(
 			return false;
 		}
 		OutTargetFile.ChunkId = ChunkInfo->ChunkId;
-		TIoStatusOr<uint64> ChunkSize = PackageStore.GetChunkSize(OutTargetFile.ChunkId);
-		if (!ChunkSize.IsOk())
+		FChunkListItem* ChunkListItem = ChunkList.Find(OutTargetFile.ChunkId);
+		if (!ChunkListItem)
 		{
 			UE_LOG(LogIoStore, Warning, TEXT("Chunk size not found for: '%s'"), *SourceFile.NormalizedPath);
 			return false;
 		}
-		OutTargetFile.SourceSize = ChunkSize.ValueOrDie();
+		OutTargetFile.SourceSize = ChunkListItem->RawSize;
 
 		if (ChunkInfo->PackageName.IsNone())
 		{
diff --git a/Engine/Source/Developer/IoStoreUtilities/Private/ZenStoreHttpClient.cpp b/Engine/Source/Developer/IoStoreUtilities/Private/ZenStoreHttpClient.cpp
--- a/Engine/Source/Developer/IoStoreUtilities/Private/ZenStoreHttpClient.cpp
+++ b/Engine/Source/Developer/IoStoreUtilities/Private/ZenStoreHttpClient.cpp
@@ -582,6 +582,35 @@ TFuture<TIoStatusOr<FCbObject>> FZenStoreHttpClient::GetFiles()
 	});
 }
 
+TFuture<TIoStatusOr<FCbObject>> FZenStoreHttpClient::GetChunkInfos()
+{
+#if WITH_EDITOR
+	EAsyncExecution ThreadPool = EAsyncExecution::LargeThreadPool;
+#else
+	EAsyncExecution ThreadPool = EAsyncExecution::ThreadPool;
+#endif
+	return Async(ThreadPool, [this]
+	{
+		UE::Zen::FZenScopedRequestPtr Request(RequestPool.Get());
+		
+		TStringBuilder<128> Uri;
+		Uri << OplogPath << "/chunkinfos";
+
+		TArray64<uint8> GetBuffer;
+		UE::Zen::FZenHttpRequest::Result Res = Request->PerformBlockingDownload(Uri, &GetBuffer, Zen::EContentType::CbObject);
+
+		if (Res == Zen::FZenHttpRequest::Result::Success && Request->GetResponseCode() == 200)
+		{
+			FCbObjectView Response(GetBuffer.GetData());
+			return TIoStatusOr<FCbObject>(FCbObject::Clone(Response));
+		}
+		else
+		{
+			return TIoStatusOr<FCbObject>(FIoStatus(EIoErrorCode::NotFound));
+		}
+	});
+}
+
 void 
 FZenStoreHttpClient::StartBuildPass()
 {
@@ -703,6 +732,11 @@ TFuture<TIoStatusOr<FCbObject>> FZenStoreHttpClient::GetFiles()
 	return TFuture<TIoStatusOr<FCbObject>>();
 }
 
+TFuture<TIoStatusOr<FCbObject>> FZenStoreHttpClient::GetChunkInfos()
+{
+	return TFuture<TIoStatusOr<FCbObject>>();
+}
+
 }
 
 #endif // desktop platform
