Improve error handling in FZenStoreHttpClient::AppendOp

diff --git a/Engine/Source/Developer/IoStoreUtilities/Private/ZenStoreHttpClient.cpp b/Engine/Source/Developer/IoStoreUtilities/Private/ZenStoreHttpClient.cpp
--- a/Engine/Source/Developer/IoStoreUtilities/Private/ZenStoreHttpClient.cpp
+++ b/Engine/Source/Developer/IoStoreUtilities/Private/ZenStoreHttpClient.cpp
@@ -284,6 +284,7 @@ TIoStatusOr<uint64> FZenStoreHttpClient::AppendOp(FCbPackage OpEntry)
 
 			TSet<FIoHash> NeedChunks;
 
+			if (!Attachments.IsEmpty())
 			{
 				FCbWriter Writer;
 				Writer.BeginObject();
@@ -319,6 +320,15 @@ TIoStatusOr<uint64> FZenStoreHttpClient::AppendOp(FCbPackage OpEntry)
 						IsOk = true;
 					}
 				}
+
+				if (!IsOk)
+				{
+					UE_LOG(LogZenStore, Warning, TEXT("Unable to check for needed chunks, assuming all are needed"));
+					for (const FCbAttachment& Attachment : Attachments)
+					{
+						NeedChunks.Add(Attachment.GetHash());
+					}
+				}
 			}
 
 			// This uses a slight variation for package attachment serialization
@@ -363,18 +373,24 @@ TIoStatusOr<uint64> FZenStoreHttpClient::AppendOp(FCbPackage OpEntry)
 
 						if (IFileHandle* FileHandle = PlatformFile.OpenWrite(*TempFilePath))
 						{
-							FileHandle->Write((const uint8*)AttachmentData.GetData(), AttachmentData.GetSize());
+							bool WriteOK = FileHandle->Write((const uint8*)AttachmentData.GetData(), AttachmentData.GetSize());
 							delete FileHandle;
-
-							Writer.AddHash(AttachmentHash);
-							bIsSerialized = true;
-							bIsUsingTempFiles = true;
+							if (WriteOK)
+							{
+								Writer.AddHash(AttachmentHash);
+								bIsSerialized = true;
+								bIsUsingTempFiles = true;
+							}
+							else
+							{
+								(void)PlatformFile.DeleteFile(*TempFilePath);
+								// Take the slow path if we can't write the payload file in the large attachment directory
+								UE_LOG(LogZenStore, Warning, TEXT("Could not write to file '%s', taking slow path for large attachment"), *TempFilePath);
+							}
 						}
 						else
 						{
-							// Take the slow path if we can't open the payload file in the
-							// large attachment directory
-
+							// Take the slow path if we can't open the payload file in the large attachment directory
 							UE_LOG(LogZenStore, Warning, TEXT("Could not create file '%s', taking slow path for large attachment"), *TempFilePath);
 						}
 					}
