Reference Chunk logging improvement.

diff --git a/Engine/Source/Developer/IoStoreUtilities/Private/IoStoreUtilities.cpp b/Engine/Source/Developer/IoStoreUtilities/Private/IoStoreUtilities.cpp
--- a/Engine/Source/Developer/IoStoreUtilities/Private/IoStoreUtilities.cpp
+++ b/Engine/Source/Developer/IoStoreUtilities/Private/IoStoreUtilities.cpp
@@ -105,19 +105,40 @@ static const uint64 DefaultMemoryMappingAlignment = 16 << 10;
 static TUniquePtr<FIoStoreReader> CreateIoStoreReader(const TCHAR* Path, const FKeyChain& KeyChain);
 bool UploadIoStoreContainerFiles(const UE::IO::IAS::FIoStoreUploadParams& UploadParams, TConstArrayView<FString> ContainerFiles, const FKeyChain& KeyChain);
 
+/*
+* Provides access to previously compressed chunks to the iostore writer, allowing
+* a) avoiding recompressing things and b) tweaks to compressors dont cause massive patches.
+*/
 class FIoStoreChunkDatabase : public IIoStoreWriterReferenceChunkDatabase
 {
 public:
+	static const uint8 IoChunkTypeCount = (uint8)EIoChunkType::MAX;
 
 	TArray<TUniquePtr<FIoStoreReader>> Readers;
 	struct FReaderChunks
 	{
 		int32 ReaderIndex;
-		TMap<FIoChunkHash, FIoChunkId> Chunks;
+		TMap<FIoChunkHash, FIoStoreTocChunkInfo> Chunks;
+		TSet<FIoChunkId> ChunkIds;
+
+		std::atomic_uint64_t ChangedChunkCount[IoChunkTypeCount];
+		std::atomic_uint64_t NewChunkCount[IoChunkTypeCount];
+		std::atomic_uint64_t UsedChunkCount[IoChunkTypeCount];
+	};
+
+	struct FMissingContainerInfo
+	{
+		std::atomic_uint64_t RequestedChunkCount[IoChunkTypeCount];
 	};
 
-	TMap<FIoContainerId, FReaderChunks> ChunkDatabase;
-	int32 RequestCount = 0;
+	TMap<FIoContainerId, TUniquePtr<FMissingContainerInfo>> MissingContainerIds;
+
+	TMap<FIoContainerId, FString> ContainerNameMap;
+
+	TMap<FIoContainerId, TUniquePtr<FReaderChunks>> ChunkDatabase;
+
+
+	std::atomic_int32_t RequestCount = 0;
 	int32 FulfillCount = 0;
 	int32 ContainerNotFound = 0;
 	int64 FulfillBytes = 0;
@@ -147,7 +168,7 @@ public:
 		}
 
 		CompressionBlockSize = 0;
-		int64 IoChunkCount = 0;		
+		int64 IoChunkCount = 0;
 		for (const FString& ContainerFilePath : ContainerFilePaths)
 		{
 			TUniquePtr<FIoStoreReader> Reader = CreateIoStoreReader(*ContainerFilePath, InDecryptionKeychain);
@@ -156,11 +177,22 @@ public:
 				UE_LOG(LogIoStore, Error, TEXT("Failed to open reference chunk container %s"), *ContainerFilePath);
 				return false;
 			}
-			FReaderChunks& ReaderChunks = ChunkDatabase.FindOrAdd(Reader->GetContainerId());
 
-			Reader->EnumerateChunks([&ReaderChunks](const FIoStoreTocChunkInfo& ChunkInfo)
+			if (ChunkDatabase.Contains(Reader->GetContainerId()))
+			{
+				UE_LOG(LogIoStore, Error, TEXT("Duplicate container id found in reference chunk container directory %s"), *ContainerFilePath);
+				UE_LOG(LogIoStore, Error, TEXT("is duplicate of %s"), *ContainerNameMap[Reader->GetContainerId()]);
+				return false;
+			}
+
+			TUniquePtr<FReaderChunks> ReaderChunks(new FReaderChunks());
+
+			ContainerNameMap.Add(Reader->GetContainerId(), *ContainerFilePath);
+
+			Reader->EnumerateChunks([ReaderChunks = ReaderChunks.Get()](FIoStoreTocChunkInfo&& ChunkInfo)
 			{
-				ReaderChunks.Chunks.Add(TPair<FIoChunkHash, FIoChunkId>(ChunkInfo.Hash, ChunkInfo.Id));
+				ReaderChunks->ChunkIds.Add(ChunkInfo.Id);
+				ReaderChunks->Chunks.Add(TPair<FIoChunkHash, FIoStoreTocChunkInfo>(ChunkInfo.Hash, MoveTemp(ChunkInfo)));
 				return true;
 			});
 
@@ -174,39 +206,76 @@ public:
 				return false;
 			}
 
-			IoChunkCount += ReaderChunks.Chunks.Num();
-			ReaderChunks.ReaderIndex = Readers.Num();
+			IoChunkCount += ReaderChunks->Chunks.Num();
+			ReaderChunks->ReaderIndex = Readers.Num();
+			ChunkDatabase.Add(Reader->GetContainerId(), MoveTemp(ReaderChunks));
 			Readers.Add(MoveTemp(Reader));
+
+			
 		}
 
 		UE_LOG(LogIoStore, Display, TEXT("Block reference loaded %d containers and %s chunks, in %.1f seconds"), Readers.Num(), *FText::AsNumber(IoChunkCount).ToString(), FPlatformTime::Seconds() - StartTime);
 		return true;
 	}
 
+	virtual void NotifyAddedToWriter(const FIoContainerId& InContainerId) override
+	{
+		// If we don't have this container in our chunk database, add it to a tracking
+		// structure so we can see how many chunks we're missing.
+		if (ChunkDatabase.Contains(InContainerId) == false)
+		{
+			MissingContainerIds.Add(InContainerId, MakeUnique<FMissingContainerInfo>());
+		}
+	}
+
 	virtual uint32 GetCompressionBlockSize() const override
 	{
 		return CompressionBlockSize;
 	}
 
-	virtual bool ChunkExists(const TPair<FIoContainerId, FIoChunkHash>& InChunkKey, uint32& OutNumChunkBlocks)
+	// Returns whether we expect to be able to load the chunk from the reference chunk database.
+	// This can be called from any thread, though in the presence of existing hashes it's single threaded.
+	virtual bool ChunkExists(const TPair<FIoContainerId, FIoChunkHash>& InChunkKey, const FIoChunkId& InChunkId, uint32& OutNumChunkBlocks)
 	{
-		FReaderChunks* ReaderChunks = ChunkDatabase.Find(InChunkKey.Key);
+		RequestCount.fetch_add(1, std::memory_order_relaxed);
+
+		TUniquePtr<FReaderChunks>* ReaderChunks = ChunkDatabase.Find(InChunkKey.Key);
 		if (ReaderChunks == nullptr)
 		{
-			// Container doesn't exist - likely provided the path to a different project. Mark this
-			// error as happening once so we can log at the end.
+			// Container doesn't exist - likely provided the path to a different project.
+			TUniquePtr<FMissingContainerInfo>* MissingContainerHitCount = MissingContainerIds.Find(InChunkKey.Key);
+
+			if (MissingContainerHitCount == nullptr)
+			{
+				UE_LOG(LogIoStore, Warning, TEXT("We got a container id that was never added! id = %llu"), InChunkKey.Key.Value());
+			}
+			else
+			{
+				MissingContainerHitCount[0]->RequestedChunkCount[(uint8)InChunkId.GetChunkType()].fetch_add(1, std::memory_order_relaxed);
+			}
+
 			ContainerNotFound++;
 			return false;
 		}
 
-		FIoChunkId* ChunkId = ReaderChunks->Chunks.Find(InChunkKey.Value);
-		if (ChunkId == nullptr)
+		FIoStoreTocChunkInfo* ChunkInfo = ReaderChunks[0]->Chunks.Find(InChunkKey.Value);
+		if (ChunkInfo == nullptr)
 		{
-			// No exact chunk data match - this is a normal exit condition for a changed block.
+			// No exact chunk data match - this is a normal exit condition for a changed block or
+			// a new block.
+			bool bChangedBlock = ReaderChunks[0]->ChunkIds.Contains(InChunkId);
+			if (bChangedBlock)
+			{
+				ReaderChunks[0]->ChangedChunkCount[(uint8)InChunkId.GetChunkType()].fetch_add(1, std::memory_order_relaxed);
+			}
+			else
+			{
+				ReaderChunks[0]->NewChunkCount[(uint8)InChunkId.GetChunkType()].fetch_add(1, std::memory_order_relaxed);
+			}
 			return false;
 		}
 
-		OutNumChunkBlocks = Readers[ReaderChunks->ReaderIndex]->GetChunkInfo(*ChunkId).ValueOrDie().NumCompressedBlocks;
+		OutNumChunkBlocks = ChunkInfo->NumCompressedBlocks;
 		return true;
 	}
 
@@ -214,28 +283,28 @@ public:
 	// Not thread safe, called from the BeginCompress dispatch thread.
 	virtual bool RetrieveChunk(const TPair<FIoContainerId, FIoChunkHash>& InChunkKey, TUniqueFunction<void(TIoStatusOr<FIoStoreCompressedReadResult>)> InCompleteCallback)
 	{
-		RequestCount++;
-
-		FReaderChunks* ReaderChunks = ChunkDatabase.Find(InChunkKey.Key);
+		TUniquePtr<FReaderChunks>* ReaderChunks = ChunkDatabase.Find(InChunkKey.Key);
 		if (ReaderChunks == nullptr)
 		{
-			// Container doesn't exist - likely provided the path to a different project. Mark this
-			// error as happening once so we can log at the end.
-			ContainerNotFound++;
+			// This should never happen now as we wrap this in a ChunkExists call.
+			UE_LOG(LogIoStore, Warning, TEXT("RetrieveChunk can't find the container - invariant violated!"));
 			return false;
 		}
 
-		FIoChunkId* ChunkId = ReaderChunks->Chunks.Find(InChunkKey.Value);
-		if (ChunkId == nullptr)
+		FIoStoreTocChunkInfo* ChunkInfo = ReaderChunks[0]->Chunks.Find(InChunkKey.Value);
+		if (ChunkInfo == nullptr)
 		{
-			// No exact chunk data match - this is a normal exit condition for a changed block.
+			// This should never happen now as we wrap this in a ChunkExists call.
+			UE_LOG(LogIoStore, Warning, TEXT("RetrieveChunk can't find the chunnk - invariant violated!"));
 			return false;
 		}
 
+		ReaderChunks[0]->UsedChunkCount[(uint8)ChunkInfo->Id.GetChunkType()].fetch_add(1, std::memory_order_relaxed);
+
 		uint64 TotalCompressedSize = 0;
 		uint64 TotalUncompressedSize = 0;
 		uint32 CompressedBlockCount = 0;
-		Readers[ReaderChunks->ReaderIndex]->EnumerateCompressedBlocksForChunk(*ChunkId, [&TotalUncompressedSize, &CompressedBlockCount, &TotalCompressedSize](const FIoStoreTocCompressedBlockInfo& BlockInfo)
+		Readers[ReaderChunks[0]->ReaderIndex]->EnumerateCompressedBlocksForChunk(ChunkInfo->Id, [&TotalUncompressedSize, &CompressedBlockCount, &TotalCompressedSize](const FIoStoreTocCompressedBlockInfo& BlockInfo)
 		{			
 			TotalCompressedSize += BlockInfo.CompressedSize;
 			TotalUncompressedSize += BlockInfo.UncompressedSize;
@@ -243,16 +312,16 @@ public:
 			return true;
 		});
 
-		FulfillBytesPerChunk[(int8)ChunkId->GetChunkType()] += TotalCompressedSize;
+		FulfillBytesPerChunk[(int8)ChunkInfo->ChunkType] += TotalCompressedSize;
 		FulfillBytes += TotalCompressedSize;
 		FulfillCount++;
 
 		//
 		// At this point we know we can use the block so we can go async.
 		//
-		FFunctionGraphTask::CreateAndDispatchWhenReady([this, ChunkId, ReaderIndex = ReaderChunks->ReaderIndex, CompleteCallback = MoveTemp(InCompleteCallback)]()
+		FFunctionGraphTask::CreateAndDispatchWhenReady([this, Id = ChunkInfo->Id, ReaderIndex = ReaderChunks[0]->ReaderIndex, CompleteCallback = MoveTemp(InCompleteCallback)]()
 		{
-			TIoStatusOr<FIoStoreCompressedReadResult> Result = Readers[ReaderIndex]->ReadCompressed(*ChunkId, FIoReadOptions());
+			TIoStatusOr<FIoStoreCompressedReadResult> Result = Readers[ReaderIndex]->ReadCompressed(Id, FIoReadOptions());
 			CompleteCallback(Result);
 		}, TStatId(), nullptr, ENamedThreads::AnyHiPriThreadNormalTask);
 
@@ -5688,22 +5757,82 @@ int32 CreateTarget(const FIoStoreArguments& Arguments, const FIoStoreWriterSetti
 	}
 	UE_LOG(LogIoStore, Display, TEXT("Output: %8.2lf MB InitialLoadData"), (double)InitialLoadSize / 1024.0 / 1024.0);
 
+
 	if (ChunkDatabase.IsValid())
 	{
-		uint64 TotalCompressedBytes = 0;
+		//
+		// If we are using a reference cache, the assumption is that we are expecting a high hit rate
+		// on the chunks - this is supposed to be a current release, so theoretically we should only
+		// miss on changed or new data, so it's useful to know where the misses are in case something
+		// major happened that maybe shouldn't have.
+		//
+		uint64 TotalEntryBytes = 0;
+		uint64 TotalMissBytes = 0;
+		TMap<FIoContainerId, FString> ContainerNameMap;
+
 		for (const FIoStoreWriterResult& Result : IoStoreWriterResults)
 		{
-			TotalCompressedBytes += Result.CompressedContainerSize;
+			TotalEntryBytes += Result.TotalEntryCompressedSize;
+			TotalMissBytes += Result.ReferenceCacheMissBytes;
+			ContainerNameMap.Add(Result.ContainerId, Result.ContainerName);
 		}
 
 		FIoStoreChunkDatabase& ChunkDatabaseRef = (FIoStoreChunkDatabase&)*ChunkDatabase;
-		UE_LOG(LogIoStore, Display, TEXT("Reference Chunk: %s reused bytes out of %s possible: %.1f%%"), *FText::AsNumber(ChunkDatabaseRef.FulfillBytes).ToString(), *FText::AsNumber(TotalCompressedBytes).ToString(), 100.0 * ChunkDatabaseRef.FulfillBytes / TotalCompressedBytes);
-		UE_LOG(LogIoStore, Display, TEXT("Reference Chunk: %s chunks found / %s requests"), *FText::AsNumber(ChunkDatabaseRef.FulfillCount).ToString(), *FText::AsNumber(ChunkDatabaseRef.RequestCount).ToString());
+
+		uint64 TotalCandidateBytes = ChunkDatabaseRef.FulfillBytes + TotalMissBytes;
+
+		UE_LOG(LogIoStore, Display, TEXT("Reference Chunk Database:"));
+		UE_LOG(LogIoStore, Display, TEXT("    %s reused bytes out of %s candidate bytes - %.1f%% hit rate."),
+			*NumberString(ChunkDatabaseRef.FulfillBytes),
+			*NumberString(TotalCandidateBytes),
+			100.0 * ChunkDatabaseRef.FulfillBytes / (TotalCandidateBytes));
+		UE_LOG(LogIoStore, Display, TEXT("    %s candidate bytes out of %s io chunk bytes - %.1f%% coverage."),
+			*NumberString(TotalCandidateBytes),
+			*NumberString(TotalEntryBytes),
+			100.0 * TotalCandidateBytes / (TotalEntryBytes));
+
+		UE_LOG(LogIoStore, Display, TEXT("    %s chunks found out of %s requests"), 
+			*NumberString(ChunkDatabaseRef.FulfillCount),
+			*NumberString(ChunkDatabaseRef.RequestCount.load()));
+		for (const TPair<FIoContainerId, TUniquePtr<FIoStoreChunkDatabase::FReaderChunks>>& ContainerInDatabase : ChunkDatabaseRef.ChunkDatabase)
+		{
+			for (uint8 i = 0; i < FIoStoreChunkDatabase::IoChunkTypeCount; i++)
+			{
+				if (ContainerInDatabase.Value->ChangedChunkCount[i] ||
+					ContainerInDatabase.Value->NewChunkCount[i] ||
+					ContainerInDatabase.Value->UsedChunkCount[i])
+				{
+					UE_LOG(LogIoStore, Display, TEXT("        %s[%s]:    %s changed, %s new, %s reused"), 
+						*ContainerNameMap[ContainerInDatabase.Key],
+						*LexToString((EIoChunkType)i),
+						*NumberString(ContainerInDatabase.Value->ChangedChunkCount[i]),
+						*NumberString(ContainerInDatabase.Value->NewChunkCount[i]),
+						*NumberString(ContainerInDatabase.Value->UsedChunkCount[i]));
+				}
+			}
+		}
+
 		if (ChunkDatabaseRef.ContainerNotFound)
 		{
-			UE_LOG(LogIoStore, Warning, TEXT("Reference Chunk had %s requests for a container that wasn't loaded. This means the "), *FText::AsNumber(ChunkDatabaseRef.ContainerNotFound).ToString());
-			UE_LOG(LogIoStore, Warning, TEXT("new output has a container that wasn't deployed before. If that doesn't sound right"));
-			UE_LOG(LogIoStore, Warning, TEXT("verify that you used reference containers from the same project."));
+			UE_LOG(LogIoStore, Display, TEXT("    %s containers were requested that weren't available. This means the "),
+				*NumberString(ChunkDatabaseRef.MissingContainerIds.Num()));
+			UE_LOG(LogIoStore, Display, TEXT("    previous release didn't have these containers. If that doesn't sound right verify"));
+			UE_LOG(LogIoStore, Display, TEXT("    that you used reference containers from the same project. Missing containers:"));
+
+			for (const TPair<FIoContainerId, TUniquePtr<FIoStoreChunkDatabase::FMissingContainerInfo>>& MissingId : ChunkDatabaseRef.MissingContainerIds)
+			{
+				 // we know this is a valid entry because the writers are what gave us the id in the first place.
+				for (uint8 i = 0; i < FIoStoreChunkDatabase::IoChunkTypeCount; i++)
+				{
+					if (MissingId.Value->RequestedChunkCount[i])
+					{
+						UE_LOG(LogIoStore, Display, TEXT("        %s[%s]: %s requests"), 
+							*ContainerNameMap[MissingId.Key], 
+							*LexToString((EIoChunkType)i),
+							*NumberString(MissingId.Value->RequestedChunkCount[i].load(std::memory_order_relaxed)));
+					}
+				}
+			}
 		}
 	}
 
diff --git a/Engine/Source/Runtime/Core/Private/IO/IoStore.cpp b/Engine/Source/Runtime/Core/Private/IO/IoStore.cpp
--- a/Engine/Source/Runtime/Core/Private/IO/IoStore.cpp
+++ b/Engine/Source/Runtime/Core/Private/IO/IoStore.cpp
@@ -172,7 +172,7 @@ struct FChunkBlock
 	FName CompressionMethod = NAME_None;
 	FSHAHash Signature;
 
-	/** Hash of the block data as it would be found on disk */
+	/** Hash of the block data as it would be found on disk - this includes encryption alignment padding */
 	FIoHash DiskHash;
 };
 
@@ -191,7 +191,11 @@ struct FIoStoreWriteQueueEntry
 	// We make this optional because at the latest it might not be valid until FinishCompressionBarrior
 	// completes and we'd like to have a check() on that.
 	TOptional<uint64> UncompressedSize = 0;
-	uint64 CompressedSize = 0;
+
+	// this is not filled out until after encryption completes and *includes the alignment padding for encryption*!
+	// think of this as "size on disk".
+	uint64 CompressedSize = 0; 
+
 	uint64 Padding = 0;
 	uint64 Offset = 0;
 	TArray<FChunkBlock> ChunkBlocks;
@@ -212,6 +216,7 @@ struct FIoStoreWriteQueueEntry
 	bool bModified = false;
 	bool bStoreCompressedDataInDDC = false;
 	
+	bool bCouldBeFromReferenceDb = false; // Whether the chunk is a valid candidate for the reference db.
 	bool bLoadingFromReferenceDb = false;
 	// When we know we're loading from the reference chunk db, we don't read the source
 	// buffer but we still need to know the number of chunks which we get from the refdb.
@@ -726,13 +731,22 @@ public:
 
 	void SetReferenceChunkDatabase(TSharedPtr<IIoStoreWriterReferenceChunkDatabase> InReferenceChunkDatabase)
 	{
-		if (InReferenceChunkDatabase.IsValid() && InReferenceChunkDatabase->GetCompressionBlockSize() != WriterContext->GetSettings().CompressionBlockSize)
+		if (InReferenceChunkDatabase.IsValid() == false)
+		{
+			ReferenceChunkDatabase = InReferenceChunkDatabase;
+			return;
+		}
+
+		if (InReferenceChunkDatabase->GetCompressionBlockSize() != WriterContext->GetSettings().CompressionBlockSize)
 		{
 			UE_LOG(LogIoStore, Warning, TEXT("Reference chunk database has a different compression block size than the current writer!"));
 			UE_LOG(LogIoStore, Warning, TEXT("No chunks will match, so ignoring. ReferenceChunkDb: %d, IoStoreWriter: %d"), InReferenceChunkDatabase->GetCompressionBlockSize(), WriterContext->GetSettings().CompressionBlockSize);
 			return;
 		}
 		ReferenceChunkDatabase = InReferenceChunkDatabase;
+		
+		// Add ourselves to the reference chunk db's list of possibles
+		ReferenceChunkDatabase->NotifyAddedToWriter(ContainerSettings.ContainerId);
 	}
 	void SetHashDatabase(TSharedPtr<IIoStoreWriterHashDatabase> InHashDatabase, bool bInVerifyHashDatabase)
 	{
@@ -866,7 +880,8 @@ public:
 				if (ReferenceChunkDatabase.IsValid() && CompressionMethodForEntry(Entry) != NAME_None)
 				{
 					TPair<FIoContainerId, FIoChunkHash> ChunkKey(ContainerSettings.ContainerId, Entry->ChunkHash);
-					Entry->bLoadingFromReferenceDb = ReferenceChunkDatabase->ChunkExists(ChunkKey, Entry->NumChunkBlocksFromRefDb);
+					Entry->bLoadingFromReferenceDb = ReferenceChunkDatabase->ChunkExists(ChunkKey, Entry->ChunkId, Entry->NumChunkBlocksFromRefDb);
+					Entry->bCouldBeFromReferenceDb = true;
 				}
 				return;
 			}
@@ -901,7 +916,8 @@ public:
 			if (ReferenceChunkDatabase.IsValid() && CompressionMethodForEntry(Entry) != NAME_None)
 			{
 				TPair<FIoContainerId, FIoChunkHash> ChunkKey(ContainerSettings.ContainerId, Entry->ChunkHash);
-				Entry->bLoadingFromReferenceDb = ReferenceChunkDatabase->ChunkExists(ChunkKey, Entry->NumChunkBlocksFromRefDb);
+				Entry->bLoadingFromReferenceDb = ReferenceChunkDatabase->ChunkExists(ChunkKey, Entry->ChunkId, Entry->NumChunkBlocksFromRefDb);
+				Entry->bCouldBeFromReferenceDb = true;
 			}
 
 			// Release the source data buffer, it will be reloaded later when we start compressing the chunk
@@ -1211,6 +1227,7 @@ public:
 					CompressedContainerSize += ExtraPaddingBytes;
 					UncompressedContainerSize += ExtraPaddingBytes;
 					Partition.Offset += ExtraPaddingBytes;
+					TotalPaddingSize += ExtraPaddingBytes;
 				}
 			}
 			
@@ -1271,6 +1288,8 @@ public:
 		Result.PaddingSize = TotalPaddingSize;
 		Result.UncompressedContainerSize = UncompressedContainerSize;
 		Result.CompressedContainerSize = CompressedContainerSize;
+		Result.TotalEntryCompressedSize = TotalEntryCompressedSize;
+		Result.ReferenceCacheMissBytes = ReferenceCacheMissBytes;
 		Result.DirectoryIndexSize = TocResource.Header.DirectoryIndexSize;
 		Result.CompressionMethod = EnumHasAnyFlags(ContainerSettings.ContainerFlags, EIoContainerFlags::Compressed)
 			? WriterSettings.CompressionMethod
@@ -1588,6 +1607,7 @@ private:
 
 				uint64 TotalUncompressedSize = 0;
 				uint8* ReferenceData = ReadResult.IoBuffer.GetData();
+				uint64 TotalAlignedSize = 0;
 				for (int32 BlockIndex = 0; BlockIndex < ReadResult.Blocks.Num(); ++BlockIndex)
 				{
 					FIoStoreCompressedBlockInfo& ReferenceBlock = ReadResult.Blocks[BlockIndex];
@@ -1608,6 +1628,15 @@ private:
 					// we "could".
 					FMemory::Memcpy(Block.IoBuffer->GetData(), ReferenceData, Block.CompressedSize);
 					ReferenceData += ReferenceBlock.AlignedSize;
+					TotalAlignedSize += ReferenceBlock.AlignedSize;
+				}
+
+				if (TotalAlignedSize != ReadResult.IoBuffer.GetSize())
+				{
+					// If we hit this, we might have read garbage memory above! This is very bad.
+					UE_LOG(LogIoStore, Error, TEXT("Block aligned size does not match iobuffer source size! Blocks: %s source size: %s"),
+						*FText::AsNumber(TotalAlignedSize).ToString(),
+						*FText::AsNumber(ReadResult.IoBuffer.GetSize()).ToString());
 				}
 
 				Entry->UncompressedSize.Emplace(TotalUncompressedSize);
@@ -1845,6 +1874,10 @@ private:
 			}
 			if (WriterSettings.CompressionBlockAlignment)
 			{
+				// Try and prevent entries from crossing compression alignment blocks if possible. This is to avoid
+				// small entries from causing multiple file system block reads afaict. Large entries necesarily get
+				// aligned to prevent things like a blocksize + 2 entry being at alignment -1, causing 3 low level reads.
+				// ...I think.
 				bool bCrossesBlockBoundary = Align(TargetPartition->Offset, WriterSettings.CompressionBlockAlignment) != Align(TargetPartition->Offset + Entry->CompressedSize - 1, WriterSettings.CompressionBlockAlignment);
 				if (bCrossesBlockBoundary)
 				{
@@ -1899,6 +1932,13 @@ private:
 			BlockEntry.SetUncompressedSize(uint32(ChunkBlock.UncompressedSize));
 			BlockEntry.SetCompressionMethodIndex(Toc.AddCompressionMethodEntry(ChunkBlock.CompressionMethod));
 
+			// We do this here so that we get the total size of data excluding the encryption alignment
+			TotalEntryCompressedSize += ChunkBlock.CompressedSize;
+			if (Entry->bCouldBeFromReferenceDb && !Entry->bLoadingFromReferenceDb)
+			{
+				ReferenceCacheMissBytes += ChunkBlock.CompressedSize;
+			}
+
 			if (!ChunkBlock.CompressionMethod.IsNone())
 			{
 				ChunkMeta.Flags |= FIoStoreTocEntryMetaFlags::Compressed;
@@ -1984,10 +2024,12 @@ private:
 	TUniquePtr<FArchive>		CsvArchive;
 	FIoStoreWriterResult		Result;
 	uint64						UncompressedFileOffset = 0;
-	uint64						TotalEntryUncompressedSize = 0;
+	uint64						TotalEntryUncompressedSize = 0; // sum of all entry source buffer sizes
+	uint64						TotalEntryCompressedSize = 0; // entry compressed size excluding encryption alignment
+	uint64						ReferenceCacheMissBytes = 0; // number of compressed bytes excluding alignment that could have been from refcache but weren't.
 	uint64						TotalPaddingSize = 0;
-	uint64						UncompressedContainerSize = 0;
-	uint64						CompressedContainerSize = 0;
+	uint64						UncompressedContainerSize = 0; // this is the size the container would be if it were uncompressed.
+	uint64						CompressedContainerSize = 0; // this is the size of the container with the given compression (which may be none).
 	int32						CurrentPartitionIndex = 0;
 	bool						bHasMemoryMappedEntry = false;
 	bool						bHasFlushed = false;
diff --git a/Engine/Source/Runtime/Core/Public/IO/IoDispatcher.h b/Engine/Source/Runtime/Core/Public/IO/IoDispatcher.h
--- a/Engine/Source/Runtime/Core/Public/IO/IoDispatcher.h
+++ b/Engine/Source/Runtime/Core/Public/IO/IoDispatcher.h
@@ -451,6 +451,9 @@ struct FIoStoreWriterSettings
 {
 	FName CompressionMethod = NAME_None;
 	uint64 CompressionBlockSize = 64 << 10;
+
+	// This does not align every entry - it tries to prevent excess crossings of this boundary by inserting padding.
+	// and happens whether or not the entry is compressed.
 	uint64 CompressionBlockAlignment = 0;
 	int32 CompressionMinBytesSaved = 0;
 	int32 CompressionMinPercentSaved = 0;
@@ -514,9 +517,11 @@ struct FIoStoreWriterResult
 	int64 TocSize = 0;
 	int64 TocEntryCount = 0;
 	int64 PaddingSize = 0;
-	int64 UncompressedContainerSize = 0;
-	int64 CompressedContainerSize = 0;
+	int64 UncompressedContainerSize = 0; // this is the size the container would be if it were uncompressed.
+	int64 CompressedContainerSize = 0; // this is the size of the container with the given compression (which may be none). Should be the sum of all partition file sizes.
 	int64 DirectoryIndexSize = 0;
+	uint64 TotalEntryCompressedSize = 0; // sum of the compressed size of entries excluding encryption alignment.
+	uint64 ReferenceCacheMissBytes = 0; // number of compressed bytes excluding alignment that could have been from refcache but weren't.
 	uint64 AddedChunksCount = 0;
 	uint64 AddedChunksSize = 0;
 	uint64 ModifiedChunksCount = 0;
@@ -702,13 +707,18 @@ public:
 	* Quick synchronous existence check that returns the number of blocks for the chunk. This is used to set up
 	* the necessary structures without needing to read the source data for the chunk.
 	*/
-	virtual bool ChunkExists(const TPair<FIoContainerId, FIoChunkHash>& InChunkKey, uint32& OutNumChunkBlocks) = 0;
+	virtual bool ChunkExists(const TPair<FIoContainerId, FIoChunkHash>& InChunkKey, const FIoChunkId& InChunkId, uint32& OutNumChunkBlocks) = 0;
 
 	/*
 	* Returns the compression block size that was used to break up the IoChunks in the source containers. If this is different than what we want, 
 	* then none of the chunks will ever match. Knowing this up front allows us to only match on hash
 	*/
 	virtual uint32 GetCompressionBlockSize() const = 0;
+
+	/*
+	* Called by an iostore writer implementation to notify the ref cache it's been added
+	*/
+	virtual void NotifyAddedToWriter(const FIoContainerId& InContainerId) = 0;
 };
 
 /**
