Get Distribution and Fallback URLs in the same fashion as the Service one

#rnx

diff --git a/Engine/Source/Developer/IoStoreUtilities/Private/IoStoreUtilities.cpp b/Engine/Source/Developer/IoStoreUtilities/Private/IoStoreUtilities.cpp
--- a/Engine/Source/Developer/IoStoreUtilities/Private/IoStoreUtilities.cpp
+++ b/Engine/Source/Developer/IoStoreUtilities/Private/IoStoreUtilities.cpp
@@ -9675,18 +9675,14 @@ bool UploadIoStoreContainerFiles(const UE::IO::IAS::FIoStoreUploadParams& Upload
 		FStringBuilderBase Sb;
 		Sb << TEXT("[Endpoint]") << TEXT("\r\n");
 
-		FString DistributionUrl;
-		if (FParse::Value(FCommandLine::Get(), TEXT("DistributionUrl="), DistributionUrl))
+		if (!UploadParams.DistributionUrl.IsEmpty())
 		{
-			Sb << TEXT("DistributionUrl=\"") << DistributionUrl << TEXT("\"\r\n");
+			Sb << TEXT("DistributionUrl=\"") << UploadParams.DistributionUrl << TEXT("\"\r\n");
 
-			FString FallbackUrl;
-			if (FParse::Value(FCommandLine::Get(), TEXT("FallbackUrl="), FallbackUrl))
+			if (!UploadParams.FallbackUrl.IsEmpty())
 			{
-				Sb << TEXT("FallbackUrl=\"") << FallbackUrl << TEXT("\"\r\n");
+				Sb << TEXT("FallbackUrl=\"") << UploadParams.FallbackUrl << TEXT("\"\r\n");
 			}
-			
-			
 		}
 		else
 		{
diff --git a/Engine/Source/Runtime/Experimental/IoStoreOnDemand/Private/IoStoreOnDemand.cpp b/Engine/Source/Runtime/Experimental/IoStoreOnDemand/Private/IoStoreOnDemand.cpp
--- a/Engine/Source/Runtime/Experimental/IoStoreOnDemand/Private/IoStoreOnDemand.cpp
+++ b/Engine/Source/Runtime/Experimental/IoStoreOnDemand/Private/IoStoreOnDemand.cpp
@@ -1032,6 +1032,8 @@ TIoStatusOr<FIoStoreUploadParams> FIoStoreUploadParams::Parse(const TCHAR* Comma
 
 	FParse::Value(CommandLine, TEXT("BucketPrefix="), Params.BucketPrefix);
 	FParse::Value(CommandLine, TEXT("ServiceUrl="), Params.ServiceUrl);
+	FParse::Value(CommandLine, TEXT("DistributionUrl="), Params.DistributionUrl);
+	FParse::Value(CommandLine, TEXT("FallbackUrl="), Params.FallbackUrl);
 	FParse::Value(CommandLine, TEXT("Region="), Params.Region);
 	FParse::Value(CommandLine, TEXT("AccessKey="), Params.AccessKey);
 	FParse::Value(CommandLine, TEXT("SecretKey="), Params.SecretKey);
diff --git a/Engine/Source/Runtime/Experimental/IoStoreOnDemand/Public/IO/IoStoreOnDemand.h b/Engine/Source/Runtime/Experimental/IoStoreOnDemand/Public/IO/IoStoreOnDemand.h
--- a/Engine/Source/Runtime/Experimental/IoStoreOnDemand/Public/IO/IoStoreOnDemand.h
+++ b/Engine/Source/Runtime/Experimental/IoStoreOnDemand/Public/IO/IoStoreOnDemand.h
@@ -160,6 +160,8 @@ TIoStatusOr<FOnDemandToc> LoadTocFromUrl(const FString& ServiceURL, const FStrin
 struct FIoStoreUploadParams
 {
 	FString ServiceUrl;
+	FString DistributionUrl;
+	FString FallbackUrl;
 	FString Bucket;
 	FString BucketPrefix;
 	FString Region; 
