Key name for replays as an upload object parameter

#rnx

diff --git a/Engine/Source/Developer/IoStoreUtilities/Private/IoStoreUtilities.cpp b/Engine/Source/Developer/IoStoreUtilities/Private/IoStoreUtilities.cpp
--- a/Engine/Source/Developer/IoStoreUtilities/Private/IoStoreUtilities.cpp
+++ b/Engine/Source/Developer/IoStoreUtilities/Private/IoStoreUtilities.cpp
@@ -9693,9 +9693,10 @@ bool UploadIoStoreContainerFiles(const UE::IO::IAS::FIoStoreUploadParams& Upload
 
 		// Temporary solution to get replays working with encrypted on demand content
 		{
-			FString EncryptionKeyName;
-			if (FParse::Value(FCommandLine::Get(), TEXT("OnDemandEncryptionKeyName="), EncryptionKeyName))
+			if (!UploadParams.EncryptionKeyName.IsEmpty())
 			{
+				const TCHAR* EncryptionKeyName = *(UploadParams.EncryptionKeyName);
+
 				TOptional<FNamedAESKey> EncryptionKey;
 				for (const TPair<FGuid, FNamedAESKey>& KeyPair: KeyChain.GetEncryptionKeys())
 				{
@@ -9712,7 +9713,7 @@ bool UploadIoStoreContainerFiles(const UE::IO::IAS::FIoStoreUploadParams& Upload
 				}
 				else
 				{
-					UE_LOG(LogIoStore, Warning, TEXT("Failed to encryption key '%s' in key chain"), *EncryptionKeyName);
+					UE_LOG(LogIoStore, Warning, TEXT("Failed to encryption key '%s' in key chain"), EncryptionKeyName);
 				}
 			}
 		}
diff --git a/Engine/Source/Runtime/Experimental/IoStoreOnDemand/Private/IoStoreOnDemand.cpp b/Engine/Source/Runtime/Experimental/IoStoreOnDemand/Private/IoStoreOnDemand.cpp
--- a/Engine/Source/Runtime/Experimental/IoStoreOnDemand/Private/IoStoreOnDemand.cpp
+++ b/Engine/Source/Runtime/Experimental/IoStoreOnDemand/Private/IoStoreOnDemand.cpp
@@ -1070,6 +1070,7 @@ TIoStatusOr<FIoStoreUploadParams> FIoStoreUploadParams::Parse(const TCHAR* Comma
 	FParse::Value(CommandLine, TEXT("CredentialsFileKeyName="), Params.CredentialsFileKeyName);
 	FParse::Value(CommandLine, TEXT("BuildVersion="), Params.BuildVersion);
 	FParse::Value(CommandLine, TEXT("TargetPlatform="), Params.TargetPlatform);
+	FParse::Value(CommandLine, TEXT("EncryptionKeyName="), Params.EncryptionKeyName);
 	Params.bDeleteContainerFiles = !FParse::Param(CommandLine, TEXT("KeepContainerFiles"));
 	Params.bDeletePakFiles = !FParse::Param(CommandLine, TEXT("KeepPakFiles"));
 	Params.bWriteTocToDisk = FParse::Param(CommandLine, TEXT("WriteTocToDisk"));
diff --git a/Engine/Source/Runtime/Experimental/IoStoreOnDemand/Public/IO/IoStoreOnDemand.h b/Engine/Source/Runtime/Experimental/IoStoreOnDemand/Public/IO/IoStoreOnDemand.h
--- a/Engine/Source/Runtime/Experimental/IoStoreOnDemand/Public/IO/IoStoreOnDemand.h
+++ b/Engine/Source/Runtime/Experimental/IoStoreOnDemand/Public/IO/IoStoreOnDemand.h
@@ -186,6 +186,7 @@ struct FIoStoreUploadParams
 	FString CredentialsFileKeyName;
 	FString BuildVersion;
 	FString TargetPlatform;
+	FString EncryptionKeyName;
 	int32 MaxConcurrentUploads = 16;
 	bool bDeleteContainerFiles = true;
 	bool bDeletePakFiles = true;
