Container upload takes a KeyChain object instead of a bespoke map.

This is a strategic changing. There have been many ways to load keys, all except FKeyChain are deprecated. As IAS is modern it only needs to support the latest, eliminating parsing and handling code for various historical permutations of command line arguments and formats.

#rnx

diff --git a/Engine/Source/Developer/IoStoreUtilities/Private/IoStoreUtilities.cpp b/Engine/Source/Developer/IoStoreUtilities/Private/IoStoreUtilities.cpp
--- a/Engine/Source/Developer/IoStoreUtilities/Private/IoStoreUtilities.cpp
+++ b/Engine/Source/Developer/IoStoreUtilities/Private/IoStoreUtilities.cpp
@@ -9654,13 +9654,7 @@ int32 CreateIoStoreContainerFiles(const TCHAR* CmdLine)
 
 bool UploadIoStoreContainerFiles(const UE::IO::IAS::FIoStoreUploadParams& UploadParams, TConstArrayView<FString> ContainerFiles, const FKeyChain& KeyChain)
 {
-	TMap<FGuid, FAES::FAESKey> EncryptionKeys;
-	for (const TPair<FGuid, FNamedAESKey>& KeyPair: KeyChain.GetEncryptionKeys())
-	{
-		EncryptionKeys.Add(KeyPair.Key, KeyPair.Value.Key);
-	}
-
-	TIoStatusOr<UE::IO::IAS::FIoStoreUploadResult> Result = UE::IO::IAS::UploadContainerFiles(UploadParams, ContainerFiles, EncryptionKeys);
+	TIoStatusOr<UE::IO::IAS::FIoStoreUploadResult> Result = UE::IO::IAS::UploadContainerFiles(UploadParams, ContainerFiles, KeyChain);
 	if (Result.IsOk() == false)
 	{
 		UE_LOG(LogIoStore, Error, TEXT("Failed to upload container file(s), reason '%s'"), *Result.Status().ToString());
diff --git a/Engine/Source/Runtime/Experimental/IoStoreOnDemand/Private/IoStoreOnDemand.cpp b/Engine/Source/Runtime/Experimental/IoStoreOnDemand/Private/IoStoreOnDemand.cpp
--- a/Engine/Source/Runtime/Experimental/IoStoreOnDemand/Private/IoStoreOnDemand.cpp
+++ b/Engine/Source/Runtime/Experimental/IoStoreOnDemand/Private/IoStoreOnDemand.cpp
@@ -38,6 +38,7 @@
 #include "HAL/PlatformFileManager.h"
 #include "IO/IoStore.h"
 #include "Misc/EnumClassFlags.h"
+#include "Misc/KeyChainUtilities.h"
 #include "Serialization/JsonSerializer.h"
 #include "Serialization/JsonWriter.h"
 #include "S3/S3Client.h"
@@ -1121,8 +1122,14 @@ FIoStatus FIoStoreUploadParams::Validate() const
 TIoStatusOr<FIoStoreUploadResult> UploadContainerFiles(
 	const FIoStoreUploadParams& UploadParams,
 	TConstArrayView<FString> ContainerFiles,
-	const TMap<FGuid, FAES::FAESKey>& EncryptionKeys)
+	const FKeyChain& KeyChain)
 {
+	TMap<FGuid, FAES::FAESKey> EncryptionKeys;
+	for (const TPair<FGuid, FNamedAESKey>& KeyPair: KeyChain.GetEncryptionKeys())
+	{
+		EncryptionKeys.Add(KeyPair.Key, KeyPair.Value.Key);
+	}
+
 	struct FContainerStats
 	{
 		FString ContainerName;
diff --git a/Engine/Source/Runtime/Experimental/IoStoreOnDemand/Public/IO/IoStoreOnDemand.h b/Engine/Source/Runtime/Experimental/IoStoreOnDemand/Public/IO/IoStoreOnDemand.h
--- a/Engine/Source/Runtime/Experimental/IoStoreOnDemand/Public/IO/IoStoreOnDemand.h
+++ b/Engine/Source/Runtime/Experimental/IoStoreOnDemand/Public/IO/IoStoreOnDemand.h
@@ -21,6 +21,7 @@
 class FArchive;
 class FCbFieldView;
 class FCbWriter;
+struct FKeyChain;
 class IIoStoreWriter;
 struct FAnalyticsEventAttribute;
 struct FIoContainerSettings;
@@ -218,7 +219,7 @@ struct FIoStoreUploadResult
 UE_API TIoStatusOr<FIoStoreUploadResult> UploadContainerFiles(
 	const FIoStoreUploadParams& UploadParams,
 	TConstArrayView<FString> ContainerFiles,
-	const TMap<FGuid, FAES::FAESKey>& EncryptionKeys);
+	const FKeyChain& KeyChain);
 
 ////////////////////////////////////////////////////////////////////////////////
 struct FIoStoreDownloadParams
