ZenFileSystemManifest fixes for missing plugins on target platforms (relax optimizations in 25794520).
When adding plugin directories also add any valid plugin extension directories (inspired by cl 24072232).
Don't filter out plugins disabled in the editor/cooker, they may still be enabled on the target platform.

todo: Reintroduce more accurate plugin filtering logic.

diff --git a/Engine/Source/Developer/IoStoreUtilities/Private/ZenFileSystemManifest.cpp b/Engine/Source/Developer/IoStoreUtilities/Private/ZenFileSystemManifest.cpp
--- a/Engine/Source/Developer/IoStoreUtilities/Private/ZenFileSystemManifest.cpp
+++ b/Engine/Source/Developer/IoStoreUtilities/Private/ZenFileSystemManifest.cpp
@@ -218,7 +218,6 @@ int32 FZenFileSystemManifest::Generate()
 			GetExtensionDirs(ExtensionDirs, *EngineDir, ExtensionSubDir, PlatformDirectoryNames);
 			for (const FString& Dir : ExtensionDirs)
 			{
-				//AddFilesFromDirectory(Dir.Replace(*EngineDir, TEXT("/{engine}")), Dir, true, AdditionalFilter);
 				AddFilesFromDirectory(Dir.Replace(*EngineDir, TEXT("/{engine}")), Dir, true, AdditionalFilter);
 			}
 			ExtensionDirs.Reset();
@@ -314,8 +313,13 @@ int32 FZenFileSystemManifest::Generate()
 		.IncludeExtension(TEXT("locmeta"))
 		.IncludeExtension(TEXT("locres"));
 
-	const bool FilterDisabledPlugins = true;
+	FFileFilter PluginFilter = FFileFilter()
+		.IncludeExtension(TEXT("uplugin"));
+
+	const bool FilterDisabledPlugins = false;
 	FString PluginTargetPlatformString = PlatformInfo.UBTPlatformString;
+	TSet<FString> PlatformDirectoryNameSet;
+	PlatformDirectoryNameSet.Append(PlatformDirectoryNames);
 	IPluginManager& PluginManager = IPluginManager::Get();
 	TArray<TSharedRef<IPlugin>> DiscoveredPlugins = PluginManager.GetDiscoveredPlugins();
 	for (TSharedRef<IPlugin>& Plugin : DiscoveredPlugins)
@@ -359,6 +363,51 @@ int32 FZenFileSystemManifest::Generate()
 		AddFromPluginPath(EngineDir, ClientDirectory, SourcePath, *ProjectFile);
 		AddFromPluginDir(EngineDir, ClientDirectory, SourcePath, LocalizationDir, true, &LocalizationFilter);
 		AddFromPluginDir(EngineDir, ClientDirectory, SourcePath, ConfigDir, true, &ConfigFilter);
+
+		// Next add any valid plugin extension directories of this plugin.
+		TArray<FString> ExtensionBaseDirs = Plugin->GetExtensionBaseDirs();
+		for (const FString& ExtensionBaseDir : ExtensionBaseDirs)
+		{
+			// Scan the extension path for "Platforms/X" and include this extension if it is not platform specific at all,
+			// or if X is found and it is a valid target platform
+			bool bFoundPlatformsComponent = false;
+			bool bDone = false;
+			bool bIncludeExtension = true;
+			FPathViews::IterateComponents(
+				ExtensionBaseDir,
+				[&bFoundPlatformsComponent, &bDone, &bIncludeExtension, &PlatformDirectoryNameSet](FStringView CurrentPathComponent)
+				{
+					if (!bFoundPlatformsComponent)
+					{
+						if (CurrentPathComponent == TEXTVIEW("Platforms"))
+						{
+							bFoundPlatformsComponent = true;
+						}
+					}
+					else if (!bDone)
+					{
+						const bool bIsValidPlatform = PlatformDirectoryNameSet.Contains(FString(CurrentPathComponent));
+						bIncludeExtension = bIsValidPlatform;
+						bDone = true;
+					}
+					else
+					{
+						// Do nothing.
+					}
+				}
+			);
+
+			if (bIncludeExtension)
+			{
+				FString ExtensionLocalizationDir = ExtensionBaseDir / TEXT("Content") / TEXT("Localization");
+				FString ExtensionConfigDir = ExtensionBaseDir / TEXT("Config");
+				UE_LOG(LogZenFileSystemManifest, Verbose, TEXT("Plugin '%s': ExtensionBaseDir: '%s'"), *ProjectName, *ExtensionBaseDir);
+
+				AddFromPluginDir(EngineDir, ClientDirectory, SourcePath, ExtensionBaseDir, false, &PluginFilter);
+				AddFromPluginDir(EngineDir, ClientDirectory, SourcePath, ExtensionLocalizationDir, true, &LocalizationFilter);
+				AddFromPluginDir(EngineDir, ClientDirectory, SourcePath, ExtensionConfigDir, true, &ConfigFilter);
+			}
+		}
 	}
 
 	FString InternationalizationPresetAsString = UEnum::GetValueAsString(PackagingSettings->InternationalizationPreset);
