UnrealPak: Replace check with warning for missing import Script Object in IoStore describe

diff --git a/Engine/Source/Developer/IoStoreUtilities/Private/IoStoreUtilities.cpp b/Engine/Source/Developer/IoStoreUtilities/Private/IoStoreUtilities.cpp
--- a/Engine/Source/Developer/IoStoreUtilities/Private/IoStoreUtilities.cpp
+++ b/Engine/Source/Developer/IoStoreUtilities/Private/IoStoreUtilities.cpp
@@ -7273,8 +7273,14 @@ namespace DescribeUtils
 						else
 						{
 							FScriptObjectDesc* ScriptObjectDesc = ScriptObjectByGlobalIdMap.Find(Import.GlobalImportIndex);
-							check(ScriptObjectDesc);
-							Import.Name = ScriptObjectDesc->FullName;
+							if (ScriptObjectDesc)
+							{
+								Import.Name = ScriptObjectDesc->FullName;
+							}
+							else
+							{
+								UE_LOG(LogIoStore, Warning, TEXT("Missing Script Object for Import: 0x%llX in package 0x%llX '%s'"), Import.GlobalImportIndex.Value(), PackageDesc->PackageId.ValueForDebugging(), *PackageDesc->PackageName.ToString());
+							}
 						}
 					}
 				}
