Updating optional segment iostore containers to have a unique container id instead of reusing their base container's id. This time to verify I removed the FIoContainerId struct _entirely_ from the codebase to make sure I caught all use cases, and kicked the relevant build that failed.

diff --git a/Engine/Source/Developer/IoStoreUtilities/Private/IoStoreUtilities.cpp b/Engine/Source/Developer/IoStoreUtilities/Private/IoStoreUtilities.cpp
--- a/Engine/Source/Developer/IoStoreUtilities/Private/IoStoreUtilities.cpp
+++ b/Engine/Source/Developer/IoStoreUtilities/Private/IoStoreUtilities.cpp
@@ -1047,6 +1047,7 @@ struct FIoStoreArguments
 struct FContainerTargetSpec
 {
 	FIoContainerId ContainerId;
+	FIoContainerId OptionalSegmentContainerId;
 	FIoContainerHeader Header;
 	FIoContainerHeader OptionalSegmentHeader;
 	FName Name;
@@ -2836,7 +2837,15 @@ void InitializeContainerTargetsAndPackages(
 					ContainerTarget->OptionalSegmentOutputPath = FPaths::Combine(ContainerSource.OptionalOutputPath, FPaths::GetCleanFilename(ContainerTarget->OutputPath) + FPackagePath::GetOptionalSegmentExtensionModifier());
 				}
 
-				UE_LOG(LogIoStore, Display, TEXT("Saving optional container to: '%s'"), *ContainerTarget->OptionalSegmentOutputPath);
+				// The IoContainerId is the hash of the name of the container, which gets returned in the results
+				// as the output path we provide with the extension removed, which for optional containers means
+				// that it contains the .o in the name - so make sure we have a separate id for this.
+				ContainerTarget->OptionalSegmentContainerId = FIoContainerId::FromName(*FPaths::GetCleanFilename(ContainerTarget->OptionalSegmentOutputPath));
+
+				UE_LOG(LogIoStore, Display, TEXT("Saving optional container to: '%s', id: 0x%llx (base container id: 0x%llx)"),
+					*ContainerTarget->OptionalSegmentOutputPath,
+					ContainerTarget->OptionalSegmentContainerId.Value(),
+					ContainerTarget->ContainerId.Value());
 			}
 		}
 	}
@@ -4973,7 +4982,7 @@ void CreateContainerHeader(FContainerTargetSpec& ContainerTarget, bool bIsOption
 {
 	TRACE_CPUPROFILER_EVENT_SCOPE(CreateContainerHeader);
 	FIoContainerHeader& Header = bIsOptional ? ContainerTarget.OptionalSegmentHeader : ContainerTarget.Header;
-	Header.ContainerId = ContainerTarget.ContainerId;
+	Header.ContainerId = bIsOptional ? ContainerTarget.OptionalSegmentContainerId : ContainerTarget.ContainerId;
 
 	int32 NonOptionalSegmentStoreEntriesCount = 0;
 	int32 OptionalSegmentStoreEntriesCount = 0;
@@ -5276,7 +5285,10 @@ int32 CreateTarget(const FIoStoreArguments& Arguments, const FIoStoreWriterSetti
 
 				if (!ContainerTarget->OptionalSegmentOutputPath.IsEmpty())
 				{
+					ContainerSettings.ContainerId = ContainerTarget->OptionalSegmentContainerId;
 					ContainerTarget->OptionalSegmentIoStoreWriter = IoStoreWriterContext->CreateContainer(*ContainerTarget->OptionalSegmentOutputPath, ContainerSettings);
+					ContainerSettings.ContainerId = ContainerTarget->ContainerId;
+
 					ContainerTarget->OptionalSegmentIoStoreWriter->SetReferenceChunkDatabase(ChunkDatabase);
 					IoStoreWriters.Add(ContainerTarget->OptionalSegmentIoStoreWriter);
 					IoStoreWriterInfos.Add({UE::Cook::EPluginSizeTypes::OptionalSegment, ContainerTarget->Name});
@@ -5402,13 +5414,16 @@ int32 CreateTarget(const FIoStoreArguments& Arguments, const FIoStoreWriterSetti
 		{
 			if (ContainerTarget->IoStoreWriter)
 			{
-				auto WriteContainerHeaderChunk = [](FIoContainerHeader& Header, IIoStoreWriter* IoStoreWriter)
+				auto WriteContainerHeaderChunk = [](FIoContainerHeader& Header, IIoStoreWriter* IoStoreWriter, const FIoContainerId& IoStoreWriterId)
 				{
 					FLargeMemoryWriter HeaderAr(0, true);
 					HeaderAr << Header;
 					int64 DataSize = HeaderAr.TotalSize();
 					FIoBuffer ContainerHeaderBuffer(FIoBuffer::AssumeOwnership, HeaderAr.ReleaseOwnership(), DataSize);
 
+					// The header must have the same ID so that the loading code can find it.
+					check(IoStoreWriterId == Header.ContainerId);
+
 					FIoWriteOptions WriteOptions;
 					WriteOptions.DebugName = TEXT("ContainerHeader");
 					WriteOptions.bForceUncompressed = true;
@@ -5419,12 +5434,12 @@ int32 CreateTarget(const FIoStoreArguments& Arguments, const FIoStoreWriterSetti
 				};
 
 				CreateContainerHeader(*ContainerTarget, false);
-				WriteContainerHeaderChunk(ContainerTarget->Header, ContainerTarget->IoStoreWriter.Get());
+				WriteContainerHeaderChunk(ContainerTarget->Header, ContainerTarget->IoStoreWriter.Get(), ContainerTarget->ContainerId);
 
 				if (ContainerTarget->OptionalSegmentIoStoreWriter)
 				{
 					CreateContainerHeader(*ContainerTarget, true);
-					WriteContainerHeaderChunk(ContainerTarget->OptionalSegmentHeader, ContainerTarget->OptionalSegmentIoStoreWriter.Get());
+					WriteContainerHeaderChunk(ContainerTarget->OptionalSegmentHeader, ContainerTarget->OptionalSegmentIoStoreWriter.Get(), ContainerTarget->OptionalSegmentContainerId);
 				}
 			}
 
diff --git a/Engine/Source/Runtime/Core/Private/IO/IoStore.cpp b/Engine/Source/Runtime/Core/Private/IO/IoStore.cpp
--- a/Engine/Source/Runtime/Core/Private/IO/IoStore.cpp
+++ b/Engine/Source/Runtime/Core/Private/IO/IoStore.cpp
@@ -368,7 +368,7 @@ public:
 		return FIoStatus::Ok;
 	}
 
-	TSharedPtr<IIoStoreWriter> CreateContainer(const TCHAR* InContainerPath, const FIoContainerSettings& InContainerSettings);
+	TSharedPtr<IIoStoreWriter> CreateContainer(const TCHAR* InContainerPathAndBaseFileName, const FIoContainerSettings& InContainerSettings);
 
 	void Flush();
 
@@ -724,8 +724,8 @@ class FIoStoreWriter
 	: public IIoStoreWriter
 {
 public:
-	FIoStoreWriter(const TCHAR* InContainerPath)
-		: ContainerPath(InContainerPath)
+	FIoStoreWriter(const TCHAR* InContainerPathAndBaseFileName)
+		: ContainerPathAndBaseFileName(InContainerPathAndBaseFileName)
 	{
 	}
 
@@ -773,7 +773,7 @@ public:
 		WriterContext = &InContext;
 		ContainerSettings = InContainerSettings;
 
-		TocFilePath = ContainerPath + TEXT(".utoc");
+		TocFilePath = ContainerPathAndBaseFileName + TEXT(".utoc");
 		
 		IPlatformFile& Ipf = IPlatformFile::GetPlatformPhysical();
 		Ipf.CreateDirectoryTree(*FPaths::GetPath(TocFilePath));
@@ -1471,7 +1471,7 @@ private:
 	FIoStatus CreatePartitionContainerFile(FPartition& Partition)
 	{
 		check(!Partition.ContainerFileHandle);
-		FString ContainerFilePath = ContainerPath;
+		FString ContainerFilePath = ContainerPathAndBaseFileName;
 		if (Partition.Index > 0)
 		{
 			ContainerFilePath += FString::Printf(TEXT("_s%d"), Partition.Index);
@@ -2009,7 +2009,7 @@ private:
 		WriterContext->SerializedChunksCount.IncrementExchange();
 	}
 
-	const FString				ContainerPath;
+	const FString				ContainerPathAndBaseFileName;
 	FIoStoreWriterContextImpl*	WriterContext = nullptr;
 	FIoContainerSettings		ContainerSettings;
 	FString						TocFilePath;
@@ -2042,9 +2042,11 @@ private:
 	friend class FIoStoreWriterContextImpl;
 };
 
-TSharedPtr<IIoStoreWriter> FIoStoreWriterContextImpl::CreateContainer(const TCHAR* InContainerPath, const FIoContainerSettings& InContainerSettings)
+// InContainerPathAndBaseFileName: the utoc file will just be this with .utoc appended.
+// The base filename ends up getting returned as the container name in the writer results.
+TSharedPtr<IIoStoreWriter> FIoStoreWriterContextImpl::CreateContainer(const TCHAR* InContainerPathAndBaseFileName, const FIoContainerSettings& InContainerSettings)
 {
-	TSharedPtr<FIoStoreWriter> IoStoreWriter = MakeShared<FIoStoreWriter>(InContainerPath);
+	TSharedPtr<FIoStoreWriter> IoStoreWriter = MakeShared<FIoStoreWriter>(InContainerPathAndBaseFileName);
 	FIoStatus IoStatus = IoStoreWriter->Initialize(*this, InContainerSettings);
 	check(IoStatus.IsOk());
 	IoStoreWriters.Add(IoStoreWriter);
diff --git a/Engine/Source/Runtime/Core/Public/IO/IoDispatcher.h b/Engine/Source/Runtime/Core/Public/IO/IoDispatcher.h
--- a/Engine/Source/Runtime/Core/Public/IO/IoDispatcher.h
+++ b/Engine/Source/Runtime/Core/Public/IO/IoDispatcher.h
@@ -513,7 +513,7 @@ struct FIoContainerSettings
 struct FIoStoreWriterResult
 {
 	FIoContainerId ContainerId;
-	FString ContainerName;
+	FString ContainerName; // This is the base filename of the utoc used for output.
 	int64 TocSize = 0;
 	int64 TocEntryCount = 0;
 	int64 PaddingSize = 0;
