Reference Chunk Database:
* Optionally output a CSV of changed/new packages.
* Add an additional directory input if containers get separated off during a build process
* Catch optional segment containers and patch up the container id per the id fix in a separate CL.
bunch of different preflights due to CIS issues.

diff --git a/Engine/Source/Developer/IoStoreUtilities/Private/IoStoreUtilities.cpp b/Engine/Source/Developer/IoStoreUtilities/Private/IoStoreUtilities.cpp
--- a/Engine/Source/Developer/IoStoreUtilities/Private/IoStoreUtilities.cpp
+++ b/Engine/Source/Developer/IoStoreUtilities/Private/IoStoreUtilities.cpp
@@ -94,7 +94,6 @@ struct FIoStoreChunkSource
 {
 	FIoStoreTocChunkInfo ChunkInfo;
 	UE::Cook::EPluginSizeTypes SizeType;
-
 };
 
 static const FName DefaultCompressionMethod = NAME_Zlib;
@@ -114,62 +113,112 @@ class FIoStoreChunkDatabase : public IIoStoreWriterReferenceChunkDatabase
 public:
 	static const uint8 IoChunkTypeCount = (uint8)EIoChunkType::MAX;
 
+	struct FRequestedChunkInfo
+	{
+		FIoChunkId Id;
+		FIoChunkHash Hash;
+	};
+
 	TArray<TUniquePtr<FIoStoreReader>> Readers;
 	struct FReaderChunks
 	{
-		int32 ReaderIndex;
-		TMap<FIoChunkHash, FIoStoreTocChunkInfo> Chunks;
-		TSet<FIoChunkId> ChunkIds;
+		FString ContainerPathAndFileName;
+		FString ContainerName;
+		FIoContainerId ContainerId;
+
+		int32 ReaderIndex = -1;
+
+		TMap<FIoChunkId, uint32> ChunkIds; // maps to index for chunk specific structures
 
-		std::atomic_uint64_t ChangedChunkCount[IoChunkTypeCount];
-		std::atomic_uint64_t NewChunkCount[IoChunkTypeCount];
-		std::atomic_uint64_t UsedChunkCount[IoChunkTypeCount];
+		TArray<FIoChunkHash> ChangedChunkHashes;
+		TArray<FIoStoreTocChunkInfo> ChunkInfos;
+		TArray<std::atomic_uint8_t> ChunkChanged;
+
+		~FReaderChunks() = default;
+		FReaderChunks() = default;
+		FReaderChunks(FReaderChunks&&) = delete;
+		FReaderChunks(FReaderChunks&) = delete;
+		FReaderChunks& operator=(FReaderChunks&) = delete;
+
+		FCriticalSection NewChunksLock;
+		TArray<FRequestedChunkInfo> NewChunks;
+		uint64 NewChunksByType[IoChunkTypeCount] = { 0 };
+
+		std::atomic_uint64_t TotalRequestCount = 0;
+		std::atomic_uint64_t ChangedCountByType[IoChunkTypeCount] = {0};
+		std::atomic_uint64_t UsedChunksByType[IoChunkTypeCount] = {0};
 	};
 
 	struct FMissingContainerInfo
 	{
-		std::atomic_uint64_t RequestedChunkCount[IoChunkTypeCount];
-	};
+		FMissingContainerInfo(FString InContainerName) {ContainerName = InContainerName;}
+		TArray<FRequestedChunkInfo> Chunks;
+		FCriticalSection ChunksLock;
 
-	TMap<FIoContainerId, TUniquePtr<FMissingContainerInfo>> MissingContainerIds;
+		FString ContainerName;
+		std::atomic_uint64_t TotalRequestCount = 0;
+		std::atomic_uint64_t RequestedChunkCount[IoChunkTypeCount] = {0};
+	};
 
-	TMap<FIoContainerId, FString> ContainerNameMap;
 
+	// Since we are notified of all containers, every id should be in our database,
+	// or in the missing list.
 	TMap<FIoContainerId, TUniquePtr<FReaderChunks>> ChunkDatabase;
+	TMap<FIoContainerId, TUniquePtr<FMissingContainerInfo>> MissingContainers;
 
-
-	std::atomic_int32_t RequestCount = 0;
-	int32 FulfillCount = 0;
+	std::atomic_int64_t MatchCount = 0;
+	std::atomic_int64_t RequestCount = 0;
 	int32 ContainerNotFound = 0;
+
+	// Bytes we actually delivered to the iostore writer
 	int64 FulfillBytes = 0;
 	int64 FulfillBytesPerChunk[(int8)EIoChunkType::MAX] = {};
+
 	uint32 CompressionBlockSize = 0;
 
-	bool Init(const FString& InGlobalContainerFileName, const FKeyChain& InDecryptionKeychain)
+	bool Init(FString InGlobalContainerFileName, FString InAdditionalContainersPath, const FKeyChain& InDecryptionKeychain)
 	{
 		double StartTime = FPlatformTime::Seconds();
 
-		TUniquePtr<FIoStoreReader> GlobalReader = CreateIoStoreReader(*InGlobalContainerFileName, InDecryptionKeychain);
+		// Try and catch common pathing mistakes since honestly we only care about the path even though we ask
+		// for the global.utoc.
+		FString ContainersDirectory = FPaths::GetPath(MoveTemp(InGlobalContainerFileName));
+		FPaths::NormalizeDirectoryName(ContainersDirectory);
+		FString GlobalFileName = ContainersDirectory / TEXT("global.utoc");
+
+		TUniquePtr<FIoStoreReader> GlobalReader = CreateIoStoreReader(*GlobalFileName, InDecryptionKeychain);
 		if (GlobalReader.IsValid() == false)
 		{
-			UE_LOG(LogIoStore, Error, TEXT("Failed to open reference chunk container %s"), *InGlobalContainerFileName);
+			UE_LOG(LogIoStore, Error, TEXT("Failed to open reference chunk container %s"), *GlobalFileName);
 			return false;
 		}
 
-		FString Directory = FPaths::GetPath(InGlobalContainerFileName);
-		FPaths::NormalizeDirectoryName(Directory);
+		FPaths::NormalizeDirectoryName(InAdditionalContainersPath);
 
-		TArray<FString> FoundContainerFiles;
-		IFileManager::Get().FindFiles(FoundContainerFiles, *(Directory / TEXT("*.utoc")), true, false);
 		TArray<FString> ContainerFilePaths;
-		for (const FString& Filename : FoundContainerFiles)
+
+		TArray<FString> FoundContainerFiles;
 		{
-			ContainerFilePaths.Emplace(Directory / Filename);
+			IFileManager::Get().FindFiles(FoundContainerFiles, *(ContainersDirectory / TEXT("*.utoc")), true, false);
+			for (const FString& Filename : FoundContainerFiles)
+			{
+				ContainerFilePaths.Emplace(ContainersDirectory / Filename);
+			}
+		}
+
+		if (InAdditionalContainersPath.Len())
+		{
+			FoundContainerFiles.Empty();
+			IFileManager::Get().FindFiles(FoundContainerFiles, *(InAdditionalContainersPath / TEXT("*.utoc")), true, false);
+			for (const FString& Filename : FoundContainerFiles)
+			{
+				ContainerFilePaths.Emplace(InAdditionalContainersPath / Filename);
+			}
 		}
 
 		CompressionBlockSize = 0;
 		int64 IoChunkCount = 0;
-		for (const FString& ContainerFilePath : ContainerFilePaths)
+		for (FString& ContainerFilePath : ContainerFilePaths)
 		{
 			TUniquePtr<FIoStoreReader> Reader = CreateIoStoreReader(*ContainerFilePath, InDecryptionKeychain);
 			if (Reader.IsValid() == false)
@@ -178,24 +227,49 @@ public:
 				return false;
 			}
 
-			if (ChunkDatabase.Contains(Reader->GetContainerId()))
+			// Work around the fact that we recently changed container ids to be unique and detect .o containers
+			// and regenerate the Id.
+			FIoContainerId ContainerId = Reader->GetContainerId();
+			if (ContainerFilePath.Contains(FPackagePath::GetOptionalSegmentExtensionModifier()))
+			{
+				FIoContainerId CheckContainerId = FIoContainerId::FromName(*FPaths::GetBaseFilename(ContainerFilePath));
+				if (CheckContainerId != ContainerId)
+				{
+					UE_LOG(LogIoStore, Display, TEXT("Reference Chunk: adjusting optional segment container id (0x%llx -> 0x%llx) for %s"), ContainerId.Value(), CheckContainerId.Value(), *ContainerFilePath);
+					ContainerId = CheckContainerId;
+				}
+			}
+
+			TUniquePtr<FReaderChunks>* ExistingContainer = ChunkDatabase.Find(ContainerId);
+			if (ExistingContainer)
 			{
-				UE_LOG(LogIoStore, Error, TEXT("Duplicate container id found in reference chunk container directory %s"), *ContainerFilePath);
-				UE_LOG(LogIoStore, Error, TEXT("is duplicate of %s"), *ContainerNameMap[Reader->GetContainerId()]);
+				UE_LOG(LogIoStore, Error, TEXT("Reference chunk database loaded containers with the same id:"));
+				UE_LOG(LogIoStore, Error, TEXT("    %s"), *ExistingContainer[0]->ContainerPathAndFileName);
+				UE_LOG(LogIoStore, Error, TEXT("    %s"), *ContainerFilePath);
 				return false;
 			}
 
+
 			TUniquePtr<FReaderChunks> ReaderChunks(new FReaderChunks());
+			ReaderChunks->ContainerName = Reader->GetContainerName();
+			ReaderChunks->ContainerPathAndFileName = MoveTemp(ContainerFilePath);
+			ReaderChunks->ContainerId = ContainerId;
 
-			ContainerNameMap.Add(Reader->GetContainerId(), *ContainerFilePath);
+			int32 ChunkCount = Reader->GetChunkCount();
+			ReaderChunks->ChunkInfos.Reserve(ChunkCount);
+			ReaderChunks->ChangedChunkHashes.SetNumZeroed(ChunkCount);
 
-			Reader->EnumerateChunks([ReaderChunks = ReaderChunks.Get()](FIoStoreTocChunkInfo&& ChunkInfo)
+			uint32 ChunkIndex = 0;
+			Reader->EnumerateChunks([ReaderChunks = ReaderChunks.Get(), &ChunkIndex](FIoStoreTocChunkInfo&& ChunkInfo)
 			{
-				ReaderChunks->ChunkIds.Add(ChunkInfo.Id);
-				ReaderChunks->Chunks.Add(TPair<FIoChunkHash, FIoStoreTocChunkInfo>(ChunkInfo.Hash, MoveTemp(ChunkInfo)));
+				ReaderChunks->ChunkIds.Add(ChunkInfo.Id, ChunkIndex);
+				ReaderChunks->ChunkInfos.Add(MoveTemp(ChunkInfo));
+				ChunkIndex++;
 				return true;
 			});
 
+			ReaderChunks->ChunkChanged.SetNumZeroed(ReaderChunks->ChunkIds.Num());
+
 			if (Readers.Num() == 0)
 			{
 				CompressionBlockSize = Reader->GetCompressionBlockSize();
@@ -206,25 +280,25 @@ public:
 				return false;
 			}
 
-			IoChunkCount += ReaderChunks->Chunks.Num();
+			IoChunkCount += ChunkCount;
 			ReaderChunks->ReaderIndex = Readers.Num();
-			ChunkDatabase.Add(Reader->GetContainerId(), MoveTemp(ReaderChunks));
-			Readers.Add(MoveTemp(Reader));
 
-			
+			ChunkDatabase.Add(ReaderChunks->ContainerId, MoveTemp(ReaderChunks));
+			Readers.Add(MoveTemp(Reader));
 		}
 
-		UE_LOG(LogIoStore, Display, TEXT("Block reference loaded %d containers and %s chunks, in %.1f seconds"), Readers.Num(), *FText::AsNumber(IoChunkCount).ToString(), FPlatformTime::Seconds() - StartTime);
+		UE_LOG(LogIoStore, Display, TEXT("Reference Chunk loaded %d containers and %s chunks, in %.1f seconds"), Readers.Num(), *FText::AsNumber(IoChunkCount).ToString(), FPlatformTime::Seconds() - StartTime);
 		return true;
 	}
 
-	virtual void NotifyAddedToWriter(const FIoContainerId& InContainerId) override
+	virtual void NotifyAddedToWriter(const FIoContainerId& InContainerId, const FString& InContainerName) override
 	{
 		// If we don't have this container in our chunk database, add it to a tracking
 		// structure so we can see how many chunks we're missing.
 		if (ChunkDatabase.Contains(InContainerId) == false)
 		{
-			MissingContainerIds.Add(InContainerId, MakeUnique<FMissingContainerInfo>());
+			TUniquePtr<FMissingContainerInfo> MissingContainer = MakeUnique<FMissingContainerInfo>(InContainerName);
+			MissingContainers.Add(InContainerId, MoveTemp(MissingContainer));
 		}
 	}
 
@@ -235,76 +309,110 @@ public:
 
 	// Returns whether we expect to be able to load the chunk from the reference chunk database.
 	// This can be called from any thread, though in the presence of existing hashes it's single threaded.
-	virtual bool ChunkExists(const TPair<FIoContainerId, FIoChunkHash>& InChunkKey, const FIoChunkId& InChunkId, uint32& OutNumChunkBlocks)
+	virtual bool ChunkExists(const FIoContainerId& InContainerId, const FIoChunkHash& InChunkHash, const FIoChunkId& InChunkId, uint32& OutNumChunkBlocks)
 	{
 		RequestCount.fetch_add(1, std::memory_order_relaxed);
 
-		TUniquePtr<FReaderChunks>* ReaderChunks = ChunkDatabase.Find(InChunkKey.Key);
-		if (ReaderChunks == nullptr)
+		uint8 ChunkType = (uint8)InChunkId.GetChunkType();
+
+		FRequestedChunkInfo RCI;
+		RCI.Hash = InChunkHash;
+		RCI.Id = InChunkId;
+
+		TUniquePtr<FReaderChunks>* ReaderChunksPtr = ChunkDatabase.Find(InContainerId);
+		if (ReaderChunksPtr == nullptr)
 		{
-			// Container doesn't exist - likely provided the path to a different project.
-			TUniquePtr<FMissingContainerInfo>* MissingContainerHitCount = MissingContainerIds.Find(InChunkKey.Key);
+			// Container doesn't exist - could have been added/created or maybe wrong source project.
+			TUniquePtr<FMissingContainerInfo>* MissingContainerPtr = MissingContainers.Find(InContainerId);
 
-			if (MissingContainerHitCount == nullptr)
+			if (MissingContainerPtr == nullptr)
 			{
-				UE_LOG(LogIoStore, Warning, TEXT("We got a container id that was never added! id = %llu"), InChunkKey.Key.Value());
+				UE_LOG(LogIoStore, Warning, TEXT("We got a container that was never added! Id = 0x%llx"), InContainerId.Value());
 			}
 			else
 			{
-				MissingContainerHitCount[0]->RequestedChunkCount[(uint8)InChunkId.GetChunkType()].fetch_add(1, std::memory_order_relaxed);
+				// We expect most chunks to have containers so we don't stress about this lock too much
+				FMissingContainerInfo* MissingContainer = MissingContainerPtr->Get();
+				{
+					FScopeLock Locker(&MissingContainer->ChunksLock);
+					MissingContainer->Chunks.Add(MoveTemp(RCI));
+				}
+
+				MissingContainer->RequestedChunkCount[ChunkType].fetch_add(1, std::memory_order_relaxed);
+				if (MissingContainer->TotalRequestCount.fetch_add(1, std::memory_order_relaxed) == 0)
+				{
+					ContainerNotFound++;
+				}
 			}
 
-			ContainerNotFound++;
 			return false;
 		}
 
-		FIoStoreTocChunkInfo* ChunkInfo = ReaderChunks[0]->Chunks.Find(InChunkKey.Value);
-		if (ChunkInfo == nullptr)
+		// We have the container at least.
+		FReaderChunks* ReaderChunks = ReaderChunksPtr->Get();
+		ReaderChunks->TotalRequestCount.fetch_add(1, std::memory_order_relaxed);
+
+		uint32_t* ChunkIndex = ReaderChunks->ChunkIds.Find(InChunkId);
+		if (ChunkIndex == nullptr)
 		{
-			// No exact chunk data match - this is a normal exit condition for a changed block or
-			// a new block.
-			bool bChangedBlock = ReaderChunks[0]->ChunkIds.Contains(InChunkId);
-			if (bChangedBlock)
-			{
-				ReaderChunks[0]->ChangedChunkCount[(uint8)InChunkId.GetChunkType()].fetch_add(1, std::memory_order_relaxed);
-			}
-			else
-			{
-				ReaderChunks[0]->NewChunkCount[(uint8)InChunkId.GetChunkType()].fetch_add(1, std::memory_order_relaxed);
-			}
+			// Chunk doesn't exist - it's entirely new.
+			// We expect most chunks to already exist so we don't stress about this lock too much
+			FScopeLock Locker(&ReaderChunks->NewChunksLock);
+			ReaderChunks->NewChunksByType[ChunkType]++;
+			ReaderChunks->NewChunks.Add(MoveTemp(RCI));
 			return false;
 		}
 
-		OutNumChunkBlocks = ChunkInfo->NumCompressedBlocks;
+		// We have the chunk - does the hash match?
+		const FIoStoreTocChunkInfo& ChunkInfo = ReaderChunks->ChunkInfos[*ChunkIndex];
+		if (ChunkInfo.Hash != InChunkHash)
+		{
+			// Chunk exists, but it's changed.
+			ReaderChunks->ChangedCountByType[ChunkType].fetch_add(1, std::memory_order_relaxed);
+			ReaderChunks->ChunkChanged[*ChunkIndex].store(1, std::memory_order_relaxed);
+
+			// This should be safe because chunk ids are unique within a container,
+			// so we are the only ones poking at this index.
+			ReaderChunks->ChangedChunkHashes[*ChunkIndex] = MoveTemp(RCI.Hash);
+			return false;
+		}
+
+		// We match.
+		MatchCount.fetch_add(1, std::memory_order_relaxed);
+		ReaderChunks->UsedChunksByType[ChunkType].fetch_add(1, std::memory_order_relaxed);
+		OutNumChunkBlocks = ChunkInfo.NumCompressedBlocks;
 		return true;
 	}
 
 
-	// Not thread safe, called from the BeginCompress dispatch thread.
-	virtual bool RetrieveChunk(const TPair<FIoContainerId, FIoChunkHash>& InChunkKey, TUniqueFunction<void(TIoStatusOr<FIoStoreCompressedReadResult>)> InCompleteCallback)
+	// This function was not written to be thread safe as it's only ever called from
+	// the iostore begindispatch thread (i.e. is not async)
+	virtual bool RetrieveChunk(const FIoContainerId& InContainerId, const FIoChunkHash& InChunkHash, const FIoChunkId& InChunkId, TUniqueFunction<void(TIoStatusOr<FIoStoreCompressedReadResult>)> InCompleteCallback)
 	{
-		TUniquePtr<FReaderChunks>* ReaderChunks = ChunkDatabase.Find(InChunkKey.Key);
-		if (ReaderChunks == nullptr)
+		TUniquePtr<FReaderChunks>* ReaderChunksPtr = ChunkDatabase.Find(InContainerId);
+		if (ReaderChunksPtr == nullptr)
 		{
 			// This should never happen now as we wrap this in a ChunkExists call.
-			UE_LOG(LogIoStore, Warning, TEXT("RetrieveChunk can't find the container - invariant violated!"));
+			UE_LOG(LogIoStore, Warning, TEXT("RetrieveChunk can't find the container  invariant violated!"));
 			return false;
 		}
 
-		FIoStoreTocChunkInfo* ChunkInfo = ReaderChunks[0]->Chunks.Find(InChunkKey.Value);
-		if (ChunkInfo == nullptr)
+		FReaderChunks* ReaderChunks = ReaderChunksPtr->Get();
+
+		uint32_t* ChunkIndex = ReaderChunks->ChunkIds.Find(InChunkId);
+		if (ChunkIndex == nullptr)
 		{
 			// This should never happen now as we wrap this in a ChunkExists call.
-			UE_LOG(LogIoStore, Warning, TEXT("RetrieveChunk can't find the chunnk - invariant violated!"));
+			UE_LOG(LogIoStore, Warning, TEXT("RetrieveChunk can't find the chunk - invariant violated!"));
 			return false;
 		}
 
-		ReaderChunks[0]->UsedChunkCount[(uint8)ChunkInfo->Id.GetChunkType()].fetch_add(1, std::memory_order_relaxed);
+		FIoStoreTocChunkInfo& ChunkInfo = ReaderChunks->ChunkInfos[*ChunkIndex];
 
 		uint64 TotalCompressedSize = 0;
 		uint64 TotalUncompressedSize = 0;
 		uint32 CompressedBlockCount = 0;
-		Readers[ReaderChunks[0]->ReaderIndex]->EnumerateCompressedBlocksForChunk(ChunkInfo->Id, [&TotalUncompressedSize, &CompressedBlockCount, &TotalCompressedSize](const FIoStoreTocCompressedBlockInfo& BlockInfo)
+		Readers[ReaderChunks->ReaderIndex]->EnumerateCompressedBlocksForChunk(InChunkId, [&TotalUncompressedSize, &CompressedBlockCount, &TotalCompressedSize](const FIoStoreTocCompressedBlockInfo& BlockInfo)
 		{			
 			TotalCompressedSize += BlockInfo.CompressedSize;
 			TotalUncompressedSize += BlockInfo.UncompressedSize;
@@ -312,14 +420,13 @@ public:
 			return true;
 		});
 
-		FulfillBytesPerChunk[(int8)ChunkInfo->ChunkType] += TotalCompressedSize;
+		FulfillBytesPerChunk[(int8)ChunkInfo.ChunkType] += TotalCompressedSize;
 		FulfillBytes += TotalCompressedSize;
-		FulfillCount++;
 
 		//
 		// At this point we know we can use the block so we can go async.
 		//
-		FFunctionGraphTask::CreateAndDispatchWhenReady([this, Id = ChunkInfo->Id, ReaderIndex = ReaderChunks[0]->ReaderIndex, CompleteCallback = MoveTemp(InCompleteCallback)]()
+		FFunctionGraphTask::CreateAndDispatchWhenReady([this, Id = InChunkId, ReaderIndex = ReaderChunks->ReaderIndex, CompleteCallback = MoveTemp(InCompleteCallback)]()
 		{
 			TIoStatusOr<FIoStoreCompressedReadResult> Result = Readers[ReaderIndex]->ReadCompressed(Id, FIoReadOptions());
 			CompleteCallback(Result);
@@ -327,6 +434,89 @@ public:
 
 		return true;
 	}
+
+	void WriteCSV(const FString& InOutputFileName, const TArray<struct FCookedPackage*>& InPackages);
+
+	void LogSummary(const TArray<FIoStoreWriterResult>& IoStoreWriterResults)
+	{
+		uint64 TotalEntryBytes = 0;
+		uint64 TotalMissBytes = 0;
+
+		for (const FIoStoreWriterResult& Result : IoStoreWriterResults)
+		{
+			TotalEntryBytes += Result.TotalEntryCompressedSize;
+			TotalMissBytes += Result.ReferenceCacheMissBytes;
+		}
+
+		uint64 TotalCandidateBytes = FulfillBytes + TotalMissBytes;
+
+		UE_LOG(LogIoStore, Display, TEXT("Reference Chunk Database:"));
+		UE_LOG(LogIoStore, Display, TEXT("    %s reused bytes out of %s candidate bytes - %.1f%% hit rate."),
+			*NumberString(FulfillBytes),
+			*NumberString(TotalCandidateBytes),
+			100.0 * FulfillBytes / (TotalCandidateBytes));
+		UE_LOG(LogIoStore, Display, TEXT("    %s candidate bytes out of %s io chunk bytes - %.1f%% coverage."),
+			*NumberString(TotalCandidateBytes),
+			*NumberString(TotalEntryBytes),
+			100.0 * TotalCandidateBytes / (TotalEntryBytes));
+
+		UE_LOG(LogIoStore, Display, TEXT("    %s chunks matched out of %s requests"), 
+			*NumberString(MatchCount.load(std::memory_order_relaxed)),
+			*NumberString(RequestCount.load(std::memory_order_relaxed)));
+
+		FString ChunkNames[FIoStoreChunkDatabase::IoChunkTypeCount];
+		for (uint8 TypeIndex = 0; TypeIndex < FIoStoreChunkDatabase::IoChunkTypeCount; TypeIndex++)
+		{
+			ChunkNames[TypeIndex] = LexToString((EIoChunkType)TypeIndex);
+		}
+
+		ChunkDatabase.ValueSort([](const TUniquePtr<FReaderChunks>& A, const TUniquePtr<FReaderChunks>& B)
+		{
+			return A->TotalRequestCount.load(std::memory_order_relaxed) > B->TotalRequestCount.load(std::memory_order_relaxed);
+		});
+
+		for (const TPair<FIoContainerId, TUniquePtr<FReaderChunks>>& ReaderPair: ChunkDatabase)
+		{
+			FReaderChunks* Reader = ReaderPair.Value.Get();
+
+			for (uint8 i = 0; i < FIoStoreChunkDatabase::IoChunkTypeCount; i++)
+			{
+				if (Reader->ChangedCountByType[i] ||
+					Reader->NewChunksByType[i] ||
+					Reader->UsedChunksByType[i])
+				{
+					UE_LOG(LogIoStore, Display, TEXT("        %s[%s]:    %s changed, %s new, %s reused"),
+						*Reader->ContainerName,
+						*ChunkNames[i],
+						*NumberString(Reader->ChangedCountByType[i]),
+						*NumberString(Reader->NewChunksByType[i]),
+						*NumberString(Reader->UsedChunksByType[i]));
+				}
+			}
+		}
+
+		if (ContainerNotFound)
+		{
+			UE_LOG(LogIoStore, Display, TEXT("    %s containers were requested that weren't available. This means the "),
+				*NumberString(ContainerNotFound));
+			UE_LOG(LogIoStore, Display, TEXT("    previous release didn't have these containers. If that doesn't sound right verify"));
+			UE_LOG(LogIoStore, Display, TEXT("    that you used reference containers from the same project. Missing containers:"));
+
+			for (const TPair<FIoContainerId, TUniquePtr<FMissingContainerInfo>>& MissingContainer : MissingContainers)
+			{
+				for (uint8 i = 0; i < FIoStoreChunkDatabase::IoChunkTypeCount; i++)
+				{
+					if (MissingContainer.Value->RequestedChunkCount[i])
+					{
+						UE_LOG(LogIoStore, Display, TEXT("        %s[%s]: %s requests"),
+							*MissingContainer.Value->ContainerName,
+							*ChunkNames[i],
+							*NumberString(MissingContainer.Value->RequestedChunkCount[i].load(std::memory_order_relaxed)));
+					}
+				}
+			}
+		}
+	}
 };
 
 struct FReleasedPackages
@@ -1020,6 +1210,8 @@ struct FIoStoreArguments
 	FString DLCPluginPath;
 	FString DLCName;
 	FString ReferenceChunkGlobalContainerFileName;
+	FString ReferenceChunkChangesCSVFileName;
+	FString ReferenceChunkAdditionalContainersPath;
 	FString CsvPath;
 	FKeyChain ReferenceChunkKeys;
 	FReleasedPackages ReleasedPackages;
@@ -5151,7 +5343,9 @@ int32 CreateTarget(const FIoStoreArguments& Arguments, const FIoStoreWriterSetti
 	if (Arguments.ReferenceChunkGlobalContainerFileName.Len())
 	{
 		ChunkDatabase = MakeShared<FIoStoreChunkDatabase>();
-		if (((FIoStoreChunkDatabase&)*ChunkDatabase).Init(Arguments.ReferenceChunkGlobalContainerFileName, Arguments.ReferenceChunkKeys) == false)
+		FIoStoreChunkDatabase* ChunkDbPtr = (FIoStoreChunkDatabase*)ChunkDatabase.Get();
+
+		if (ChunkDbPtr->Init(Arguments.ReferenceChunkGlobalContainerFileName, Arguments.ReferenceChunkAdditionalContainersPath, Arguments.ReferenceChunkKeys) == false)
 		{
 			UE_LOG(LogIoStore, Warning, TEXT("Failed to initialize reference chunk store. Pak will continue."));
 		}
@@ -5775,80 +5969,10 @@ int32 CreateTarget(const FIoStoreArguments& Arguments, const FIoStoreWriterSetti
 
 	if (ChunkDatabase.IsValid())
 	{
-		//
-		// If we are using a reference cache, the assumption is that we are expecting a high hit rate
-		// on the chunks - this is supposed to be a current release, so theoretically we should only
-		// miss on changed or new data, so it's useful to know where the misses are in case something
-		// major happened that maybe shouldn't have.
-		//
-		uint64 TotalEntryBytes = 0;
-		uint64 TotalMissBytes = 0;
-		TMap<FIoContainerId, FString> ContainerNameMap;
-
-		for (const FIoStoreWriterResult& Result : IoStoreWriterResults)
-		{
-			TotalEntryBytes += Result.TotalEntryCompressedSize;
-			TotalMissBytes += Result.ReferenceCacheMissBytes;
-			ContainerNameMap.Add(Result.ContainerId, Result.ContainerName);
-		}
-
-		FIoStoreChunkDatabase& ChunkDatabaseRef = (FIoStoreChunkDatabase&)*ChunkDatabase;
-
-		uint64 TotalCandidateBytes = ChunkDatabaseRef.FulfillBytes + TotalMissBytes;
-
-		UE_LOG(LogIoStore, Display, TEXT("Reference Chunk Database:"));
-		UE_LOG(LogIoStore, Display, TEXT("    %s reused bytes out of %s candidate bytes - %.1f%% hit rate."),
-			*NumberString(ChunkDatabaseRef.FulfillBytes),
-			*NumberString(TotalCandidateBytes),
-			100.0 * ChunkDatabaseRef.FulfillBytes / (TotalCandidateBytes));
-		UE_LOG(LogIoStore, Display, TEXT("    %s candidate bytes out of %s io chunk bytes - %.1f%% coverage."),
-			*NumberString(TotalCandidateBytes),
-			*NumberString(TotalEntryBytes),
-			100.0 * TotalCandidateBytes / (TotalEntryBytes));
-
-		UE_LOG(LogIoStore, Display, TEXT("    %s chunks found out of %s requests"), 
-			*NumberString(ChunkDatabaseRef.FulfillCount),
-			*NumberString(ChunkDatabaseRef.RequestCount.load()));
-		for (const TPair<FIoContainerId, TUniquePtr<FIoStoreChunkDatabase::FReaderChunks>>& ContainerInDatabase : ChunkDatabaseRef.ChunkDatabase)
-		{
-			for (uint8 i = 0; i < FIoStoreChunkDatabase::IoChunkTypeCount; i++)
-			{
-				if (ContainerInDatabase.Value->ChangedChunkCount[i] ||
-					ContainerInDatabase.Value->NewChunkCount[i] ||
-					ContainerInDatabase.Value->UsedChunkCount[i])
-				{
-					UE_LOG(LogIoStore, Display, TEXT("        %s[%s]:    %s changed, %s new, %s reused"), 
-						*ContainerNameMap[ContainerInDatabase.Key],
-						*LexToString((EIoChunkType)i),
-						*NumberString(ContainerInDatabase.Value->ChangedChunkCount[i]),
-						*NumberString(ContainerInDatabase.Value->NewChunkCount[i]),
-						*NumberString(ContainerInDatabase.Value->UsedChunkCount[i]));
-				}
-			}
-		}
-
-		if (ChunkDatabaseRef.ContainerNotFound)
-		{
-			UE_LOG(LogIoStore, Display, TEXT("    %s containers were requested that weren't available. This means the "),
-				*NumberString(ChunkDatabaseRef.MissingContainerIds.Num()));
-			UE_LOG(LogIoStore, Display, TEXT("    previous release didn't have these containers. If that doesn't sound right verify"));
-			UE_LOG(LogIoStore, Display, TEXT("    that you used reference containers from the same project. Missing containers:"));
-
-			for (const TPair<FIoContainerId, TUniquePtr<FIoStoreChunkDatabase::FMissingContainerInfo>>& MissingId : ChunkDatabaseRef.MissingContainerIds)
-			{
-				 // we know this is a valid entry because the writers are what gave us the id in the first place.
-				for (uint8 i = 0; i < FIoStoreChunkDatabase::IoChunkTypeCount; i++)
-				{
-					if (MissingId.Value->RequestedChunkCount[i])
-					{
-						UE_LOG(LogIoStore, Display, TEXT("        %s[%s]: %s requests"), 
-							*ContainerNameMap[MissingId.Key], 
-							*LexToString((EIoChunkType)i),
-							*NumberString(MissingId.Value->RequestedChunkCount[i].load(std::memory_order_relaxed)));
-					}
-				}
-			}
-		}
+		TRACE_CPUPROFILER_EVENT_SCOPE(ReferenceChunkLogging);
+		FIoStoreChunkDatabase* ChunkDbPtr = (FIoStoreChunkDatabase*)ChunkDatabase.Get();
+		ChunkDbPtr->WriteCSV(Arguments.ReferenceChunkChangesCSVFileName, Packages);
+		ChunkDbPtr->LogSummary(IoStoreWriterResults);	
 	}
 
 	UE_LOG(LogIoStore, Display, TEXT(""));
@@ -5877,6 +6001,118 @@ int32 CreateTarget(const FIoStoreArguments& Arguments, const FIoStoreWriterSetti
 	return 0;
 }
 
+
+void FIoStoreChunkDatabase::WriteCSV(const FString& InOutputFileName, const TArray<FCookedPackage*>& InPackages)
+{
+	if (InOutputFileName.Len() == 0)
+	{
+		return;
+	}
+
+	FString ChunkNames[FIoStoreChunkDatabase::IoChunkTypeCount];
+	for (uint8 TypeIndex = 0; TypeIndex < FIoStoreChunkDatabase::IoChunkTypeCount; TypeIndex++)
+	{
+		ChunkNames[TypeIndex] = LexToString((EIoChunkType)TypeIndex);
+	}
+
+	//
+	// If we are using a reference cache, the assumption is that we are expecting a high hit rate
+	// on the chunks - this is supposed to be a current release, so theoretically we should only
+	// miss on changed or new data, so it's useful to know where the misses are in case something
+	// major happened that maybe shouldn't have.
+	//
+
+	TMap<FPackageId, FName> PackageIdToName;
+	for (const FCookedPackage* Package : InPackages)
+	{
+		PackageIdToName.Add(Package->GlobalPackageId, Package->PackageName);
+	}
+
+	auto FindPackageForChunk = [&PackageIdToName](const FIoChunkId& InChunkId)
+	{
+		// See CreateIoChunkId and CreatePackageDataChunkId
+		uint64 ChunkValue = *(uint64*)InChunkId.GetData();
+		FName* Result = PackageIdToName.Find(FPackageId::FromValue(ChunkValue));
+		if (Result)
+		{
+			return *Result;
+		}
+		return FName();
+	};
+
+	TUniquePtr<FArchive> ChangeListArchive;
+	ChangeListArchive.Reset(IFileManager::Get().CreateFileWriter(*InOutputFileName));
+	if (ChangeListArchive.IsValid() == false)
+	{
+		UE_LOG(LogIoStore, Warning, TEXT("Unable to open reference chunk change list file %s"), *InOutputFileName);
+		return;
+	}
+
+	ChangeListArchive->Logf(TEXT("Container,ChunkType,Class,ChunkId,ChunkHash,ChunkPackage"));
+
+	for (const TPair<FIoContainerId, TUniquePtr<FReaderChunks>>& ReaderIdPair : ChunkDatabase)
+	{
+		FReaderChunks* Reader = ReaderIdPair.Value.Get();
+
+		for (TPair<FIoChunkId, int32> Chunk : Reader->ChunkIds)
+		{
+			if (Reader->ChunkChanged[Chunk.Value].load(std::memory_order_relaxed))
+			{
+				FIoChunkId& Id = Chunk.Key;
+
+				EIoChunkType ChangedType = Id.GetChunkType();
+
+				if (ChangedType == EIoChunkType::ShaderCode ||
+					ChangedType == EIoChunkType::ShaderCodeLibrary)
+				{
+					// These don't have the package name in the chunk id.
+					ChangeListArchive->Logf(TEXT("%s,%s,CHANGED,%s,%s,<shader>"), *Reader->ContainerName, *ChunkNames[(uint8)ChangedType], *LexToString(Id), *Reader->ChangedChunkHashes[Chunk.Value].ToString());
+				}
+				else
+				{
+					ChangeListArchive->Logf(TEXT("%s,%s,CHANGED,%s,%s,%s"), *Reader->ContainerName, *ChunkNames[(uint8)ChangedType], *LexToString(Id), *Reader->ChangedChunkHashes[Chunk.Value].ToString(), *FindPackageForChunk(Id).ToString());
+				}
+			}
+		}
+
+		for (FRequestedChunkInfo& NewChunk : Reader->NewChunks)
+		{
+			if (NewChunk.Id.GetChunkType() == EIoChunkType::ShaderCode ||
+				NewChunk.Id.GetChunkType() == EIoChunkType::ShaderCodeLibrary)
+			{
+				// These don't have the package name in the chunk id.
+				ChangeListArchive->Logf(TEXT("%s,%s,NEW,%s,%s,<shader>"), *Reader->ContainerName, *ChunkNames[(uint8)NewChunk.Id.GetChunkType()], *LexToString(NewChunk.Id), *NewChunk.Hash.ToString());
+			}
+			else
+			{
+				ChangeListArchive->Logf(TEXT("%s,%s,NEW,%s,%s,%s"), *Reader->ContainerName, *ChunkNames[(uint8)NewChunk.Id.GetChunkType()], *LexToString(NewChunk.Id), *NewChunk.Hash.ToString(), *FindPackageForChunk(NewChunk.Id).ToString());
+			}
+		}
+	}
+
+	for (const TPair<FIoContainerId, TUniquePtr<FMissingContainerInfo>>& MissingContainerPair : MissingContainers)
+	{
+		FMissingContainerInfo* MissingContainer = MissingContainerPair.Value.Get();
+		for (const FRequestedChunkInfo& NoContainerChunk : MissingContainer->Chunks)
+		{
+			const FString& ChunkName = ChunkNames[(uint8)NoContainerChunk.Id.GetChunkType()];
+
+			if (NoContainerChunk.Id.GetChunkType() == EIoChunkType::ShaderCode ||
+				NoContainerChunk.Id.GetChunkType() == EIoChunkType::ShaderCodeLibrary)
+			{
+				// These don't have the package name in the chunk id.
+				ChangeListArchive->Logf(TEXT("%s,%s,NOCONTAINER,%s,%s,<shader>"),
+					*MissingContainer->ContainerName, *ChunkName, *LexToString(NoContainerChunk.Id), *NoContainerChunk.Hash.ToString());
+			}
+			else
+			{
+				ChangeListArchive->Logf(TEXT("%s,%s,NOCONTAINER,%s,%s,%"),
+					*MissingContainer->ContainerName, *ChunkName, *LexToString(NoContainerChunk.Id), *NoContainerChunk.Hash.ToString(), *FindPackageForChunk(NoContainerChunk.Id).ToString());
+			}
+		}
+	}
+}
+
 bool DumpIoStoreContainerInfo(const TCHAR* InContainerFilename, const FKeyChain& InKeyChain)
 {
 	TUniquePtr<FIoStoreReader> Reader = CreateIoStoreReader(InContainerFilename, InKeyChain);
@@ -9136,6 +9372,16 @@ bool ParseContainerGenerationArguments(FIoStoreArguments& Arguments, FIoStoreWri
 			UE_LOG(LogIoStore, Display, TEXT("Parsing reference container crypto keys from a crypto key cache file '%s'"), *CryptoKeysCacheFilename);
 			KeyChainUtilities::LoadKeyChainFromFile(CryptoKeysCacheFilename, Arguments.ReferenceChunkKeys);
 		}
+
+		if (FParse::Value(FCommandLine::Get(), TEXT("-ReferenceContainerChangesCSVFileName="), Arguments.ReferenceChunkChangesCSVFileName))
+		{
+			UE_LOG(LogIoStore, Display, TEXT("ReferenceCache will write the list of changed/new packages to %s"), *Arguments.ReferenceChunkChangesCSVFileName);
+		}
+
+		if (FParse::Value(FCommandLine::Get(), TEXT("-ReferenceContainerAdditionalPath="), Arguments.ReferenceChunkAdditionalContainersPath))
+		{
+			UE_LOG(LogIoStore, Display, TEXT("ReferenceCache will read additional containers from %s"), *Arguments.ReferenceChunkAdditionalContainersPath);
+		}
 	}
 
 	// By default, we use any hashes in the asset registry that exist in order to avoid reading and hashing
diff --git a/Engine/Source/Programs/AutomationTool/AutomationUtils/ProjectParams.cs b/Engine/Source/Programs/AutomationTool/AutomationUtils/ProjectParams.cs
--- a/Engine/Source/Programs/AutomationTool/AutomationUtils/ProjectParams.cs
+++ b/Engine/Source/Programs/AutomationTool/AutomationUtils/ProjectParams.cs
@@ -519,6 +519,8 @@ namespace AutomationTool
 			string CreateReleaseVersionBasePath = null,
 			string BasedOnReleaseVersionBasePath = null,
 			string ReferenceContainerGlobalFileName = null,
+			string ReferenceContainerAdditionalPath = null,
+			string ReferenceContainerChangesCSVFileName = null,
 			string ReferenceContainerCryptoKeys = null,
 			bool? GeneratePatch = null,
 			bool? AddPatchLevel = null,
@@ -702,6 +704,8 @@ namespace AutomationTool
 			this.CreateReleaseVersionBasePath = ParseParamValueIfNotSpecified(Command, CreateReleaseVersionBasePath, "createreleaseversionroot", String.Empty);
 			this.BasedOnReleaseVersionBasePath = ParseParamValueIfNotSpecified(Command, BasedOnReleaseVersionBasePath, "basedonreleaseversionroot", String.Empty);
 			this.ReferenceContainerGlobalFileName = ParseParamValueIfNotSpecified(Command, ReferenceContainerGlobalFileName, "ReferenceContainerGlobalFileName", String.Empty);
+			this.ReferenceContainerAdditionalPath = ParseParamValueIfNotSpecified(Command, ReferenceContainerAdditionalPath, "ReferenceContainerAdditionalPath", String.Empty);
+			this.ReferenceContainerChangesCSVFileName = ParseParamValueIfNotSpecified(Command, ReferenceContainerChangesCSVFileName, "ReferenceContainerChangesCSVFileName", String.Empty);
 			this.ReferenceContainerCryptoKeys = ParseParamValueIfNotSpecified(Command, ReferenceContainerCryptoKeys, "ReferenceContainerCryptoKeys", String.Empty); 
 			this.OriginalReleaseVersion = ParseParamValueIfNotSpecified(Command, OriginalReleaseVersion, "originalreleaseversion", String.Empty);
 			this.CreateReleaseVersion = ParseParamValueIfNotSpecified(Command, CreateReleaseVersion, "createreleaseversion", String.Empty);
@@ -1838,6 +1842,17 @@ namespace AutomationTool
 		/// </summary>
 		public string ReferenceContainerGlobalFileName;
 
+		/// <summary>
+		/// Stage: A directory to look for additional IoStore containers to use during staging (e.g. optional segment containers)
+		/// </summary>
+		public string ReferenceContainerAdditionalPath;
+
+		/// <summary>
+		/// Stage: If ReferenceContainerChangesCSVFileName is specified, this provides a file name to write the list of changed
+		/// packages/chunks to.
+		/// </summary>
+		public string ReferenceContainerChangesCSVFileName;
+
 		/// <summary>
 		/// Stage: Path to the crypto.json file to use for decrypting ReferenceContainerFlobalFileName, if needed.
 		/// </summary>
diff --git a/Engine/Source/Programs/AutomationTool/Scripts/BuildCookRun.Automation.cs b/Engine/Source/Programs/AutomationTool/Scripts/BuildCookRun.Automation.cs
--- a/Engine/Source/Programs/AutomationTool/Scripts/BuildCookRun.Automation.cs
+++ b/Engine/Source/Programs/AutomationTool/Scripts/BuildCookRun.Automation.cs
@@ -87,6 +87,18 @@ public class BuildCookRun : BuildCommand, IProjectParamsHelpers
 			Params.ReferenceContainerGlobalFileName = ReferenceContainerGlobalFileName;
 		}
 
+		var ReferenceContainerAdditionalPath = ParseParamValue("ReferenceContainerAdditionalPath");
+		if (!String.IsNullOrEmpty(ReferenceContainerAdditionalPath))
+		{
+			Params.ReferenceContainerAdditionalPath = ReferenceContainerAdditionalPath;
+		}
+
+		var ReferenceContainerChangesCSVFileName = ParseParamValue("ReferenceContainerChangesCSVFileName");
+		if (!String.IsNullOrEmpty(ReferenceContainerChangesCSVFileName))
+		{
+			Params.ReferenceContainerChangesCSVFileName = ReferenceContainerChangesCSVFileName;
+		}
+
 		var ReferenceContainerCryptoKeys = ParseParamValue("ReferenceContainerCryptoKeys");
 		if (!String.IsNullOrEmpty(ReferenceContainerCryptoKeys))
 		{
diff --git a/Engine/Source/Programs/AutomationTool/Scripts/CopyBuildToStagingDirectory.Automation.cs b/Engine/Source/Programs/AutomationTool/Scripts/CopyBuildToStagingDirectory.Automation.cs
--- a/Engine/Source/Programs/AutomationTool/Scripts/CopyBuildToStagingDirectory.Automation.cs
+++ b/Engine/Source/Programs/AutomationTool/Scripts/CopyBuildToStagingDirectory.Automation.cs
@@ -3728,10 +3728,18 @@ namespace AutomationScripts
 				if (Params.ReferenceContainerGlobalFileName.Length != 0)
 				{
 					AdditionalArgs += " -ReferenceContainerGlobalFileName=\"" + Params.ReferenceContainerGlobalFileName + "\"";
+					if (Params.ReferenceContainerAdditionalPath.Length != 0)
+					{
+						AdditionalArgs += " -ReferenceContainerAdditionalPath=\"" + Params.ReferenceContainerAdditionalPath + "\"";
+					}
 					if (Params.ReferenceContainerCryptoKeys.Length != 0)
 					{
 						AdditionalArgs += " -ReferenceContainerCryptoKeys=\"" + Params.ReferenceContainerCryptoKeys + "\"";
 					}
+					if (Params.ReferenceContainerChangesCSVFileName.Length != 0)
+					{
+						AdditionalArgs += " -ReferenceContainerChangesCSVFileName=\"" + Params.ReferenceContainerChangesCSVFileName + "\"";
+					}
 				}
 
 
@@ -4911,10 +4919,18 @@ namespace AutomationScripts
 			if (!string.IsNullOrEmpty(Params.ReferenceContainerGlobalFileName))
 			{
 				CommandLineBuilder.AppendFormat("-ReferenceContainerGlobalFileName={0}", Params.ReferenceContainerGlobalFileName);
+				if (!string.IsNullOrEmpty(Params.ReferenceContainerAdditionalPath))
+				{
+					CommandLineBuilder.AppendFormat("-ReferenceContainerAdditionalPath={0}", Params.ReferenceContainerAdditionalPath);
+				}
 				if (!string.IsNullOrEmpty(Params.ReferenceContainerCryptoKeys))
 				{
 					CommandLineBuilder.AppendFormat("-ReferenceContainerCryptoKeys={0}", Params.ReferenceContainerCryptoKeys);
 				}
+				if (!string.IsNullOrEmpty(Params.ReferenceContainerChangesCSVFileName))
+				{
+					CommandLineBuilder.AppendFormat("-ReferenceContainerChangesCSVFileName={0}", Params.ReferenceContainerChangesCSVFileName);
+				}
 			}
 
 			String CompleteCommandLine = CommandLineBuilder.ToString().Trim();
diff --git a/Engine/Source/Runtime/Core/Private/IO/IoStore.cpp b/Engine/Source/Runtime/Core/Private/IO/IoStore.cpp
--- a/Engine/Source/Runtime/Core/Private/IO/IoStore.cpp
+++ b/Engine/Source/Runtime/Core/Private/IO/IoStore.cpp
@@ -746,7 +746,7 @@ public:
 		ReferenceChunkDatabase = InReferenceChunkDatabase;
 		
 		// Add ourselves to the reference chunk db's list of possibles
-		ReferenceChunkDatabase->NotifyAddedToWriter(ContainerSettings.ContainerId);
+		ReferenceChunkDatabase->NotifyAddedToWriter(ContainerSettings.ContainerId, FPaths::GetBaseFilename(TocFilePath));
 	}
 	void SetHashDatabase(TSharedPtr<IIoStoreWriterHashDatabase> InHashDatabase, bool bInVerifyHashDatabase)
 	{
@@ -879,8 +879,7 @@ public:
 
 				if (ReferenceChunkDatabase.IsValid() && CompressionMethodForEntry(Entry) != NAME_None)
 				{
-					TPair<FIoContainerId, FIoChunkHash> ChunkKey(ContainerSettings.ContainerId, Entry->ChunkHash);
-					Entry->bLoadingFromReferenceDb = ReferenceChunkDatabase->ChunkExists(ChunkKey, Entry->ChunkId, Entry->NumChunkBlocksFromRefDb);
+					Entry->bLoadingFromReferenceDb = ReferenceChunkDatabase->ChunkExists(ContainerSettings.ContainerId, Entry->ChunkHash, Entry->ChunkId, Entry->NumChunkBlocksFromRefDb);
 					Entry->bCouldBeFromReferenceDb = true;
 				}
 				return;
@@ -915,8 +914,7 @@ public:
 
 			if (ReferenceChunkDatabase.IsValid() && CompressionMethodForEntry(Entry) != NAME_None)
 			{
-				TPair<FIoContainerId, FIoChunkHash> ChunkKey(ContainerSettings.ContainerId, Entry->ChunkHash);
-				Entry->bLoadingFromReferenceDb = ReferenceChunkDatabase->ChunkExists(ChunkKey, Entry->ChunkId, Entry->NumChunkBlocksFromRefDb);
+				Entry->bLoadingFromReferenceDb = ReferenceChunkDatabase->ChunkExists(ContainerSettings.ContainerId, Entry->ChunkHash, Entry->ChunkId, Entry->NumChunkBlocksFromRefDb);
 				Entry->bCouldBeFromReferenceDb = true;
 			}
 
@@ -1592,12 +1590,10 @@ private:
 				// Everything else in a block gets filled out from the refdb.
 			}
 
-			TPair<FIoContainerId, FIoChunkHash> ChunkKey(ContainerSettings.ContainerId, Entry->ChunkHash);
-
 			// Valid chunks must create the same decompressed bits, but can have different compressed bits.
 			// Since we are on a lightweight dispatch thread, the actual read is async, as is the processing
 			// of the results.
-			bool bChunkExists = ReferenceChunkDatabase->RetrieveChunk(ChunkKey, [this, Entry](TIoStatusOr<FIoStoreCompressedReadResult> InReadResult)
+			bool bChunkExists = ReferenceChunkDatabase->RetrieveChunk(ContainerSettings.ContainerId, Entry->ChunkHash, Entry->ChunkId, [this, Entry](TIoStatusOr<FIoStoreCompressedReadResult> InReadResult)
 			{
 
 				// If we fail here, in order to recover we effectively need to re-kick this chunk's
@@ -1835,6 +1831,16 @@ private:
 		const int32* FindExistingIndex = Toc.GetTocEntryIndex(Entry->ChunkId);
 		if (FindExistingIndex)
 		{
+			// afaict this should never happen so add a warning. If there's a legit reason for it
+			// we can pull this back out. If would violate some assumptions in the reference chunk
+			// database if we DO hit this, however...
+			UE_LOG(LogIoStore, Warning, TEXT("ChunkId was added twice in container %s, %s, file %s hash %s vs %s"), 
+				*FPaths::GetBaseFilename(TocFilePath),
+				*LexToString(Entry->ChunkId), *Entry->Options.FileName,
+				*Toc.GetTocResource().ChunkMetas[*FindExistingIndex].ChunkHash.ToString(),
+				*Entry->ChunkHash.ToString()
+				);
+
 			checkf(Toc.GetTocResource().ChunkMetas[*FindExistingIndex].ChunkHash == Entry->ChunkHash, TEXT("Chunk id has already been added with different content"));
 			for (FChunkBlock& ChunkBlock : Entry->ChunkBlocks)
 			{
@@ -2434,6 +2440,16 @@ public:
 		return Toc.GetTocResource().Header.EncryptionKeyGuid;
 	}
 
+	FString GetContainerName() const
+	{
+		return FPaths::GetBaseFilename(ContainerPath);
+	}
+
+	int32 GetChunkCount() const 
+	{
+		return Toc.GetTocResource().ChunkIds.Num();
+	}
+
 	void EnumerateChunks(TFunction<bool(FIoStoreTocChunkInfo&&)>&& Callback) const
 	{
 		const FIoStoreTocResource& TocResource = Toc.GetTocResource();
@@ -3077,6 +3093,16 @@ FGuid FIoStoreReader::GetEncryptionKeyGuid() const
 	return Impl->GetEncryptionKeyGuid();
 }
 
+int32 FIoStoreReader::GetChunkCount() const
+{
+	return Impl->GetChunkCount();
+}
+
+FString FIoStoreReader::GetContainerName() const
+{
+	return Impl->GetContainerName();
+}
+
 void FIoStoreReader::EnumerateChunks(TFunction<bool(FIoStoreTocChunkInfo&&)>&& Callback) const
 {
 	Impl->EnumerateChunks(MoveTemp(Callback));
diff --git a/Engine/Source/Runtime/Core/Public/IO/IoDispatcher.h b/Engine/Source/Runtime/Core/Public/IO/IoDispatcher.h
--- a/Engine/Source/Runtime/Core/Public/IO/IoDispatcher.h
+++ b/Engine/Source/Runtime/Core/Public/IO/IoDispatcher.h
@@ -701,13 +701,14 @@ public:
 	* Chunks provided *MUST* decompress to bits that hash to the exact value provided in InChunkKey (i.e. be exactly the same bits),
 	* and also be the same number of blocks (i.e. same CompressionBlockSize)
 	*/
-	virtual bool RetrieveChunk(const TPair<FIoContainerId, FIoChunkHash>& InChunkKey, TUniqueFunction<void(TIoStatusOr<FIoStoreCompressedReadResult>)> InCompletionCallback) = 0;
+	virtual bool RetrieveChunk(const FIoContainerId& InContainerId, const FIoChunkHash& InChunkHash, const FIoChunkId& InChunkId, TUniqueFunction<void(TIoStatusOr<FIoStoreCompressedReadResult>)> InCompletionCallback) = 0;
 
 	/* 
 	* Quick synchronous existence check that returns the number of blocks for the chunk. This is used to set up
-	* the necessary structures without needing to read the source data for the chunk.
+	* the necessary structures without needing to read the source data for the chunk. This might be called from
+	* multiple threads as it has to happen after we have the source hash computed.
 	*/
-	virtual bool ChunkExists(const TPair<FIoContainerId, FIoChunkHash>& InChunkKey, const FIoChunkId& InChunkId, uint32& OutNumChunkBlocks) = 0;
+	virtual bool ChunkExists(const FIoContainerId& InContainerId, const FIoChunkHash& InChunkHash, const FIoChunkId& InChunkId, uint32& OutNumChunkBlocks) = 0;
 
 	/*
 	* Returns the compression block size that was used to break up the IoChunks in the source containers. If this is different than what we want, 
@@ -718,7 +719,7 @@ public:
 	/*
 	* Called by an iostore writer implementation to notify the ref cache it's been added
 	*/
-	virtual void NotifyAddedToWriter(const FIoContainerId& InContainerId) = 0;
+	virtual void NotifyAddedToWriter(const FIoContainerId& InContainerId, const FString& InContainerName) = 0;
 };
 
 /**
@@ -764,6 +765,8 @@ public:
 	CORE_API uint32 GetVersion() const;
 	CORE_API EIoContainerFlags GetContainerFlags() const;
 	CORE_API FGuid GetEncryptionKeyGuid() const;
+	CORE_API int32 GetChunkCount() const;
+	CORE_API FString GetContainerName() const; // The container name is the base filename of ContainerPath, e.g. "global".
 
 	CORE_API void EnumerateChunks(TFunction<bool(FIoStoreTocChunkInfo&&)>&& Callback) const;
 	CORE_API TIoStatusOr<FIoStoreTocChunkInfo> GetChunkInfo(const FIoChunkId& Chunk) const;
