Removed IAS upload from UnrealPak. It is now exposed via the IasTool program.

#rnx

diff --git a/Engine/Source/Developer/IoStoreUtilities/IoStoreUtilities.Build.cs b/Engine/Source/Developer/IoStoreUtilities/IoStoreUtilities.Build.cs
--- a/Engine/Source/Developer/IoStoreUtilities/IoStoreUtilities.Build.cs
+++ b/Engine/Source/Developer/IoStoreUtilities/IoStoreUtilities.Build.cs
@@ -30,6 +30,5 @@ public class IoStoreUtilities : ModuleRules
         PrivateDependencyModuleNames.Add("RSA");
         PrivateDependencyModuleNames.Add("DeveloperToolSettings");
         PrivateDependencyModuleNames.Add("SandboxFile");
-        PrivateDependencyModuleNames.Add("IoStoreOnDemand");
 	}
 }
diff --git a/Engine/Source/Developer/IoStoreUtilities/Private/IoStoreUtilities.cpp b/Engine/Source/Developer/IoStoreUtilities/Private/IoStoreUtilities.cpp
--- a/Engine/Source/Developer/IoStoreUtilities/Private/IoStoreUtilities.cpp
+++ b/Engine/Source/Developer/IoStoreUtilities/Private/IoStoreUtilities.cpp
@@ -37,7 +37,6 @@
 #include "Serialization/ArrayReader.h"
 #include "Serialization/ArrayWriter.h"
 #include "Settings/ProjectPackagingSettings.h" // for EAssetRegistryWritebackMethod
-#include "IO/IoStoreOnDemand.h"
 #include "IO/PackageStore.h"
 #include "UObject/Class.h"
 #include "UObject/NameBatchSerialization.h"
@@ -102,7 +101,6 @@ static const uint64 DefaultCompressionBlockAlignment = 64 << 10;
 static const uint64 DefaultMemoryMappingAlignment = 16 << 10;
 
 static TUniquePtr<FIoStoreReader> CreateIoStoreReader(const TCHAR* Path, const FKeyChain& KeyChain);
-bool UploadIoStoreContainerFiles(const UE::IO::IAS::FIoStoreUploadParams& UploadParams, TConstArrayView<FString> ContainerFiles, const FKeyChain& KeyChain);
 
 /*
 * Provides access to previously compressed chunks to the iostore writer, allowing
@@ -1223,7 +1221,6 @@ struct FIoStoreArguments
 	bool bCreateDirectoryIndex = true;
 	bool bClusterByOrderFilePriority = false;
 	bool bFileRegions = false;
-	bool bUpload = false;
 	EAssetRegistryWritebackMethod WriteBackMetadataToAssetRegistry = EAssetRegistryWritebackMethod::Disabled;
 	bool bWritePluginSizeSummaryJsons = false; // Only valid if WriteBackMetadataToAssetRegistry != Disabled.
 
@@ -5983,21 +5980,6 @@ int32 CreateTarget(const FIoStoreArguments& Arguments, const FIoStoreWriterSetti
 		FTaskGraphInterface::Get().WaitUntilTaskCompletes(WriteCsvFileTask);
 	}
 
-	if (Arguments.bUpload && OnDemandContainers.IsEmpty() == false)
-	{
-		TIoStatusOr<UE::IO::IAS::FIoStoreUploadParams> UploadParams = UE::IO::IAS::FIoStoreUploadParams::Parse(FCommandLine::Get());
-		if (UploadParams.IsOk() == false)
-		{
-			UE_LOG(LogIoStore, Warning, TEXT("Skipping upload of container file(s), reason '%s'"), *UploadParams.Status().ToString());
-			return 0;
-		}
-
-		if (UploadIoStoreContainerFiles(UploadParams.ConsumeValueOrDie(), OnDemandContainers, Arguments.KeyChain) == false)
-		{
-			return -1;
-		}
-	}
-
 	return 0;
 }
 
@@ -9751,7 +9733,6 @@ int32 CreateIoStoreContainerFiles(const TCHAR* CmdLine)
 		
 		Arguments.DLCName = FPaths::GetBaseFilename(*Arguments.DLCPluginPath);
 		Arguments.bRemapPluginContentToGame = FParse::Param(FCommandLine::Get(), TEXT("RemapPluginContentToGame"));
-		Arguments.bUpload = FParse::Param(FCommandLine::Get(), TEXT("Upload"));
 
 		UE_LOG(LogIoStore, Display, TEXT("DLC: '%s'"), *Arguments.DLCPluginPath);
 		UE_LOG(LogIoStore, Display, TEXT("Remapping plugin content to game: '%s'"), Arguments.bRemapPluginContentToGame ? TEXT("True") : TEXT("False"));
@@ -9798,7 +9779,6 @@ int32 CreateIoStoreContainerFiles(const TCHAR* CmdLine)
 	else if (FParse::Value(FCommandLine::Get(), TEXT("CreateGlobalContainer="), Arguments.GlobalContainerPath))
 	{
 		Arguments.GlobalContainerPath = FPaths::ChangeExtension(Arguments.GlobalContainerPath, TEXT(""));
-		Arguments.bUpload = FParse::Param(FCommandLine::Get(), TEXT("Upload"));
 
 		if (!ParseContainerGenerationArguments(Arguments, WriterSettings))
 		{
@@ -9918,74 +9898,3 @@ int32 CreateIoStoreContainerFiles(const TCHAR* CmdLine)
 
 	return CreateTarget(Arguments, WriterSettings);
 }
-
-bool UploadIoStoreContainerFiles(const UE::IO::IAS::FIoStoreUploadParams& UploadParams, TConstArrayView<FString> ContainerFiles, const FKeyChain& KeyChain)
-{
-	TIoStatusOr<UE::IO::IAS::FIoStoreUploadResult> Result = UE::IO::IAS::UploadContainerFiles(UploadParams, ContainerFiles, KeyChain);
-	if (Result.IsOk() == false)
-	{
-		UE_LOG(LogIoStore, Error, TEXT("Failed to upload container file(s), reason '%s'"), *Result.Status().ToString());
-		return false;
-	}
-
-	UE::IO::IAS::FIoStoreUploadResult UploadResult = Result.ConsumeValueOrDie();
-
-	return Result.IsOk();
-}
-
-bool UploadIoStoreContainerFiles(const TCHAR* ContainerPathOrWildcard)
-{
-	check(ContainerPathOrWildcard);
-
-	FKeyChain KeyChain;
-	LoadKeyChain(FCommandLine::Get(), KeyChain);
-
-	TIoStatusOr<UE::IO::IAS::FIoStoreUploadParams> UploadParams = UE::IO::IAS::FIoStoreUploadParams::Parse(FCommandLine::Get());
-	if (UploadParams.IsOk() == false)
-	{
-		UE_LOG(LogIoStore, Error, TEXT("Failed to upload container file(s), reason '%s'"), *UploadParams.Status().ToString());
-		return false;
-	}
-
-	TArray<FString> ContainerFiles;
-	{
-		if (IFileManager::Get().FileExists(ContainerPathOrWildcard))
-		{
-			ContainerFiles.Add(ContainerPathOrWildcard);
-		}
-		else if (IFileManager::Get().DirectoryExists(ContainerPathOrWildcard))
-		{
-			FString Directory = ContainerPathOrWildcard;
-			FPaths::NormalizeDirectoryName(Directory);
-
-			TArray<FString> FoundContainerFiles;
-			IFileManager::Get().FindFiles(FoundContainerFiles, *(Directory / TEXT("*.utoc")), true, false);
-
-			for (const FString& Filename : FoundContainerFiles)
-			{
-				ContainerFiles.Emplace(Directory / Filename);
-			}
-		}
-		else
-		{
-			FString Directory = FPaths::GetPath(ContainerPathOrWildcard);
-			FPaths::NormalizeDirectoryName(Directory);
-
-			TArray<FString> FoundContainerFiles;
-			IFileManager::Get().FindFiles(FoundContainerFiles, ContainerPathOrWildcard, true, false);
-
-			for (const FString& Filename : FoundContainerFiles)
-			{
-				ContainerFiles.Emplace(Directory / Filename);
-			}
-		}
-	}
-
-	if (ContainerFiles.IsEmpty())
-	{
-		UE_LOG(LogIoStore, Error, TEXT("Failed to find container file(s) '%s'"), ContainerPathOrWildcard);
-		return false;
-	}
-
-	return UploadIoStoreContainerFiles(UploadParams.ConsumeValueOrDie(), ContainerFiles, KeyChain);
-}
diff --git a/Engine/Source/Developer/IoStoreUtilities/Public/IoStoreUtilities.h b/Engine/Source/Developer/IoStoreUtilities/Public/IoStoreUtilities.h
--- a/Engine/Source/Developer/IoStoreUtilities/Public/IoStoreUtilities.h
+++ b/Engine/Source/Developer/IoStoreUtilities/Public/IoStoreUtilities.h
@@ -58,6 +58,4 @@ UE_API bool ProcessFilesFromIoStoreContainer(
 
 UE_API bool SignIoStoreContainer(const TCHAR* InContainerFilename, const FRSAKeyHandle InSigningKey);
 
-UE_API bool UploadIoStoreContainerFiles(const TCHAR* ContainerPathOrWildcard); 
-
 #undef UE_API
diff --git a/Engine/Source/Developer/PakFileUtilities/Private/PakFileUtilities.cpp b/Engine/Source/Developer/PakFileUtilities/Private/PakFileUtilities.cpp
--- a/Engine/Source/Developer/PakFileUtilities/Private/PakFileUtilities.cpp
+++ b/Engine/Source/Developer/PakFileUtilities/Private/PakFileUtilities.cpp
@@ -5463,14 +5463,6 @@ bool ExecuteUnrealPak(const TCHAR* CmdLine)
 		{
 			return ListIoStoreContainerBulkData(CmdLine);
 		}
-
-		// IAS commands
-		{
-			if (FParse::Value(CmdLine, TEXT("-Upload="), IoStoreArg))
-			{
-				return UploadIoStoreContainerFiles(*IoStoreArg) == 0;
-			}
-		}
 	}
 
 	// Parse all the non-option arguments from the command line
diff --git a/Engine/Source/Runtime/Experimental/IoStoreOnDemand/Public/IO/IoStoreOnDemand.h b/Engine/Source/Runtime/Experimental/IoStoreOnDemand/Public/IO/IoStoreOnDemand.h
--- a/Engine/Source/Runtime/Experimental/IoStoreOnDemand/Public/IO/IoStoreOnDemand.h
+++ b/Engine/Source/Runtime/Experimental/IoStoreOnDemand/Public/IO/IoStoreOnDemand.h
@@ -197,7 +197,7 @@ struct FIoStoreUploadParams
 	/** Where the .iochunktoc file should be written out. */
 	FString TocOutputDir;
 
-	UE_API static TIoStatusOr<FIoStoreUploadParams> Parse(const TCHAR* CommandLine);
+	static TIoStatusOr<FIoStoreUploadParams> Parse(const TCHAR* CommandLine);
 	FIoStatus Validate() const;
 };
 
@@ -216,7 +216,7 @@ struct FIoStoreUploadResult
 	uint64 TocSize = 0;
 };
 
-UE_API TIoStatusOr<FIoStoreUploadResult> UploadContainerFiles(
+TIoStatusOr<FIoStoreUploadResult> UploadContainerFiles(
 	const FIoStoreUploadParams& UploadParams,
 	TConstArrayView<FString> ContainerFiles,
 	const FKeyChain& KeyChain);
