Since package names are usually FNames, added a helper to avoid any extraneous allocations when grabbing the root.

diff --git a/Engine/Plugins/Editor/AssetManagerEditor/Source/AssetManagerEditor/Private/AssetManagerEditorModule.cpp b/Engine/Plugins/Editor/AssetManagerEditor/Source/AssetManagerEditor/Private/AssetManagerEditorModule.cpp
--- a/Engine/Plugins/Editor/AssetManagerEditor/Source/AssetManagerEditor/Private/AssetManagerEditorModule.cpp
+++ b/Engine/Plugins/Editor/AssetManagerEditor/Source/AssetManagerEditor/Private/AssetManagerEditorModule.cpp
@@ -1346,7 +1346,7 @@ bool FAssetManagerEditorModule::GetStringValueForCustomColumn(const FAssetData&
 	}
 	else if (ColumnName == PluginName)
 	{
-		OutValue = FPackageName::SplitPackageNameRoot(AssetData.PackageName.ToString(), nullptr);
+		OutValue = FPackageName::SplitPackageNameRoot(AssetData.PackageName, nullptr);
 		return OutValue.Len() > 0;
 	}
 	else
diff --git a/Engine/Source/Developer/IoStoreUtilities/Private/IoStoreUtilities.cpp b/Engine/Source/Developer/IoStoreUtilities/Private/IoStoreUtilities.cpp
--- a/Engine/Source/Developer/IoStoreUtilities/Private/IoStoreUtilities.cpp
+++ b/Engine/Source/Developer/IoStoreUtilities/Private/IoStoreUtilities.cpp
@@ -3867,9 +3867,8 @@ static void InsertShadersInPluginHierarchy(UE::Cook::FCookMetadataState& InCookM
 		TSet<FString> ShaderPlugins;
 		for (FName PackageName : ShaderChunkInfo.Value.ReferencedByPackages)
 		{
-			FString PackageNameStr = PackageName.ToString();
-			FStringView Plugin = FPackageName::SplitPackageNameRoot(PackageNameStr, nullptr);
-			ShaderPlugins.Add(FString(Plugin));
+			FString Plugin = FPackageName::SplitPackageNameRoot(PackageName, nullptr);
+			ShaderPlugins.Add(MoveTemp(Plugin));
 		}
 
 		if (ShaderPlugins.Num() == 1)
@@ -4149,9 +4148,8 @@ static void UpdatePluginMetadataAndWriteJsons(
 					SinglePluginShaderSize += ShaderChunkInfo.Value.CompressedSize;
 				}
 
-				FString PackageNameStr = ShaderChunkInfo.Value.ReferencedByPackages[0].ToString();
-				FStringView Plugin = FPackageName::SplitPackageNameRoot(PackageNameStr, nullptr);
-				PackageName.Append(Plugin);
+				FString Plugin = FPackageName::SplitPackageNameRoot(ShaderChunkInfo.Value.ReferencedByPackages[0], nullptr);
+				PackageName.Append(MoveTemp(Plugin));
 			}
 
 			FPluginGraphEntry** PluginEntryPtr = PluginGraph.NameToPlugin.Find(PackageName.ToString());
@@ -4251,10 +4249,7 @@ static void UpdatePluginMetadataAndWriteJsons(
 
 		// Assign the size to the package's plugin.
 		{
-			TStringBuilder<FName::StringBufferSize> PackageNameStr(InPlace, AssetPackage.Key);
-			FStringView PackageName(PackageNameStr);
-
-			FStringView PluginName = FPackageName::SplitPackageNameRoot(PackageName, nullptr);
+			FString PluginName = FPackageName::SplitPackageNameRoot(AssetPackage.Key, nullptr);
 			FPluginGraphEntry** PluginEntryPtr = PluginGraph.NameToPlugin.Find(PluginName);
 			if (PluginEntryPtr)
 			{
@@ -4319,12 +4314,11 @@ static void UpdatePluginMetadataAndWriteJsons(
 			}
 			else
 			{
-				FString AllocatedPluginName(PluginName);
 				bool bAlreadyLogged = false;
-				LoggedPluginNames.Add(MoveTemp(AllocatedPluginName), &bAlreadyLogged);
+				LoggedPluginNames.Add(MoveTemp(PluginName), &bAlreadyLogged);
 				if (bAlreadyLogged == false)
 				{
-					UE_LOG(LogIoStore, Display, TEXT("Plugin for package not found: %s (%.*s)"), PackageNameStr.GetData(), PluginName.Len(), PluginName.GetData());
+					UE_LOG(LogIoStore, Display, TEXT("Plugin for package not found: %s"), *AssetPackage.Key.ToString());
 				}
 			}
 		}
diff --git a/Engine/Source/Runtime/CoreUObject/Private/Misc/PackageName.cpp b/Engine/Source/Runtime/CoreUObject/Private/Misc/PackageName.cpp
--- a/Engine/Source/Runtime/CoreUObject/Private/Misc/PackageName.cpp
+++ b/Engine/Source/Runtime/CoreUObject/Private/Misc/PackageName.cpp
@@ -2606,6 +2606,21 @@ FString FPackageName::ExportTextPathToObjectPath(const TCHAR* InExportTextPath)
 	return ExportTextPathToObjectPath(FString(InExportTextPath));
 }
 
+FString FPackageName::SplitPackageNameRoot(FName InPackageName, FString* OutRelativePath)
+{
+	TStringBuilder<FName::StringBufferSize> PackageNameStr(InPlace, InPackageName);
+	FStringView PackageName(PackageNameStr);
+
+	FStringView RelativePath;
+	FStringView Root = SplitPackageNameRoot(PackageName, &RelativePath);
+
+	if (OutRelativePath)
+	{
+		*OutRelativePath = FString(RelativePath);
+	}
+	return FString(Root);
+}
+
 FStringView FPackageName::SplitPackageNameRoot(FStringView InPackageName, FStringView* OutRelativePath)
 {
 	if (!InPackageName.StartsWith(TEXT("/")))
diff --git a/Engine/Source/Runtime/CoreUObject/Public/Misc/PackageName.h b/Engine/Source/Runtime/CoreUObject/Public/Misc/PackageName.h
--- a/Engine/Source/Runtime/CoreUObject/Public/Misc/PackageName.h
+++ b/Engine/Source/Runtime/CoreUObject/Public/Misc/PackageName.h
@@ -670,6 +670,7 @@ public:
 	 * "/PackageRoot/Path/Leaf" -> (return "PackageRoot"; OutRelativePath = "Path/Leaf")
 	 */
 	static COREUOBJECT_API FStringView SplitPackageNameRoot(FStringView InPackageName, FStringView* OutRelativePath);
+	static COREUOBJECT_API FString SplitPackageNameRoot(FName InPackageName, FString* OutRelativePath);
 	
 	/** 
 	 * Returns the name of the package referred to by the specified object path
