Removing AssetRegistryWritebackAfterStage. This has evolved such that it needs to be run during stage and is not really useful afterwards. Was deprecated in 5.4 after being release in I think 5.3, highly unlikely anyone was using it.

diff --git a/Engine/Plugins/Editor/AssetRegistryExport/Source/Private/AssetRegistryExportCommandlet.cpp b/Engine/Plugins/Editor/AssetRegistryExport/Source/Private/AssetRegistryExportCommandlet.cpp
--- a/Engine/Plugins/Editor/AssetRegistryExport/Source/Private/AssetRegistryExportCommandlet.cpp
+++ b/Engine/Plugins/Editor/AssetRegistryExport/Source/Private/AssetRegistryExportCommandlet.cpp
@@ -872,7 +872,6 @@ int32 UAssetRegistryExportCommandlet::Main(const FString& CmdLineParams)
 	{
 		UE_LOG(LogAssetRegistryExport, Display, TEXT("Asset registry didn't have size metadata written back, CompressedSize DB column will be NULL."));
 		UE_LOG(LogAssetRegistryExport, Display, TEXT("Metadata can be written back via ProjectSettings/Packaging/WriteBackMetadataToAssetRegistry"));
-		UE_LOG(LogAssetRegistryExport, Display, TEXT("or after staging with iostore -AssetRegistryWriteback."));
 	}
 
 	if (bCSV == false)
diff --git a/Engine/Source/Developer/IoStoreUtilities/Private/IoStoreUtilities.cpp b/Engine/Source/Developer/IoStoreUtilities/Private/IoStoreUtilities.cpp
--- a/Engine/Source/Developer/IoStoreUtilities/Private/IoStoreUtilities.cpp
+++ b/Engine/Source/Developer/IoStoreUtilities/Private/IoStoreUtilities.cpp
@@ -4741,51 +4741,6 @@ static bool SaveAssetRegistry(const FString& InAssetRegistryFileName, FAssetRegi
 	return true;
 }
 
-static int32 DoAssetRegistryWritebackAfterStage(const FString& InAssetRegistryFileName, FString&& InContainerDirectory, const FKeyChain& InKeyChain)
-{
-	// This version called after the containers are already created, when you
-	// have a bunch of containers on disk and you want to add chunk info back to
-	// an asset registry.
-
-	FAssetRegistryState AssetRegistry;
-	if (LoadAssetRegistry(InAssetRegistryFileName, AssetRegistry) == 0)
-	{
-		UE_LOG(LogIoStore, Error, TEXT("Unabled to open source asset registry: %s"), *InAssetRegistryFileName);
-		return 1;
-	}
-
-	// Grab all containers in the directory
-	FPaths::NormalizeDirectoryName(InContainerDirectory);
-	TArray<FString> FoundContainerFiles;
-	IFileManager::Get().FindFiles(FoundContainerFiles, *(InContainerDirectory / TEXT("*.utoc")), true, false);
-	
-	uint64 TotalCompressedSize = 0;
-	
-	// Grab all the package infos.
-	TMap<FPackageId, TArray<FIoStoreChunkSource, TInlineAllocator<2>>> PackageToChunks;
-	for (const FString& Filename : FoundContainerFiles)
-	{
-		TUniquePtr<FIoStoreReader> Reader = CreateIoStoreReader(*(InContainerDirectory / Filename), InKeyChain);
-		if (Reader.IsValid() == false)
-		{
-			return 1; // already logged.
-		}
-		
-		Reader->EnumerateChunks([&](const FIoStoreTocChunkInfo& ChunkInfo)
-		{
-			FPackageId PackageId = FPackageId::FromValue(*(int64*)(ChunkInfo.Id.GetData()));
-			// (Deployment can't be ascertained after staging - plugin jsons are not generated.)
-			PackageToChunks.FindOrAdd(PackageId).Add({ChunkInfo, UE::Cook::EPluginSizeTypes::COUNT});
-			TotalCompressedSize += ChunkInfo.CompressedSize;
-			return true;
-		});
-	}
-
-	AddChunkInfoToAssetRegistry(MoveTemp(PackageToChunks), AssetRegistry, nullptr, 0, 0, TotalCompressedSize);
-
-	return SaveAssetRegistry(InAssetRegistryFileName, AssetRegistry, nullptr) ? 0 : 1;
-}
-
 enum class ECookMetadataFiles
 {
 	None = 0,
@@ -9545,22 +9500,6 @@ int32 CreateIoStoreContainerFiles(const TCHAR* CmdLine)
 
 		return ListContainer(Arguments.KeyChain, ContainerPathOrWildcard, Arguments.CsvPath);
 	}
-	else if (FParse::Value(FCommandLine::Get(), TEXT("AssetRegistryWriteback="), ArgumentValue))
-	{
-		//
-		// Opens a given directory of containers and a given asset registry, and adds chunk size information
-		// for an asset's package to its asset tags in the asset registry. This can also be done during the staging
-		// process with -WriteBackMetadataToAssetRegistry (below).
-		//
-		FString AssetRegistryFileName = MoveTemp(ArgumentValue);
-		FString PathToContainers;
-		if (!FParse::Value(FCommandLine::Get(), TEXT("ContainerDirectory="), PathToContainers))
-		{
-			UE_LOG(LogIoStore, Error, TEXT("Asset registry writeback requires -ContainerDirectory=Path/To/Containers"));
-		}
-		UE_LOG(LogIoStore, Warning, TEXT("AssetRegistryWriteback after stage is deprecated and will be removed in 5.5. Use writeback during stage via project packaging settings."));
-		return DoAssetRegistryWritebackAfterStage(AssetRegistryFileName, MoveTemp(PathToContainers), Arguments.KeyChain);
-	}
 	else if (FParse::Value(FCommandLine::Get(), TEXT("Describe="), ArgumentValue))
 	{
 		FString ContainerPathOrWildcard = ArgumentValue;
diff --git a/Engine/Source/Editor/UnrealEd/Private/Commandlets/AssetSizeQueryCommandlet.cpp b/Engine/Source/Editor/UnrealEd/Private/Commandlets/AssetSizeQueryCommandlet.cpp
--- a/Engine/Source/Editor/UnrealEd/Private/Commandlets/AssetSizeQueryCommandlet.cpp
+++ b/Engine/Source/Editor/UnrealEd/Private/Commandlets/AssetSizeQueryCommandlet.cpp
@@ -50,7 +50,7 @@ int32 UAssetSizeQueryCommandlet::Main(const FString& FullCommandLine)
 		UE_LOG(LogAssetSize, Display, TEXT("Params:"));
 		UE_LOG(LogAssetSize, Display, TEXT("    -AssetRegistry=path                     Provides the path to the Development asset registry. This"));
 		UE_LOG(LogAssetSize, Display, TEXT("                                            asset registry must have staging size metadata via ProjectSettings->"));
-		UE_LOG(LogAssetSize, Display, TEXT("                                            Packaging->WriteBackMetadataToAssetRegistry or iostore -AssetRegistryWriteback."));
+		UE_LOG(LogAssetSize, Display, TEXT("                                            Packaging->WriteBackMetadataToAssetRegistry."));
 		UE_LOG(LogAssetSize, Display, TEXT("    -Filter=wildcard                        (optional) Filteres the list of assets using a wildcard match."));
 		UE_LOG(LogAssetSize, Display, TEXT("    -Show=#                                 (optional) Shows only the top # classes, sorted on size (0 is all, default 10)."));
 		UE_LOG(LogAssetSize, Display, TEXT("    -CSV=path                               (optional) Output the filtered per class infomation to the given CSV file."));
@@ -181,8 +181,7 @@ int32 UAssetSizeQueryCommandlet::Main(const FString& FullCommandLine)
 	if (ImportantAssetCount == 0)
 	{
 		UE_LOG(LogAssetSize, Display, TEXT("No assets with size information found - staging metadata needs to be added to the asset registry"));
-		UE_LOG(LogAssetSize, Display, TEXT("via ProjectSettings->Packaging->WriteBackMetadataToAssetRegistry or using iostore "));
-		UE_LOG(LogAssetSize, Display, TEXT("-AssetRegistryWriteback."));
+		UE_LOG(LogAssetSize, Display, TEXT("via ProjectSettings->Packaging->WriteBackMetadataToAssetRegistry."));
 		return 1;
 	}
 
