Prevent partial initialization of the reference chunk database on failures.

diff --git a/Engine/Source/Developer/IoStoreUtilities/Private/IoStoreUtilities.cpp b/Engine/Source/Developer/IoStoreUtilities/Private/IoStoreUtilities.cpp
--- a/Engine/Source/Developer/IoStoreUtilities/Private/IoStoreUtilities.cpp
+++ b/Engine/Source/Developer/IoStoreUtilities/Private/IoStoreUtilities.cpp
@@ -173,6 +173,7 @@ public:
 	int64 FulfillBytesPerChunk[(int8)EIoChunkType::MAX] = {};
 
 	uint32 CompressionBlockSize = 0;
+	bool bValid = false;
 
 	bool Init(FString InGlobalContainerFileName, FString InAdditionalContainersPath, const FKeyChain& InDecryptionKeychain)
 	{
@@ -286,6 +287,7 @@ public:
 		}
 
 		UE_LOG(LogIoStore, Display, TEXT("Reference Chunk loaded %d containers and %s chunks, in %.1f seconds"), Readers.Num(), *FText::AsNumber(IoChunkCount).ToString(), FPlatformTime::Seconds() - StartTime);
+		bValid = true;
 		return true;
 	}
 
@@ -309,6 +311,10 @@ public:
 	// This can be called from any thread, though in the presence of existing hashes it's single threaded.
 	virtual bool ChunkExists(const FIoContainerId& InContainerId, const FIoChunkHash& InChunkHash, const FIoChunkId& InChunkId, uint32& OutNumChunkBlocks)
 	{
+		if (!bValid)
+		{
+			return false;
+		}
 		RequestCount.fetch_add(1, std::memory_order_relaxed);
 
 		uint8 ChunkType = (uint8)InChunkId.GetChunkType();
@@ -387,6 +393,11 @@ public:
 	// the iostore begindispatch thread (i.e. is not async)
 	virtual bool RetrieveChunk(const FIoContainerId& InContainerId, const FIoChunkHash& InChunkHash, const FIoChunkId& InChunkId, TUniqueFunction<void(TIoStatusOr<FIoStoreCompressedReadResult>)> InCompleteCallback)
 	{
+		if (!bValid)
+		{
+			return false;
+		}
+
 		TUniquePtr<FReaderChunks>* ReaderChunksPtr = ChunkDatabase.Find(InContainerId);
 		if (ReaderChunksPtr == nullptr)
 		{
@@ -437,6 +448,11 @@ public:
 
 	void LogSummary(const TArray<FIoStoreWriterResult>& IoStoreWriterResults)
 	{
+		if (!bValid)
+		{
+			return;
+		}
+
 		uint64 TotalEntryBytes = 0;
 		uint64 TotalMissBytes = 0;
 
@@ -5294,6 +5310,7 @@ int32 CreateTarget(const FIoStoreArguments& Arguments, const FIoStoreWriterSetti
 		if (ChunkDbPtr->Init(Arguments.ReferenceChunkGlobalContainerFileName, Arguments.ReferenceChunkAdditionalContainersPath, Arguments.ReferenceChunkKeys) == false)
 		{
 			UE_LOG(LogIoStore, Warning, TEXT("Failed to initialize reference chunk store. Pak will continue."));
+			ChunkDatabase.Reset();
 		}
 	}
 
@@ -5935,7 +5952,8 @@ int32 CreateTarget(const FIoStoreArguments& Arguments, const FIoStoreWriterSetti
 
 void FIoStoreChunkDatabase::WriteCSV(const FString& InOutputFileName, const TArray<FCookedPackage*>& InPackages)
 {
-	if (InOutputFileName.Len() == 0)
+	if (!bValid ||
+		InOutputFileName.Len() == 0)
 	{
 		return;
 	}
