Updating RCD logging to segregate optional containers.

diff --git a/Engine/Source/Developer/IoStoreUtilities/Private/IoStoreUtilities.cpp b/Engine/Source/Developer/IoStoreUtilities/Private/IoStoreUtilities.cpp
--- a/Engine/Source/Developer/IoStoreUtilities/Private/IoStoreUtilities.cpp
+++ b/Engine/Source/Developer/IoStoreUtilities/Private/IoStoreUtilities.cpp
@@ -143,6 +143,7 @@ public:
 		TArray<FRequestedChunkInfo> NewChunks;
 		uint64 NewChunksByType[IoChunkTypeCount] = { 0 };
 
+		std::atomic_uint64_t FulfillBytes = 0;
 		std::atomic_uint64_t TotalRequestCount = 0;
 		std::atomic_uint64_t ChangedCountByType[IoChunkTypeCount] = {0};
 		std::atomic_uint64_t UsedChunksByType[IoChunkTypeCount] = {0};
@@ -432,6 +433,7 @@ public:
 
 		FulfillBytesPerChunk[(int8)ChunkInfo.ChunkType] += TotalCompressedSize;
 		FulfillBytes += TotalCompressedSize;
+		ReaderChunks->FulfillBytes.fetch_add(TotalCompressedSize, std::memory_order::relaxed);
 
 		//
 		// At this point we know we can use the block so we can go async.
@@ -454,22 +456,56 @@ public:
 			return;
 		}
 
+		// In order to try and get a number that's closer to the actual patch size, we need
+		// to strip out the data that isn't actually distributed with the same and is instead
+		// streamed. This is OptionalBulkData and currently is only recognizable here by scanning
+		// the container name. (this is NOT OptionalSegment!)
+
 		uint64 TotalEntryBytes = 0;
 		uint64 TotalMissBytes = 0;
+		uint64 TotalOptionalMissBytes = 0;
 
 		for (const FIoStoreWriterResult& Result : IoStoreWriterResults)
 		{
+			bool bIsOptional = Result.ContainerName.Contains(TEXT("optional"));
+
 			TotalEntryBytes += Result.TotalEntryCompressedSize;
 			TotalMissBytes += Result.ReferenceCacheMissBytes;
+
+			if (bIsOptional)
+			{
+				TotalOptionalMissBytes += Result.ReferenceCacheMissBytes;
+			}
+		}
+
+		uint64 FulfillBytesOptional = 0;
+		for (const TPair<FIoContainerId, TUniquePtr<FReaderChunks>>& ReaderPair : ChunkDatabase)
+		{
+			bool bIsOptional = ReaderPair.Value->ContainerName.Contains(TEXT("optional"));
+			if (bIsOptional)
+			{
+				FulfillBytesOptional += ReaderPair.Value->FulfillBytes.load();
+			}
 		}
 
 		uint64 TotalCandidateBytes = FulfillBytes + TotalMissBytes;
+		uint64 TotalOptionalCandidateBytes = FulfillBytesOptional + TotalOptionalMissBytes;
+
+		uint64 TotalNonOptCandidateBytes = TotalCandidateBytes - TotalOptionalCandidateBytes;
+		uint64 TotalNonOptFulfillBytes = FulfillBytes - FulfillBytesOptional;
 
 		UE_LOG(LogIoStore, Display, TEXT("Reference Chunk Database:"));
 		UE_LOG(LogIoStore, Display, TEXT("    %s reused bytes out of %s candidate bytes - %.1f%% hit rate."),
 			*NumberString(FulfillBytes),
 			*NumberString(TotalCandidateBytes),
 			100.0 * FulfillBytes / (TotalCandidateBytes));
+		if (TotalOptionalCandidateBytes)
+		{
+			UE_LOG(LogIoStore, Display, TEXT("    %s reused non-optional bytes out of %s candidate non-optional bytes - %.1f%% hit rate."),
+				*NumberString(TotalNonOptFulfillBytes),
+				*NumberString(TotalNonOptCandidateBytes),
+				100.0 * TotalNonOptFulfillBytes / (TotalNonOptCandidateBytes));
+		}
 		UE_LOG(LogIoStore, Display, TEXT("    %s candidate bytes out of %s io chunk bytes - %.1f%% coverage."),
 			*NumberString(TotalCandidateBytes),
 			*NumberString(TotalEntryBytes),
