UnrealPak: Hash logic cleanup.
Deprecate the 32 byte FIoChunkHash (filled with 12 zero bytes) and replace it with the 20 byte FIoHash.
Remove IIoStoreWriterHashDatabase and merge AssetRegistry hashes and Zen ChunkInfo hashes into FCookedPackageStore.
Use Zen chunkinfo hashes for package and bulk data when paking from a Zen cook (package buffers are immutable).
Use AssetRegistry hashes for bulk data when paking from a loose file cook (package buffers are transformed).
Add IIoStoreWriteRequest::GetChunkHash to access the chunk hash from the cook.
Rename the argument -verifyhashdatabase to -validatechunkhashes.

Minor:
Calculate CompressionMethodForEntry once for each FIoStoreWriteQueueEntry.
Add trace cpu scopes for scheduler threads.
In FCookedPackageStore remove unused functions and optimize oplog parsing and GetChunkInfoFromFileName.

Tests:
Paking with -validatechunkhashes from a loose file cook with/without RefDB.
Paking with -validatechunkhashes from a Zen cook with/without RefDB.
Running UnrealPak "<ucas> -extract", "-list=<ucas>" and "iostore -list=<ucas>" commands on old containers serialized with FIoChunkHash in FIoStoreTocResource::ChunkMetas.

diff --git a/Engine/Source/Developer/IoStoreUtilities/Private/IoStoreLooseFiles.cpp b/Engine/Source/Developer/IoStoreUtilities/Private/IoStoreLooseFiles.cpp
--- a/Engine/Source/Developer/IoStoreUtilities/Private/IoStoreLooseFiles.cpp
+++ b/Engine/Source/Developer/IoStoreUtilities/Private/IoStoreLooseFiles.cpp
@@ -64,9 +64,6 @@ public:
 	virtual void SetReferenceChunkDatabase(TSharedPtr<IIoStoreWriterReferenceChunkDatabase> ReferenceChunkDatabase) override
 	{ }
 
-	virtual void SetHashDatabase(TSharedPtr<IIoStoreWriterHashDatabase> HashDatabase, bool bVerifyHashDatabase) override
-	{ }
-
 	virtual void EnableDiskLayoutOrdering(const TArray<TUniquePtr<FIoStoreReader>>& PatchSourceReaders = TArray<TUniquePtr<FIoStoreReader>>())
 	{ }
 
@@ -108,6 +105,11 @@ public:
 				return TArrayView<const FFileRegion>();
 			}
 
+			virtual const FIoHash* GetChunkHash() override
+			{
+				return nullptr;
+			}
+
 			FIoBuffer SourceBuffer;
 			uint64 OrderHint;
 		};
diff --git a/Engine/Source/Developer/IoStoreUtilities/Private/IoStoreUtilities.cpp b/Engine/Source/Developer/IoStoreUtilities/Private/IoStoreUtilities.cpp
--- a/Engine/Source/Developer/IoStoreUtilities/Private/IoStoreUtilities.cpp
+++ b/Engine/Source/Developer/IoStoreUtilities/Private/IoStoreUtilities.cpp
@@ -103,6 +103,20 @@ static const uint64 DefaultMemoryMappingAlignment = 16 << 10;
 
 static TUniquePtr<FIoStoreReader> CreateIoStoreReader(const TCHAR* Path, const FKeyChain& KeyChain);
 
+enum class ECookMetadataFiles
+{
+	None = 0,
+	AssetRegistry = 1,
+	CookMetadata = 2,
+	All = 4
+};
+ENUM_CLASS_FLAGS(ECookMetadataFiles);
+
+static ECookMetadataFiles FindAndLoadMetadataFiles(
+	const FString& InCookedDir, ECookMetadataFiles InRequiredFiles, 
+	FAssetRegistryState& OutAssetRegistry, FString* OutAssetRegistryFileName /*optional, set on success*/,
+	UE::Cook::FCookMetadataState* OutCookMetadata, FString* OutCookMetadataFileName /*optional, set on success or need*/);
+
 /*
 * Provides access to previously compressed chunks to the iostore writer, allowing
 * a) avoiding recompressing things and b) tweaks to compressors dont cause massive patches.
@@ -115,7 +129,7 @@ public:
 	struct FRequestedChunkInfo
 	{
 		FIoChunkId Id;
-		FIoChunkHash Hash;
+		FIoHash ChunkHash;
 	};
 
 	TArray<TUniquePtr<FIoStoreReader>> Readers;
@@ -129,7 +143,7 @@ public:
 
 		TMap<FIoChunkId, uint32> ChunkIds; // maps to index for chunk specific structures
 
-		TArray<FIoChunkHash> ChangedChunkHashes;
+		TArray<FIoHash> ChangedChunkHashes;
 		TArray<FIoStoreTocChunkInfo> ChunkInfos;
 		TArray<std::atomic_uint8_t> ChunkChanged;
 
@@ -311,7 +325,7 @@ public:
 
 	// Returns whether we expect to be able to load the chunk from the reference chunk database.
 	// This can be called from any thread, though in the presence of existing hashes it's single threaded.
-	virtual bool ChunkExists(const FIoContainerId& InContainerId, const FIoChunkHash& InChunkHash, const FIoChunkId& InChunkId, uint32& OutNumChunkBlocks)
+	virtual bool ChunkExists(const FIoContainerId& InContainerId, const FIoHash& InChunkHash, const FIoChunkId& InChunkId, uint32& OutNumChunkBlocks)
 	{
 		if (!bValid)
 		{
@@ -322,7 +336,7 @@ public:
 		uint8 ChunkType = (uint8)InChunkId.GetChunkType();
 
 		FRequestedChunkInfo RCI;
-		RCI.Hash = InChunkHash;
+		RCI.ChunkHash = InChunkHash;
 		RCI.Id = InChunkId;
 
 		TUniquePtr<FReaderChunks>* ReaderChunksPtr = ChunkDatabase.Find(InContainerId);
@@ -371,7 +385,7 @@ public:
 
 		// We have the chunk - does the hash match?
 		const FIoStoreTocChunkInfo& ChunkInfo = ReaderChunks->ChunkInfos[*ChunkIndex];
-		if (ChunkInfo.Hash != InChunkHash)
+		if (ChunkInfo.ChunkHash != InChunkHash)
 		{
 			// Chunk exists, but it's changed.
 			ReaderChunks->ChangedCountByType[ChunkType].fetch_add(1, std::memory_order_relaxed);
@@ -379,7 +393,7 @@ public:
 
 			// This should be safe because chunk ids are unique within a container,
 			// so we are the only ones poking at this index.
-			ReaderChunks->ChangedChunkHashes[*ChunkIndex] = MoveTemp(RCI.Hash);
+			ReaderChunks->ChangedChunkHashes[*ChunkIndex] = MoveTemp(RCI.ChunkHash);
 			return false;
 		}
 
@@ -393,7 +407,7 @@ public:
 
 	// This function was not written to be thread safe as it's only ever called from
 	// the iostore begindispatch thread (i.e. is not async)
-	virtual bool RetrieveChunk(const FIoContainerId& InContainerId, const FIoChunkHash& InChunkHash, const FIoChunkId& InChunkId, TUniqueFunction<void(TIoStatusOr<FIoStoreCompressedReadResult>)> InCompleteCallback)
+	virtual bool RetrieveChunk(const FIoContainerId& InContainerId, const FIoHash& InChunkHash, const FIoChunkId& InChunkId, TUniqueFunction<void(TIoStatusOr<FIoStoreCompressedReadResult>)> InCompleteCallback)
 	{
 		if (!bValid)
 		{
@@ -997,6 +1011,7 @@ struct FContainerTargetFile
 	uint64 SourceSize = 0;
 	uint64 IdealOrder = 0;
 	FIoChunkId ChunkId;
+	FIoHash ChunkHash;
 	TArray<uint8> PackageHeaderData;
 	EContainerChunkType ChunkType;
 	bool bForceUncompressed = false;
@@ -1010,6 +1025,8 @@ public:
 	struct FChunkInfo
 	{
 		FIoChunkId ChunkId;
+		FIoHash ChunkHash;
+		uint64 ChunkSize = 0;
 		FName PackageName;
 		FString RelativeFileName;
 		TArray<FFileRegion> FileRegions;
@@ -1023,6 +1040,7 @@ public:
 	FIoStatus Load(const TCHAR* ManifestFilename)
 	{
 		IOSTORE_CPU_SCOPE(LoadCookedPackageStore);
+		double StartTime = FPlatformTime::Seconds();
 
 		TUniquePtr<FArchive> Ar(IFileManager::Get().CreateFileReader(ManifestFilename));
 		if (!Ar)
@@ -1041,29 +1059,50 @@ public:
 			ZenStoreClient = MakeUnique<UE::FZenStoreHttpClient>(MoveTemp(ZenServiceSettings));
 			ZenStoreClient->InitializeReadOnly(ProjectId, OplogId);
 
+			IOSTORE_CPU_SCOPE(FetchOplog);
 			TIoStatusOr<FCbObject> OplogStatus = ZenStoreClient->GetOplog().Get();
 			if (!OplogStatus.IsOk())
 			{
 				return OplogStatus.Status();
 			}
-			
 			OplogObject = OplogStatus.ConsumeValueOrDie();
+			
 		}
 		else
 		{
 			OplogObject = ManifestObject["oplog"].AsObject();
 		}
+		UE_LOG(LogIoStore, Display, TEXT("Fetched %d oplog items from %s in %.2lf seconds"),
+			OplogObject["entries"].AsArrayView().Num(),
+			HasZenStoreClient() ? TEXT("Zen") : TEXT("Manifest"),
+			FPlatformTime::Seconds() - StartTime);
+
+		ParseOplog(OplogObject);
+		LoadChunkHashes();
+
+		return FIoStatus::Ok;
+	}
+
+	void ParseOplog(FCbObject& OplogObject)
+	{
+		IOSTORE_CPU_SCOPE(ParseOplog);
+		double StartTime = FPlatformTime::Seconds();
 
-		for (FCbField& OplogEntry : OplogObject["entries"].AsArray())
+		const FCbArrayView EntriesArray = OplogObject["entries"].AsArrayView();
+		const int32 EstimatedChunksCount = 3 * EntriesArray.Num(); //PackageData+BulkData+OptionalBulkData
+		ChunkInfoMap.Reserve(EstimatedChunksCount);
+		FilenameToChunkIdMap.Reserve(EstimatedChunksCount);
+		PackageIdToEntry.Reserve(EntriesArray.Num());
+		for (FCbFieldView& OplogEntry : EntriesArray)
 		{
-			FCbObject OplogObj = OplogEntry.AsObject();
-			FPackageStoreEntryResource PackageStoreEntry = FPackageStoreEntryResource::FromCbObject(OplogObj["packagestoreentry"].AsObject());
+			FCbObjectView OplogObj = OplogEntry.AsObjectView();
+			FPackageStoreEntryResource PackageStoreEntry = FPackageStoreEntryResource::FromCbObject(OplogObj["packagestoreentry"].AsObjectView());
 
 			auto AddChunksFromOplog = [this, &OplogEntry, &PackageStoreEntry](const char* Field)
 			{
-				for (FCbField& ChunkEntry : OplogEntry[Field].AsArray())
+				for (FCbFieldView& ChunkEntry : OplogEntry[Field].AsArrayView())
 				{
-					FCbObject ChunkObj = ChunkEntry.AsObject();
+					FCbObjectView ChunkObj = ChunkEntry.AsObjectView();
 					FIoChunkId ChunkId;
 					ChunkId.Set(ChunkObj["id"].AsObjectId().GetView());
 					FChunkInfo& ChunkInfo = ChunkInfoMap.Add(ChunkId);
@@ -1095,12 +1134,105 @@ public:
 
 			PackageIdToEntry.Add(PackageStoreEntry.GetPackageId(), MoveTemp(PackageStoreEntry));
 		}
+
+		UE_LOG(LogIoStore, Display, TEXT("Parsed %d oplog items %.2lf seconds, %d chunks"),
+			EntriesArray.Num(),
+			FPlatformTime::Seconds() - StartTime,
+			ChunkInfoMap.Num());
+	}
+
+	FIoStatus LoadChunkHashes()
+	{
+		IOSTORE_CPU_SCOPE(LoadChunkHashes);
+		double StartLoadTime = FPlatformTime::Seconds();
+		double StartUpdateTime = StartLoadTime;
+
+		uint32 LoadedHashCount = 0;
+		uint32 UpdatedHashCount = 0;
+
+		if (HasZenStoreClient())
+		{
+			FCbObject ChunksObj;
+			{
+				IOSTORE_CPU_SCOPE(GetChunkInfos);
+				TIoStatusOr<FCbObject> Chunks = ZenStoreClient->GetChunkInfos().Get();
+				if (!Chunks.IsOk())
+				{
+					return Chunks.Status();
+				}
+				ChunksObj = Chunks.ConsumeValueOrDie();
+			}
+
+			StartUpdateTime = FPlatformTime::Seconds();
+
+			IOSTORE_CPU_SCOPE(ParseChunkInfos);
+			for (FCbField& ChunkEntry : ChunksObj["chunkinfos"])
+			{
+				FCbObject ChunkObj = ChunkEntry.AsObject();
+				FIoChunkId ChunkId;
+				if (!LoadFromCompactBinary(ChunkObj["id"], ChunkId))
+				{
+					UE_LOG(LogIoStore, Warning, TEXT("Received invalid chunk id, skipping."));
+					continue;
+				}
+				if (FChunkInfo* FindChunkInfo = ChunkInfoMap.Find(ChunkId); FindChunkInfo != nullptr)
+				{
+					FindChunkInfo->ChunkHash = ChunkObj["rawhash"].AsHash();
+					FindChunkInfo->ChunkSize = ChunkObj["rawsize"].AsUInt64();
+					++UpdatedHashCount;
+				}
+				++LoadedHashCount;
+			}
+		}
+		else
+		{
+			FAssetRegistryState AssetRegistry;
+			if (FindAndLoadMetadataFiles(CookedDir, ECookMetadataFiles::None, AssetRegistry, nullptr, nullptr, nullptr) == ECookMetadataFiles::None)
+			{
+				return FIoStatus(EIoErrorCode::NotFound);
+			}
+
+			StartUpdateTime = FPlatformTime::Seconds();
+
+			const TMap<FName, const FAssetPackageData*>& Packages = AssetRegistry.GetAssetPackageDataMap();
+			for (auto PackageIter : Packages)
+			{
+				for (const TPair<FIoChunkId, FIoHash>& HashIter : PackageIter.Value->ChunkHashes)
+				{
+					// For the moment, only bulk data types are added to teh asset registry - gate here so that
+					// we remember to verify all the hashes match when they eventually get added during cook.
+					if (HashIter.Key.GetChunkType() == EIoChunkType::BulkData ||
+						HashIter.Key.GetChunkType() == EIoChunkType::OptionalBulkData)
+					{
+						if (FChunkInfo* FindChunkInfo = ChunkInfoMap.Find(HashIter.Key); FindChunkInfo != nullptr)
+						{
+							FindChunkInfo->ChunkHash = HashIter.Value;
+							++UpdatedHashCount;
+						}
+						++LoadedHashCount;
+					}
+				}
+			}
+		}
+
+		UE_LOG(LogIoStore, Display, TEXT("Loaded %u chunk hashes from %s in %.2lf seconds, %d hashes updated in %.2lf seconds"),
+			LoadedHashCount,
+			HasZenStoreClient() ? TEXT("Zen") : TEXT("AssetRegistry"),
+			StartUpdateTime - StartLoadTime,
+			UpdatedHashCount,
+			FPlatformTime::Seconds() - StartUpdateTime);
+
 		return FIoStatus::Ok;
 	}
 
 	FIoChunkId GetChunkIdFromFileName(const FString& Filename) const
 	{
-		return FilenameToChunkIdMap.FindRef(*Filename);
+		return FilenameToChunkIdMap.FindRef(Filename);
+	}
+
+	const FChunkInfo* GetChunkInfoFromChunkId(const FIoChunkId& ChunkId) const
+	{
+		return ChunkInfoMap.Find(ChunkId);
 	}
 
 	const FChunkInfo* GetChunkInfoFromFileName(const FString& Filename) const
@@ -1145,16 +1277,6 @@ public:
 		return ZenStoreClient.IsValid();
 	}
 
-	TFuture<TIoStatusOr<FCbObject>> GetChunkInfos()
-	{
-		return ZenStoreClient->GetChunkInfos();
-	}
-
-	TIoStatusOr<uint64> GetChunkSize(const FIoChunkId& ChunkId)
-	{
-		return ZenStoreClient->GetChunkSize(ChunkId);
-	}
-
 	TIoStatusOr<FIoBuffer> ReadChunk(const FIoChunkId& ChunkId)
 	{
 		FIoReadOptions ReadOptions;
@@ -1193,39 +1315,6 @@ public:
 
 		(new FAutoDeleteAsyncTask<FReadChunkTask>(ZenStoreClient.Get(), ChunkId, MoveTemp(Callback)))->StartBackgroundTask();
 	}
-	
-	TIoStatusOr<FIoBuffer> ReadPackageHeaderFromZen(FPackageId PackageId, uint16 ChunkIndex)
-	{
-		if (const FPackageStoreEntryResource* Entry = PackageIdToEntry.Find(PackageId))
-		{
-			FIoReadOptions ReadOptions;
-			ReadOptions.SetRange(0, 64 << 10);
-			
-			TIoStatusOr<FIoBuffer> Status = ZenStoreClient->ReadChunk(CreateIoChunkId(PackageId.Value(), ChunkIndex, EIoChunkType::ExportBundleData), ReadOptions.GetOffset(), ReadOptions.GetSize());
-			if (!Status.IsOk())
-			{
-				return Status;
-			}
-
-			FIoBuffer Buffer = Status.ConsumeValueOrDie();
-			uint32 HeaderSize = reinterpret_cast<const FZenPackageSummary*>(Buffer.Data())->HeaderSize;
-			if (HeaderSize > Buffer.DataSize())
-			{
-				ReadOptions.SetRange(0, HeaderSize);
-
-				Status = ZenStoreClient->ReadChunk(CreateIoChunkId(PackageId.Value(), ChunkIndex, EIoChunkType::ExportBundleData), ReadOptions.GetOffset(), ReadOptions.GetSize());
-				if (!Status.IsOk())
-				{
-					return Status;
-				}
-				Buffer = Status.ConsumeValueOrDie();
-			}
-			
-			return FIoBuffer(Buffer.Data(), HeaderSize, Buffer);
-		}
-
-		return FIoStatus(EIoErrorCode::NotFound);
-	}
 
 private:
 	TUniquePtr<UE::FZenStoreHttpClient> ZenStoreClient;
@@ -1268,7 +1357,6 @@ struct FIoStoreArguments
 	FReleasedPackages ReleasedPackages;
 	TUniquePtr<FCookedPackageStore> PackageStore;
 	TUniquePtr<FIoBuffer> ScriptObjects;
-	bool bVerifyHashDatabase = false;
 	bool bSign = false;
 	bool bRemapPluginContentToGame = false;
 	bool bCreateDirectoryIndex = true;
@@ -1349,7 +1437,7 @@ public:
 				Info.OffsetOnDisk,
 				Info.Size,
 				Info.CompressedSize,
-				*Info.Hash.ToString(),
+				*LexToString(Info.ChunkHash),
 				*LexToString(Info.ChunkType)
 			);
 		}
@@ -2840,6 +2928,20 @@ void InitializeContainerTargetsAndPackages(
 				UE_LOG(LogIoStore, Fatal, TEXT("Unexpected file type %d for file '%s'"), CookedFileStatData->FileType, *OutTargetFile.NormalizedSourcePath);
 				return false;
 			}
+
+			if (CookedFileStatData->FileType == FCookedFileStatData::BulkData ||
+				CookedFileStatData->FileType == FCookedFileStatData::OptionalBulkData)
+			{
+				const FCookedPackageStore::FChunkInfo* ChunkInfo =
+					Arguments.PackageStore->GetChunkInfoFromChunkId(OutTargetFile.ChunkId);
+				if (!ChunkInfo)
+				{
+					UE_LOG(LogIoStore, Warning, TEXT("File not found in manifest: '%s'"), *SourceFile.NormalizedPath);
+					return false;
+				}
+				// It is ok for the actual hash to be missing here, if so it will be calculated later
+				OutTargetFile.ChunkHash = ChunkInfo->ChunkHash;
+			}
 		}
 
 		// Only keep the regions for the file if neither compression nor encryption are enabled, otherwise the regions will be meaningless.
@@ -2867,43 +2969,7 @@ void InitializeContainerTargetsAndPackages(
 		return true;
 	};
 
-	struct FChunkListItem
-	{
-		FIoHash RawHash;
-		uint64 RawSize;
-	};
-
-	TMap<FIoChunkId, FChunkListItem> ChunkList;
-	if (Arguments.PackageStore->HasZenStoreClient())
-	{
-		double StartChunkInfoTime = FPlatformTime::Seconds();
-
-		TIoStatusOr<FCbObject> Chunks = Arguments.PackageStore->GetChunkInfos().Get();
-		if (!Chunks.IsOk())
-		{
-			UE_LOG(LogIoStore, Error, TEXT("Failed to retrieve chunk list"));
-			return;
-		}
-
-		FCbObject ChunksObj = Chunks.ConsumeValueOrDie();
-		for (FCbField& ChunkEntry : ChunksObj["chunkinfos"])
-		{
-			FCbObject ChunkObj = ChunkEntry.AsObject();
-			FIoChunkId ChunkId;
-			if (!LoadFromCompactBinary(ChunkObj["id"], ChunkId))
-			{
-				UE_LOG(LogIoStore, Warning, TEXT("Received invalid chunk id, skipping."));
-				continue;
-			}
-			ChunkList.Add(ChunkId, { ChunkObj["rawhash"].AsHash(), ChunkObj["rawsize"].AsUInt64() });
-		}
-
-		UE_LOG(LogIoStore, Display, TEXT("Fetched '%d' chunk infos in %f seconds"), ChunkList.Num(), FPlatformTime::Seconds() - StartChunkInfoTime);
-
-	}
-
 	auto CreateTargetFileFromZen = [
-		&ChunkList,
 		&Arguments,
 		&Packages,
 		&PackageNameMap,
@@ -2934,13 +3000,13 @@ void InitializeContainerTargetsAndPackages(
 			return false;
 		}
 		OutTargetFile.ChunkId = ChunkInfo->ChunkId;
-		FChunkListItem* ChunkListItem = ChunkList.Find(OutTargetFile.ChunkId);
-		if (!ChunkListItem)
+		if (ChunkInfo->ChunkSize == 0)
 		{
 			UE_LOG(LogIoStore, Warning, TEXT("Chunk size not found for: '%s'"), *SourceFile.NormalizedPath);
 			return false;
 		}
-		OutTargetFile.SourceSize = ChunkListItem->RawSize;
+		OutTargetFile.SourceSize = ChunkInfo->ChunkSize;
+		OutTargetFile.ChunkHash = ChunkInfo->ChunkHash;
 
 		if (ChunkInfo->PackageName.IsNone())
 		{
@@ -3028,6 +3094,7 @@ void InitializeContainerTargetsAndPackages(
 		{
 			IOSTORE_CPU_SCOPE(ProcessSourceFiles);
 			bool bHasOptionalSegmentPackages = false;
+			ContainerTarget->TargetFiles.Reserve(ContainerSource.SourceFiles.Num());
 			for (const FContainerSourceFile& SourceFile : ContainerSource.SourceFiles)
 			{
 				FContainerTargetFile TargetFile;
@@ -3305,12 +3372,6 @@ private:
 		{
 		}
 
-		virtual void PrepareSourceBufferAsync(FGraphEventRef InCompletionEvent) override
-		{
-			CompletionEvent = InCompletionEvent;
-			Manager.ScheduleLoad(this);
-		}
-
 		virtual uint64 GetOrderHint() override
 		{
 			return TargetFile.IdealOrder;
@@ -3321,6 +3382,17 @@ private:
 			return FileRegions;
 		}
 
+		virtual const FIoHash* GetChunkHash() override
+		{
+			return TargetFile.ChunkHash.IsZero() ? nullptr : &TargetFile.ChunkHash;
+		}
+		
+		virtual void PrepareSourceBufferAsync(FGraphEventRef InCompletionEvent) override
+		{
+			CompletionEvent = InCompletionEvent;
+			Manager.ScheduleLoad(this);
+		}
+
 		virtual const FIoBuffer* GetSourceBuffer() override
 		{
 			return &SourceBuffer;
@@ -3595,6 +3667,7 @@ private:
 
 	void InitiatorThreadFunc()
 	{
+		TRACE_CPUPROFILER_EVENT_SCOPE(SourceReadInitiatorThread);
 		for (;;)
 		{
 			FQueueEntry* QueueEntry = InitiatorQueue.DequeueOrWait();
@@ -3613,6 +3686,7 @@ private:
 
 	void RetirerThreadFunc()
 	{
+		TRACE_CPUPROFILER_EVENT_SCOPE(SourceReadRetirerThread);
 		for (;;)
 		{
 			FQueueEntry* QueueEntry = RetirerQueue.DequeueOrWait();
@@ -4794,15 +4868,6 @@ static bool SaveAssetRegistry(const FString& InAssetRegistryFileName, FAssetRegi
 	return true;
 }
 
-enum class ECookMetadataFiles
-{
-	None = 0,
-	AssetRegistry = 1,
-	CookMetadata = 2,
-	All = 4
-};
-ENUM_CLASS_FLAGS(ECookMetadataFiles);
-
 static ECookMetadataFiles FindAndLoadMetadataFiles(
 	const FString& InCookedDir, ECookMetadataFiles InRequiredFiles, 
 	FAssetRegistryState& OutAssetRegistry, FString* OutAssetRegistryFileName /*optional, set on success*/,
@@ -5064,59 +5129,6 @@ static bool DoAssetRegistryWritebackDuringStage(
 	return true;
 }
 
-// Implements providing the chunk hashes that exist in the asset registry to the
-// iostore writer to avoid reading and hashing redundently.
-class FIoStoreHashDb : public IIoStoreWriterHashDatabase
-{
-public:
-	virtual ~FIoStoreHashDb() {}
-
-	TMap<FIoChunkId, FIoHash> Hashes;
-
-	bool Initialize(const FString& InCookedDir)
-	{
-		FString AssetRegistryFileName;
-		FAssetRegistryState AssetRegistry;
-		if (FindAndLoadMetadataFiles(InCookedDir, ECookMetadataFiles::None, AssetRegistry, nullptr, nullptr, nullptr) == ECookMetadataFiles::None)
-		{
-			// already logged
-			return false;
-		}
-
-		double StartTime = FPlatformTime::Seconds();
-
-		const TMap<FName, const FAssetPackageData*>& Packages = AssetRegistry.GetAssetPackageDataMap();
-		for (auto PackageIter : Packages)
-		{
-			for (const TPair<FIoChunkId, FIoHash>& HashIter : PackageIter.Value->ChunkHashes)
-			{
-				// For the moment, only bulk data types are added to teh asset registry - gate here so that
-				// we remember to verify all the hashes match when they eventually get added during cook.
-				if (HashIter.Key.GetChunkType() == EIoChunkType::BulkData ||
-					HashIter.Key.GetChunkType() == EIoChunkType::OptionalBulkData)
-				{
-					Hashes.Add(HashIter.Key, HashIter.Value);
-				}
-			}
-		}
-
-		double EndTime = FPlatformTime::Seconds();
-		UE_LOG(LogIoStore, Display, TEXT("Added %d hashes to the hash database, init took %f seconds"), Hashes.Num(), EndTime - StartTime);
-		return true;
-	}
-
-	virtual bool FindHashForChunkId(const FIoChunkId& ChunkId, FIoChunkHash& OutHash) const override
-	{
-		const FIoHash* Exists = Hashes.Find(ChunkId);
-		if (Exists)
-		{
-			OutHash = FIoChunkHash::CreateFromIoHash(*Exists);
-			return true;
-		}
-		return false;
-	}
-};
-
 // modified copy from PakFileUtilities
 static FName RemapLocalizationPathIfNeeded(const FString& Path)
 {
@@ -5351,16 +5363,6 @@ int32 CreateTarget(const FIoStoreArguments& Arguments, const FIoStoreWriterSetti
 		}
 	}
 
-	TSharedPtr<IIoStoreWriterHashDatabase> HashDatabase = MakeShared<FIoStoreHashDb>();
-	if (((FIoStoreHashDb&)*HashDatabase).Initialize(Arguments.CookedDir) == false)
-	{
-		UE_LOG(LogIoStore, Display, TEXT("Unabled to initialize the hash database from the asset registry!"));
-	}
-	if (Arguments.bVerifyHashDatabase)
-	{
-		UE_LOG(LogIoStore, Display, TEXT("Hash database verification on: hashes will be checked for accuracy during this run."));
-	}
-
 	TArray<FCookedPackage*> Packages;
 	FPackageNameMap PackageNameMap;
 	FPackageIdMap PackageIdMap;
@@ -5450,7 +5452,6 @@ int32 CreateTarget(const FIoStoreArguments& Arguments, const FIoStoreWriterSetti
 				ContainerTarget->IoStoreWriter = IoStoreWriterContext->CreateContainer(*ContainerTarget->OutputPath, ContainerSettings);
 				ContainerTarget->IoStoreWriter->EnableDiskLayoutOrdering(ContainerTarget->PatchSourceReaders);
 				ContainerTarget->IoStoreWriter->SetReferenceChunkDatabase(ChunkDatabase);
-				ContainerTarget->IoStoreWriter->SetHashDatabase(HashDatabase, Arguments.bVerifyHashDatabase);
 				IoStoreWriters.Add(ContainerTarget->IoStoreWriter);
 
 				if (EnumHasAnyFlags(ContainerTarget->ContainerFlags, EIoContainerFlags::OnDemand))
@@ -6052,11 +6053,11 @@ void FIoStoreChunkDatabase::WriteCSV(const FString& InOutputFileName, const TArr
 					ChangedType == EIoChunkType::ShaderCodeLibrary)
 				{
 					// These don't have the package name in the chunk id.
-					ChangeListArchive->Logf(TEXT("%s,%s,CHANGED,%s,%s,<shader>"), *Reader->ContainerName, *ChunkNames[(uint8)ChangedType], *LexToString(Id), *Reader->ChangedChunkHashes[Chunk.Value].ToString());
+					ChangeListArchive->Logf(TEXT("%s,%s,CHANGED,%s,%s,<shader>"), *Reader->ContainerName, *ChunkNames[(uint8)ChangedType], *LexToString(Id), *LexToString(Reader->ChangedChunkHashes[Chunk.Value]));
 				}
 				else
 				{
-					ChangeListArchive->Logf(TEXT("%s,%s,CHANGED,%s,%s,%s"), *Reader->ContainerName, *ChunkNames[(uint8)ChangedType], *LexToString(Id), *Reader->ChangedChunkHashes[Chunk.Value].ToString(), *FindPackageForChunk(Id).ToString());
+					ChangeListArchive->Logf(TEXT("%s,%s,CHANGED,%s,%s,%s"), *Reader->ContainerName, *ChunkNames[(uint8)ChangedType], *LexToString(Id), *LexToString(Reader->ChangedChunkHashes[Chunk.Value]), *FindPackageForChunk(Id).ToString());
 				}
 			}
 		}
@@ -6067,11 +6068,11 @@ void FIoStoreChunkDatabase::WriteCSV(const FString& InOutputFileName, const TArr
 				NewChunk.Id.GetChunkType() == EIoChunkType::ShaderCodeLibrary)
 			{
 				// These don't have the package name in the chunk id.
-				ChangeListArchive->Logf(TEXT("%s,%s,NEW,%s,%s,<shader>"), *Reader->ContainerName, *ChunkNames[(uint8)NewChunk.Id.GetChunkType()], *LexToString(NewChunk.Id), *NewChunk.Hash.ToString());
+				ChangeListArchive->Logf(TEXT("%s,%s,NEW,%s,%s,<shader>"), *Reader->ContainerName, *ChunkNames[(uint8)NewChunk.Id.GetChunkType()], *LexToString(NewChunk.Id), *LexToString(NewChunk.ChunkHash));
 			}
 			else
 			{
-				ChangeListArchive->Logf(TEXT("%s,%s,NEW,%s,%s,%s"), *Reader->ContainerName, *ChunkNames[(uint8)NewChunk.Id.GetChunkType()], *LexToString(NewChunk.Id), *NewChunk.Hash.ToString(), *FindPackageForChunk(NewChunk.Id).ToString());
+				ChangeListArchive->Logf(TEXT("%s,%s,NEW,%s,%s,%s"), *Reader->ContainerName, *ChunkNames[(uint8)NewChunk.Id.GetChunkType()], *LexToString(NewChunk.Id), *LexToString(NewChunk.ChunkHash), *FindPackageForChunk(NewChunk.Id).ToString());
 			}
 		}
 	}
@@ -6088,12 +6089,12 @@ void FIoStoreChunkDatabase::WriteCSV(const FString& InOutputFileName, const TArr
 			{
 				// These don't have the package name in the chunk id.
 				ChangeListArchive->Logf(TEXT("%s,%s,NOCONTAINER,%s,%s,<shader>"),
-					*MissingContainer->ContainerName, *ChunkName, *LexToString(NoContainerChunk.Id), *NoContainerChunk.Hash.ToString());
+					*MissingContainer->ContainerName, *ChunkName, *LexToString(NoContainerChunk.Id), *LexToString(NoContainerChunk.ChunkHash));
 			}
 			else
 			{
 				ChangeListArchive->Logf(TEXT("%s,%s,NOCONTAINER,%s,%s,%"),
-					*MissingContainer->ContainerName, *ChunkName, *LexToString(NoContainerChunk.Id), *NoContainerChunk.Hash.ToString(), *FindPackageForChunk(NoContainerChunk.Id).ToString());
+					*MissingContainer->ContainerName, *ChunkName, *LexToString(NoContainerChunk.Id), *LexToString(NoContainerChunk.ChunkHash), *FindPackageForChunk(NoContainerChunk.Id).ToString());
 			}
 		}
 	}
@@ -6180,12 +6181,12 @@ int32 CreateContentPatch(const FIoStoreArguments& Arguments, const FIoStoreWrite
 
 		TSharedPtr<IIoStoreWriter> IoStoreWriter = IoStoreWriterContext->CreateContainer(*Container.OutputPath, ContainerSettings);
 		IoStoreWriters.Add(IoStoreWriter);
-		TMap<FIoChunkId, FIoChunkHash> SourceHashByChunkId;
+		TMap<FIoChunkId, FIoHash> SourceHashByChunkId;
 		for (const TUniquePtr<FIoStoreReader>& SourceReader : SourceReaders)
 		{
 			SourceReader->EnumerateChunks([&SourceHashByChunkId](const FIoStoreTocChunkInfo& ChunkInfo)
 			{
-				SourceHashByChunkId.Add(ChunkInfo.Id, ChunkInfo.Hash);
+				SourceHashByChunkId.Add(ChunkInfo.Id, ChunkInfo.ChunkHash);
 				return true;
 			});
 		}
@@ -6204,8 +6205,8 @@ int32 CreateContentPatch(const FIoStoreArguments& Arguments, const FIoStoreWrite
 
 		TargetReader->EnumerateChunks([&TargetReader, &SourceHashByChunkId, &IoStoreWriter, &ChunkFileNamesMap](const FIoStoreTocChunkInfo& ChunkInfo)
 		{
-			FIoChunkHash* FindSourceHash = SourceHashByChunkId.Find(ChunkInfo.Id);
-			if (!FindSourceHash || *FindSourceHash != ChunkInfo.Hash)
+			FIoHash* FindSourceHash = SourceHashByChunkId.Find(ChunkInfo.Id);
+			if (!FindSourceHash || *FindSourceHash != ChunkInfo.ChunkHash)
 			{
 				FIoReadOptions ReadOptions;
 				TIoStatusOr<FIoBuffer> ChunkBuffer = TargetReader->Read(ChunkInfo.Id, ReadOptions);
@@ -6662,7 +6663,7 @@ bool LegacyListIoStoreContainer(
 	struct FEntry
 	{
 		FIoChunkId ChunkId;
-		FIoChunkHash Hash;
+		FIoHash ChunkHash;
 		FString FileName;
 		int64 Offset;
 		int64 Size;
@@ -6677,7 +6678,7 @@ bool LegacyListIoStoreContainer(
 			
 			FEntry& Entry = Entries.AddDefaulted_GetRef();
 			Entry.ChunkId = ChunkInfo.Id;
-			Entry.Hash = ChunkInfo.Hash;
+			Entry.ChunkHash = ChunkInfo.ChunkHash;
 			Entry.FileName = ChunkInfo.FileName;
 			Entry.Offset = CompressedBlocks[FirstBlockIndex].Offset;
 			Entry.Size = ChunkInfo.CompressedSize;
@@ -6709,7 +6710,7 @@ bool LegacyListIoStoreContainer(
 				*Entry.FileName,
 				Entry.Offset,
 				Entry.Size,
-				*Entry.Hash.ToString(),
+				*LexToString(Entry.ChunkHash),
 				TEXT("false"),
 				bWasCompressed ? TEXT("true") : TEXT("false"),
 				Entry.CompressionMethodIndex));
@@ -6733,7 +6734,7 @@ bool LegacyListIoStoreContainer(
 				*Entry.FileName,
 				Entry.Offset,
 				Entry.Size,
-				*Entry.Hash.ToString(),
+				*LexToString(Entry.ChunkHash),
 				*CompressionMethodNames[Entry.CompressionMethodIndex]);
 		}
 		FileSize += Entry.Size;
@@ -6874,7 +6875,7 @@ int32 ProfileReadSpeed(const TCHAR* InCommandLine, const FKeyChain& InKeyChain)
 		OutstandingJobs++;
 		UE::Tasks::Launch(TEXT("IoStoreUtil::ReadJob"), [Id = Id, &OutstandingJobs, &JustGotSpaceEvent, &TotalJobsRemaining, &BytesRead, &Reader, &CompletedEvent, MaxOutstandingJobs, ReadType, bValidate]()
 		{
-			FIoChunkHash ReadHash;
+			FIoHash ReadHash;
 			bool bHashValid = false;
 
 			switch (ReadType)
@@ -6892,7 +6893,7 @@ int32 ProfileReadSpeed(const TCHAR* InCommandLine, const FKeyChain& InKeyChain)
 
 					if (bValidate)
 					{
-						ReadHash = FIoChunkHash::HashBuffer(Result.GetData(), Result.GetSize());
+						ReadHash = FIoHash::HashBuffer(Result.GetData(), Result.GetSize());
 						bHashValid = true;
 					}
 
@@ -6908,7 +6909,7 @@ int32 ProfileReadSpeed(const TCHAR* InCommandLine, const FKeyChain& InKeyChain)
 
 					if (bValidate)
 					{
-						ReadHash = FIoChunkHash::HashBuffer(Result.GetData(), Result.GetSize());
+						ReadHash = FIoHash::HashBuffer(Result.GetData(), Result.GetSize());
 						bHashValid = true;
 					}
 					break;
@@ -6917,7 +6918,7 @@ int32 ProfileReadSpeed(const TCHAR* InCommandLine, const FKeyChain& InKeyChain)
 
 			if (bHashValid)
 			{
-				FIoChunkHash CheckAgainstHash = Reader->GetChunkInfo(Id).ValueOrDie().Hash;
+				FIoHash CheckAgainstHash = Reader->GetChunkInfo(Id).ValueOrDie().ChunkHash;
 				if (ReadHash != CheckAgainstHash)
 				{
 					UE_LOG(LogIoStore, Warning, TEXT("Read hash mismatch: Chunk %s"), *LexToString(Id));
@@ -8189,7 +8190,7 @@ static int32 Diff(
 		{
 			if (const FIoStoreTocChunkInfo* SourceChunkInfo = SourceContainer.ChunkInfoById.Find(TargetChunkInfo.Key))
 			{
-				if (SourceChunkInfo->Hash != TargetChunkInfo.Value.Hash)
+				if (SourceChunkInfo->ChunkHash != TargetChunkInfo.Value.ChunkHash)
 				{
 					ContainerDiff.Modified.Add(TargetChunkInfo.Key);
 					ContainerDiff.ModifiedCompressedSize += TargetChunkInfo.Value.CompressedSize;
@@ -8444,7 +8445,7 @@ bool LegacyDiffIoStoreContainers(const TCHAR* InContainerFilename1, const TCHAR*
 	struct FEntry
 	{
 		FString FileName;
-		FIoChunkHash Hash;
+		FIoHash ChunkHash;
 		uint64 Size;
 	};
 
@@ -8453,7 +8454,7 @@ bool LegacyDiffIoStoreContainers(const TCHAR* InContainerFilename1, const TCHAR*
 		{
 			FEntry& Entry = Container1Entries.Add(ChunkInfo.Id);
 			Entry.FileName = ChunkInfo.FileName;
-			Entry.Hash = ChunkInfo.Hash;
+			Entry.ChunkHash = ChunkInfo.ChunkHash;
 			Entry.Size = ChunkInfo.Size;
 			return true;
 		});
@@ -8472,7 +8473,7 @@ bool LegacyDiffIoStoreContainers(const TCHAR* InContainerFilename1, const TCHAR*
 					UE_LOG(LogIoStore, Log, TEXT("FilesizeDifferent, %s, %llu, %llu"), *ChunkInfo.FileName, FindContainer1Entry->Size, ChunkInfo.Size);
 					++NumDifferentContents;
 				}
-				else if (FindContainer1Entry->Hash != ChunkInfo.Hash)
+				else if (FindContainer1Entry->ChunkHash != ChunkInfo.ChunkHash)
 				{
 					UE_LOG(LogIoStore, Log, TEXT("ContentsDifferent, %s, %llu, %llu"), *ChunkInfo.FileName, FindContainer1Entry->Size, ChunkInfo.Size);
 					++NumDifferentContents;
@@ -8863,7 +8864,7 @@ bool ProcessFilesFromIoStoreContainer(
 		uint64 Offset;
 		bool bIsCompressed;
 
-		FIoChunkHash Hash;
+		FIoHash ChunkHash;
 	};
 	TArray<FEntry> Entries;
 	const FIoDirectoryIndexReader& IndexReader = Reader->GetDirectoryIndexReader();
@@ -8887,7 +8888,7 @@ bool ProcessFilesFromIoStoreContainer(
 			Entry.Offset = ChunkInfo.Offset;
 			Entry.bIsCompressed = ChunkInfo.bIsCompressed;
 
-			Entry.Hash = ChunkInfo.Hash;
+			Entry.ChunkHash = ChunkInfo.ChunkHash;
 
 			return true;
 		});
@@ -9370,11 +9371,9 @@ bool ParseContainerGenerationArguments(FIoStoreArguments& Arguments, FIoStoreWri
 		}
 	}
 
-	// By default, we use any hashes in the asset registry that exist in order to avoid reading and hashing
-	// chunk unnecessarily. This flag causes us to read and hash anyway, and then ensure they match what is
-	// in the asset registry. It is very bad if this fails!
-	Arguments.bVerifyHashDatabase = FParse::Param(FCommandLine::Get(), TEXT("verifyhashdatabase"));
-
+	// By default, we use uncompressed chunk hashes from cook to avoid reading and hashing chunks unnecessarily.
+	// This flag causes us to read and hash anyway, and then ensure they match. It is very bad if this fails!
+	WriterSettings.bValidateChunkHashes = FParse::Param(FCommandLine::Get(), TEXT("validateChunkHashes"));
 
 	uint64 PatchPaddingAlignment = 0;
 	if (ParseSizeArgument(FCommandLine::Get(), TEXT("-patchpaddingalign="), PatchPaddingAlignment))
diff --git a/Engine/Source/Developer/IoStoreUtilities/Private/IoStoreWriter.cpp b/Engine/Source/Developer/IoStoreUtilities/Private/IoStoreWriter.cpp
--- a/Engine/Source/Developer/IoStoreUtilities/Private/IoStoreWriter.cpp
+++ b/Engine/Source/Developer/IoStoreUtilities/Private/IoStoreWriter.cpp
@@ -43,7 +43,7 @@ struct FIoStoreWriteQueueEntry
 	class FIoStoreWriter* Writer = nullptr;
 	IIoStoreWriteRequest* Request = nullptr;
 	FIoChunkId ChunkId;
-	FIoChunkHash ChunkHash;
+	FIoHash ChunkHash;
 	/** Hash of the block data as it would be found on disk after compression and encryption */
 	FIoHash ChunkDiskHash;
 
@@ -61,6 +61,7 @@ struct FIoStoreWriteQueueEntry
 	uint64 Offset = 0;
 	TArray<FChunkBlock> ChunkBlocks;
 	FIoWriteOptions Options;
+	FName CompressionMethod = NAME_None;
 	FGraphEventRef HashBarrier;
 	FGraphEventRef HashTask;
 	FGraphEventRef BeginCompressionBarrier;
@@ -535,11 +536,6 @@ public:
 		// Add ourselves to the reference chunk db's list of possibles
 		ReferenceChunkDatabase->NotifyAddedToWriter(ContainerSettings.ContainerId, FPaths::GetBaseFilename(TocFilePath));
 	}
-	void SetHashDatabase(TSharedPtr<IIoStoreWriterHashDatabase> InHashDatabase, bool bInVerifyHashDatabase)
-	{
-		HashDatabase = InHashDatabase;
-		bVerifyHashDatabase = bInVerifyHashDatabase;
-	}
 
 	void EnumerateChunks(TFunction<bool(FIoStoreTocChunkInfo&&)>&& Callback) const
 	{
@@ -585,7 +581,7 @@ public:
 			PatchSourceReader->EnumerateChunks([this, &PrevEntryLink, &LayoutEntriesWithOffsets](const FIoStoreTocChunkInfo& ChunkInfo)
 				{
 					FLayoutEntry* PreviousBuildEntry = new FLayoutEntry();
-					PreviousBuildEntry->Hash = ChunkInfo.Hash;
+					PreviousBuildEntry->ChunkHash = ChunkInfo.ChunkHash;
 					PreviousBuildEntry->PartitionIndex = ChunkInfo.PartitionIndex;
 					PreviousBuildEntry->CompressedSize = ChunkInfo.CompressedSize;
 					LayoutEntriesWithOffsets.Emplace(ChunkInfo.Offset, PreviousBuildEntry);
@@ -635,70 +631,63 @@ public:
 		check(!bHasFlushed);
 		checkf(ChunkId.IsValid(), TEXT("ChunkId is not valid!"));
 
+		WriterContext->TotalChunksCount.IncrementExchange();
 		FIoStoreWriteQueueEntry* Entry = new FIoStoreWriteQueueEntry();
+		Entries.Add(Entry);
 		Entry->Writer = this;
 		Entry->Sequence = Entries.Num();
-		WriterContext->TotalChunksCount.IncrementExchange();
-		Entries.Add(Entry);
 		Entry->ChunkId = ChunkId;
 		Entry->Options = WriteOptions;
+		Entry->CompressionMethod = CompressionMethodForEntry(WriteOptions);
 		Entry->Request = Request;		
 		Entry->BeginCompressionBarrier = FGraphEvent::CreateGraphEvent();
 		Entry->FinishCompressionBarrier = FGraphEvent::CreateGraphEvent();
 		Entry->FinishEncryptionAndSigningBarrier = FGraphEvent::CreateGraphEvent();
 		Entry->BeginWriteBarrier = FGraphEvent::CreateGraphEvent();
 		Entry->WriteFinishedEvent = FGraphEvent::CreateGraphEvent();
-		
-		bool bTestHashValid = false;
-		FIoChunkHash TestHash;
-		// If we can get the hash without reading the whole thing and hashing it, do so
-		// to avoid the IO.
-		if (HashDatabase.IsValid() &&
-			HashDatabase->FindHashForChunkId(ChunkId, Entry->ChunkHash))
+
+		// If we can get the hash without reading the whole thing and hashing it, do so to avoid the IO.
+		if (const FIoHash* ChunkHash = Request->GetChunkHash(); ChunkHash != nullptr)
 		{
-			if (bVerifyHashDatabase == false)
+			check(!ChunkHash->IsZero());
+			Entry->ChunkHash = *ChunkHash;
+			if (WriterContext->WriterSettings.bValidateChunkHashes == false)
 			{
 				// If we aren't validating then we just use it and bail.
 				WriterContext->HashDbChunksCount.IncrementExchange();
 				WriterContext->HashDbChunksByType[(int8)Entry->ChunkId.GetChunkType()].IncrementExchange();
 				WriterContext->HashedChunksCount.IncrementExchange();
 
-				if (ReferenceChunkDatabase.IsValid() && CompressionMethodForEntry(Entry) != NAME_None)
+				if (ReferenceChunkDatabase.IsValid() && Entry->CompressionMethod != NAME_None)
 				{
 					Entry->bLoadingFromReferenceDb = ReferenceChunkDatabase->ChunkExists(ContainerSettings.ContainerId, Entry->ChunkHash, Entry->ChunkId, Entry->NumChunkBlocksFromRefDb);
 					Entry->bCouldBeFromReferenceDb = true;
 				}
 				return;
 			}
-
-			// If we are validating, copy the hash out and run the normal path
-			// to check against it.
-			TestHash = Entry->ChunkHash;
-			bTestHashValid = true;
+			// If we are validating run the normal path to verify it.
 		}
 		// Otherwise, we have to do the load & hash
 
 		Entry->HashBarrier = FGraphEvent::CreateGraphEvent();
-
-		FGraphEventArray HashPrereqs;
-		HashPrereqs.Add(Entry->HashBarrier);
-		Entry->HashTask = FFunctionGraphTask::CreateAndDispatchWhenReady([this, Entry, TestHash, bTestHashValid]()
+		Entry->HashTask = FFunctionGraphTask::CreateAndDispatchWhenReady([this, Entry]()
 		{
 			TRACE_CPUPROFILER_EVENT_SCOPE(HashChunk);
 			const FIoBuffer* SourceBuffer = Entry->Request->GetSourceBuffer();
-			Entry->ChunkHash = FIoChunkHash::HashBuffer(SourceBuffer->Data(), SourceBuffer->DataSize());
-			WriterContext->HashedChunksCount.IncrementExchange();
+			FIoHash ChunkHash = FIoHash::HashBuffer(SourceBuffer->Data(), SourceBuffer->DataSize());
 
-			if (bVerifyHashDatabase && bTestHashValid)
+			if (!Entry->ChunkHash.IsZero() && Entry->ChunkHash != ChunkHash)
 			{
-				if (TestHash != Entry->ChunkHash)
-				{
-					UE_LOG(LogIoStore, Warning, TEXT("HashDb Validation Failed: ChunkId %s has mismatching hash"), *LexToString(Entry->ChunkId));
-				}
+				UE_LOG(LogIoStore, Warning, TEXT("Hash Validation Failed: ChunkId %s has mismatching hash, new calculated '%s' vs old cached '%s'"),
+					*LexToString(Entry->ChunkId),
+					*LexToString(ChunkHash),
+					*LexToString(Entry->ChunkHash));
 			}
 
+			Entry->ChunkHash = ChunkHash;
+			WriterContext->HashedChunksCount.IncrementExchange();
 
-			if (ReferenceChunkDatabase.IsValid() && CompressionMethodForEntry(Entry) != NAME_None)
+			if (ReferenceChunkDatabase.IsValid() && Entry->CompressionMethod != NAME_None)
 			{
 				Entry->bLoadingFromReferenceDb = ReferenceChunkDatabase->ChunkExists(ContainerSettings.ContainerId, Entry->ChunkHash, Entry->ChunkId, Entry->NumChunkBlocksFromRefDb);
 				Entry->bCouldBeFromReferenceDb = true;
@@ -706,7 +695,9 @@ public:
 
 			// Release the source data buffer, it will be reloaded later when we start compressing the chunk
 			Entry->Request->FreeSourceBuffer();
-		}, TStatId(), &HashPrereqs, ENamedThreads::AnyHiPriThreadHiPriTask);
+		}, TStatId(), Entry->HashBarrier, ENamedThreads::AnyHiPriThreadHiPriTask);
+
+		// Kick off the source buffer read to run the hash task
 		Entry->Request->PrepareSourceBufferAsync(Entry->HashBarrier);
 	}
 
@@ -748,6 +739,11 @@ public:
 				return TArrayView<const FFileRegion>();
 			}
 
+			virtual const FIoHash* GetChunkHash() override
+			{
+				return nullptr;
+			}
+
 			FIoBuffer SourceBuffer;
 			uint64 OrderHint;
 		};
@@ -1124,7 +1120,7 @@ private:
 		FLayoutEntry* Next = nullptr;
 		uint64 IdealOrder = 0;
 		uint64 CompressedSize = uint64(-1);
-		FIoChunkHash Hash;
+		FIoHash ChunkHash;
 		FIoStoreWriteQueueEntry* QueueEntry = nullptr;
 		int32 PartitionIndex = -1;
 	};
@@ -1152,7 +1148,7 @@ private:
 			FLayoutEntry* FindPreviousEntry = PreviousBuildLayoutEntryByChunkId.FindRef(WriteQueueEntry->ChunkId);
 			if (FindPreviousEntry)
 			{
-				if (FindPreviousEntry->Hash != WriteQueueEntry->ChunkHash)
+				if (FindPreviousEntry->ChunkHash != WriteQueueEntry->ChunkHash)
 				{
 					WriteQueueEntry->bModified = true;
 				}
@@ -1339,11 +1335,11 @@ private:
 		return !Ar.IsError();
 	}
 
-	FName CompressionMethodForEntry(FIoStoreWriteQueueEntry* Entry) const
+	FName CompressionMethodForEntry(const FIoWriteOptions& Options) const
 	{
 		FName CompressionMethod = NAME_None;
 		const FIoStoreWriterSettings& WriterSettings = WriterContext->WriterSettings;
-		if (ContainerSettings.IsCompressed() && !Entry->Options.bForceUncompressed && !Entry->Options.bIsMemoryMapped)
+		if (ContainerSettings.IsCompressed() && !Options.bForceUncompressed && !Options.bIsMemoryMapped)
 		{
 			CompressionMethod = WriterSettings.CompressionMethod;
 		}
@@ -1435,8 +1431,6 @@ private:
 
 		const FIoStoreWriterSettings& WriterSettings = WriterContext->WriterSettings;
 
-		FName CompressionMethod = CompressionMethodForEntry(Entry);
-
 		const FIoBuffer* SourceBuffer = Entry->Request->GetSourceBuffer();
 		Entry->UncompressedSize.Emplace(SourceBuffer->DataSize());
 
@@ -1459,7 +1453,7 @@ private:
 			{
 				FChunkBlock& Block = Entry->ChunkBlocks[BlockIndex];
 				Block.IoBuffer = WriterContext->AllocCompressionBuffer(NumChunkBlocks);
-				Block.CompressionMethod = CompressionMethod;
+				Block.CompressionMethod = Entry->CompressionMethod;
 				Block.UncompressedSize = FMath::Min(BytesToProcess, WriterSettings.CompressionBlockSize);
 				Block.UncompressedData = UncompressedData;
 				BytesToProcess -= Block.UncompressedSize;
@@ -1467,7 +1461,7 @@ private:
 			}
 		}
 
-		if (CompressionMethod == NAME_None)
+		if (Entry->CompressionMethod == NAME_None)
 		{
 			for (FChunkBlock& Block : Entry->ChunkBlocks)
 			{
@@ -1486,7 +1480,7 @@ private:
 		if (bUseDDCCompression)
 		{
 			TStringBuilder<256> CacheKeySuffix;
-			CacheKeySuffix.Append(Entry->ChunkHash.ToString());
+			CacheKeySuffix.Append(LexToString(Entry->ChunkHash));
 			WriterSettings.CompressionMethod.AppendString(CacheKeySuffix);
 			CacheKeySuffix.Append(FCompression::GetCompressorDDCSuffix(WriterSettings.CompressionMethod));
 			CacheKeySuffix.Appendf(TEXT("%d_%d"),
@@ -1517,7 +1511,12 @@ private:
 				Entry->bStoreCompressedDataInDDC = true;
 			}
 		}
-		
+
+		ScheduleCompressionTasks(Entry);
+	}
+
+	void ScheduleCompressionTasks(FIoStoreWriteQueueEntry* Entry)
+	{
 		for (FChunkBlock& Block : Entry->ChunkBlocks)
 		{
 			WriterContext->ScheduledCompressionTasksCount.IncrementExchange();
@@ -1595,9 +1594,10 @@ private:
 			// database if we DO hit this, however...
 			UE_LOG(LogIoStore, Warning, TEXT("ChunkId was added twice in container %s, %s, file %s hash %s vs %s"), 
 				*FPaths::GetBaseFilename(TocFilePath),
-				*LexToString(Entry->ChunkId), *Entry->Options.FileName,
-				*TocBuilder.GetTocResource().ChunkMetas[*FindExistingIndex].ChunkHash.ToString(),
-				*Entry->ChunkHash.ToString()
+				*LexToString(Entry->ChunkId),
+				*Entry->Options.FileName,
+				*LexToString(TocBuilder.GetTocResource().ChunkMetas[*FindExistingIndex].ChunkHash),
+				*LexToString(Entry->ChunkHash)
 				);
 
 			checkf(TocBuilder.GetTocResource().ChunkMetas[*FindExistingIndex].ChunkHash == Entry->ChunkHash, TEXT("Chunk id has already been added with different content"));
@@ -1793,8 +1793,6 @@ private:
 	bool						bHasFlushed = false;
 	bool						bHasResult = false;
 	TSharedPtr<IIoStoreWriterReferenceChunkDatabase> ReferenceChunkDatabase;
-	TSharedPtr<IIoStoreWriterHashDatabase> HashDatabase;
-	bool						bVerifyHashDatabase = false;
 
 
 	friend class FIoStoreWriterContextImpl;
@@ -1894,6 +1892,7 @@ void FIoStoreWriterContextImpl::Flush()
 
 void FIoStoreWriterContextImpl::BeginCompressionThreadFunc()
 {
+	TRACE_CPUPROFILER_EVENT_SCOPE(BeginCompressionThread);
 	for (;;)
 	{
 		FIoStoreWriteQueueEntry* Entry = BeginCompressionQueue.DequeueOrWait();
@@ -1914,6 +1913,7 @@ void FIoStoreWriterContextImpl::BeginCompressionThreadFunc()
 
 void FIoStoreWriterContextImpl::BeginEncryptionAndSigningThreadFunc()
 {
+	TRACE_CPUPROFILER_EVENT_SCOPE(BeginEncryptionAndSigningThread);
 	for (;;)
 	{
 		FIoStoreWriteQueueEntry* Entry = BeginEncryptionAndSigningQueue.DequeueOrWait();
@@ -1948,6 +1948,7 @@ void FIoStoreWriterContextImpl::BeginEncryptionAndSigningThreadFunc()
 
 void FIoStoreWriterContextImpl::FinishEncryptionAndSigningThreadFunc()
 {
+	TRACE_CPUPROFILER_EVENT_SCOPE(FinishEncryptionAndSigningThread);
 	for (;;)
 	{
 		FIoStoreWriteQueueEntry* Entry = FinishEncryptionAndSigningQueue.DequeueOrWait();
@@ -1974,6 +1975,7 @@ void FIoStoreWriterContextImpl::FinishEncryptionAndSigningThreadFunc()
 
 void FIoStoreWriterContextImpl::WriterThreadFunc()
 {
+	TRACE_CPUPROFILER_EVENT_SCOPE(WriterThread);
 	for (;;)
 	{
 		FIoStoreWriteQueueEntry* Entry = WriterQueue.DequeueOrWait();
diff --git a/Engine/Source/Developer/IoStoreUtilities/Private/IoStoreWriter.h b/Engine/Source/Developer/IoStoreUtilities/Private/IoStoreWriter.h
--- a/Engine/Source/Developer/IoStoreUtilities/Private/IoStoreWriter.h
+++ b/Engine/Source/Developer/IoStoreUtilities/Private/IoStoreWriter.h
@@ -24,6 +24,7 @@ struct FIoStoreWriterSettings
 	uint64 MaxPartitionSize = 0;
 	bool bEnableFileRegions = false;
 	bool bCompressionEnableDDC = false;
+	bool bValidateChunkHashes = false;
 };
 
 struct FIoStoreWriterResult
@@ -97,10 +98,12 @@ class IIoStoreWriteRequest
 public:
 	virtual ~IIoStoreWriteRequest() = default;
 
-	// Launches any async operations necessary in order to access the buffer. CompletionEvent is set once it's ready, which may be immediate.
-	virtual void PrepareSourceBufferAsync(FGraphEventRef CompletionEvent) = 0;
 	virtual uint64 GetOrderHint() = 0;
 	virtual TArrayView<const FFileRegion> GetRegions() = 0;
+	virtual const FIoHash* GetChunkHash() = 0;
+
+	// Launches any async operations necessary in order to access the buffer. CompletionEvent is set once it's ready, which may be immediate.
+	virtual void PrepareSourceBufferAsync(FGraphEventRef CompletionEvent) = 0;
 
 	// Only valid after the completion event passed to PrepareSourceBufferAsync has fired.
 	virtual const FIoBuffer* GetSourceBuffer() = 0;
@@ -132,14 +135,14 @@ public:
 	* Chunks provided *MUST* decompress to bits that hash to the exact value provided in InChunkKey (i.e. be exactly the same bits),
 	* and also be the same number of blocks (i.e. same CompressionBlockSize)
 	*/
-	virtual bool RetrieveChunk(const FIoContainerId& InContainerId, const FIoChunkHash& InChunkHash, const FIoChunkId& InChunkId, TUniqueFunction<void(TIoStatusOr<FIoStoreCompressedReadResult>)> InCompletionCallback) = 0;
+	virtual bool RetrieveChunk(const FIoContainerId& InContainerId, const FIoHash& InChunkHash, const FIoChunkId& InChunkId, TUniqueFunction<void(TIoStatusOr<FIoStoreCompressedReadResult>)> InCompletionCallback) = 0;
 
 	/* 
 	* Quick synchronous existence check that returns the number of blocks for the chunk. This is used to set up
 	* the necessary structures without needing to read the source data for the chunk. This might be called from
 	* multiple threads as it has to happen after we have the source hash computed.
 	*/
-	virtual bool ChunkExists(const FIoContainerId& InContainerId, const FIoChunkHash& InChunkHash, const FIoChunkId& InChunkId, uint32& OutNumChunkBlocks) = 0;
+	virtual bool ChunkExists(const FIoContainerId& InContainerId, const FIoHash& InChunkHash, const FIoChunkId& InChunkId, uint32& OutNumChunkBlocks) = 0;
 
 	/*
 	* Returns the compression block size that was used to break up the IoChunks in the source containers. If this is different than what we want, 
@@ -153,20 +156,6 @@ public:
 	virtual void NotifyAddedToWriter(const FIoContainerId& InContainerId, const FString& InContainerName) = 0;
 };
 
-/**
-*	Allows the IIoStoreWriter to avoid loading and hashing chunks, saving pak/stage time, as the normal
-*	process involved loading the chunks, hashing them, freeing them, making some decisions, then loading
-*	them _again_ for compression/writting. It's completely fine for this to not have all available hashes,
-*	but they have to match when provided!
-*/
-class IIoStoreWriterHashDatabase
-{
-public:
-	virtual ~IIoStoreWriterHashDatabase() = default;
-	virtual bool FindHashForChunkId(const FIoChunkId& ChunkId, FIoChunkHash& OutHash) const = 0;
-};
-
-
 class IIoStoreWriter
 {
 public:
@@ -177,7 +166,6 @@ public:
 	*	from previous containers instead of recompressing input data. This must be set before any writes are appended.
 	*/
 	virtual void SetReferenceChunkDatabase(TSharedPtr<IIoStoreWriterReferenceChunkDatabase> ReferenceChunkDatabase) = 0;
-	virtual void SetHashDatabase(TSharedPtr<IIoStoreWriterHashDatabase> HashDatabase, bool bVerifyHashDatabase) = 0;
 	virtual void EnableDiskLayoutOrdering(const TArray<TUniquePtr<FIoStoreReader>>& PatchSourceReaders = TArray<TUniquePtr<FIoStoreReader>>()) = 0;
 	virtual void Append(const FIoChunkId& ChunkId, FIoBuffer Chunk, const FIoWriteOptions& WriteOptions, uint64 OrderHint = MAX_uint64) = 0;
 	virtual void Append(const FIoChunkId& ChunkId, IIoStoreWriteRequest* Request, const FIoWriteOptions& WriteOptions) = 0;
diff --git a/Engine/Source/Runtime/Core/Internal/IO/IoStore.h b/Engine/Source/Runtime/Core/Internal/IO/IoStore.h
--- a/Engine/Source/Runtime/Core/Internal/IO/IoStore.h
+++ b/Engine/Source/Runtime/Core/Internal/IO/IoStore.h
@@ -29,6 +29,7 @@ enum class EIoStoreTocVersion : uint8
 	PerfectHashWithOverflow,
 	OnDemandMetaData,
 	RemovedOnDemandMetaData,
+	ReplaceIoChunkHashWithIoHash,
 	LatestPlusOne,
 	Latest = LatestPlusOne - 1
 };
@@ -90,7 +91,7 @@ ENUM_CLASS_FLAGS(FIoStoreTocEntryMetaFlags);
 struct FIoStoreTocEntryMeta
 {
 	// Source data hash (i.e. not the on disk data)
-	FIoChunkHash ChunkHash;
+	FIoHash ChunkHash;
 	FIoStoreTocEntryMetaFlags Flags;
 };
 
diff --git a/Engine/Source/Runtime/Core/Private/IO/IoStore.cpp b/Engine/Source/Runtime/Core/Private/IO/IoStore.cpp
--- a/Engine/Source/Runtime/Core/Private/IO/IoStore.cpp
+++ b/Engine/Source/Runtime/Core/Private/IO/IoStore.cpp
@@ -277,7 +277,10 @@ FIoStoreTocChunkInfo FIoStoreTocResource::GetTocChunkInfo(int32 TocEntryIndex) c
 	FIoStoreTocChunkInfo ChunkInfo;
 	ChunkInfo.Id = ChunkIds[TocEntryIndex];
 	ChunkInfo.ChunkType = ChunkInfo.Id.GetChunkType();
-	ChunkInfo.Hash = Meta.ChunkHash;
+	PRAGMA_DISABLE_DEPRECATION_WARNINGS
+	ChunkInfo.Hash = FIoChunkHash::CreateFromIoHash(Meta.ChunkHash);
+	PRAGMA_ENABLE_DEPRECATION_WARNINGS
+	ChunkInfo.ChunkHash = Meta.ChunkHash;
 	ChunkInfo.bHasValidFileName = false;
 	ChunkInfo.bIsCompressed = EnumHasAnyFlags(Meta.Flags, FIoStoreTocEntryMetaFlags::Compressed);
 	ChunkInfo.bIsMemoryMapped = EnumHasAnyFlags(Meta.Flags, FIoStoreTocEntryMetaFlags::MemoryMapped);
@@ -1270,8 +1273,28 @@ FIoStatus FIoStoreTocResource::Read(const TCHAR* TocFilePath, EIoStoreTocReadOpt
 	{
 		const uint8* TocMeta = (uint8*)DirectoryIndexBuffer + Header.DirectoryIndexSize;
 
-		const FIoStoreTocEntryMeta* ChunkMetas = reinterpret_cast<const FIoStoreTocEntryMeta*>(TocMeta);
-		OutTocResource.ChunkMetas = MakeArrayView<FIoStoreTocEntryMeta const>(ChunkMetas, Header.TocEntryCount);
+		if (Header.Version >= static_cast<uint8>(EIoStoreTocVersion::ReplaceIoChunkHashWithIoHash))
+		{
+			const FIoStoreTocEntryMeta* ChunkMetas = reinterpret_cast<const FIoStoreTocEntryMeta*>(TocMeta);
+			OutTocResource.ChunkMetas = MakeArrayView<FIoStoreTocEntryMeta const>(ChunkMetas, Header.TocEntryCount);
+		}
+		else
+		{
+			struct FIoStoreTocEntryMetaOld
+			{
+				uint8 ChunkHash[32];
+				FIoStoreTocEntryMetaFlags Flags;
+			};
+			const FIoStoreTocEntryMetaOld* ChunkMetas = reinterpret_cast<const FIoStoreTocEntryMetaOld*>(TocMeta);
+			TConstArrayView<FIoStoreTocEntryMetaOld> OldChunkMetas = MakeArrayView<FIoStoreTocEntryMetaOld const>(ChunkMetas, Header.TocEntryCount);
+			OutTocResource.ChunkMetas.Reserve(OldChunkMetas.Num());
+			for (const FIoStoreTocEntryMetaOld& OldChunkMeta : OldChunkMetas)
+			{
+				FIoStoreTocEntryMeta& ChunkMeta = OutTocResource.ChunkMetas.Emplace_GetRef();
+				FMemory::Memcpy(ChunkMeta.ChunkHash.GetBytes(), &OldChunkMeta.ChunkHash, sizeof ChunkMeta.ChunkHash);
+				ChunkMeta.Flags = OldChunkMeta.Flags;
+			}
+		}
 	}
 
 	if (Header.Version < static_cast<uint8>(EIoStoreTocVersion::PartitionSize))
diff --git a/Engine/Source/Runtime/Core/Private/IO/PackageStore.cpp b/Engine/Source/Runtime/Core/Private/IO/PackageStore.cpp
--- a/Engine/Source/Runtime/Core/Private/IO/PackageStore.cpp
+++ b/Engine/Source/Runtime/Core/Private/IO/PackageStore.cpp
@@ -74,7 +74,7 @@ FCbWriter& operator<<(FCbWriter& Writer, const FPackageStoreEntryResource& Packa
 	return Writer;
 }
 
-FPackageStoreEntryResource FPackageStoreEntryResource::FromCbObject(const FCbObject& Obj)
+FPackageStoreEntryResource FPackageStoreEntryResource::FromCbObject(FCbObjectView Obj)
 {
 	FPackageStoreEntryResource Entry;
 
@@ -92,7 +92,7 @@ FPackageStoreEntryResource FPackageStoreEntryResource::FromCbObject(const FCbObj
 	
 	if (Obj["shadermaphashes"])
 	{
-		for (FCbField& ArrayField : Obj["shadermaphashes"].AsArray())
+		for (FCbFieldView& ArrayField : Obj["shadermaphashes"].AsArrayView())
 		{
 			FSHAHash& ShaderMapHash = Entry.ShaderMapHashes.AddDefaulted_GetRef();
 			ShaderMapHash.FromString(FUTF8ToTCHAR(ArrayField.AsString()));
diff --git a/Engine/Source/Runtime/Core/Public/IO/IoDispatcher.h b/Engine/Source/Runtime/Core/Public/IO/IoDispatcher.h
--- a/Engine/Source/Runtime/Core/Public/IO/IoDispatcher.h
+++ b/Engine/Source/Runtime/Core/Public/IO/IoDispatcher.h
@@ -89,7 +89,9 @@ private:
 	int32			Order = 0;
 };
 
-class FIoChunkHash
+class
+UE_DEPRECATED(5.5, "FIoChunkHash is deprecated. Use FIoHash instead.")
+FIoChunkHash
 {
 public:
 	friend uint32 GetTypeHash(const FIoChunkHash& InChunkHash)
@@ -523,8 +525,12 @@ struct FIoContainerSettings
 struct FIoStoreTocChunkInfo
 {
 	FIoChunkId Id;
-	FString FileName;
+	FIoHash ChunkHash;
+	UE_DEPRECATED(5.5, "Hash of type FIoChunkHash is deprecated. Use ChunkHash of type FIoHash instead.")
+	PRAGMA_DISABLE_DEPRECATION_WARNINGS
 	FIoChunkHash Hash;
+	PRAGMA_ENABLE_DEPRECATION_WARNINGS
+	FString FileName;
 	uint64 Offset;
 	uint64 OffsetOnDisk;
 	uint64 Size;
@@ -536,6 +542,14 @@ struct FIoStoreTocChunkInfo
 	bool bForceUncompressed;
 	bool bIsMemoryMapped;
 	bool bIsCompressed;
+
+	PRAGMA_DISABLE_DEPRECATION_WARNINGS // Compilers can complain about deprecated members in compiler generated code
+	FIoStoreTocChunkInfo() = default;
+	FIoStoreTocChunkInfo(const FIoStoreTocChunkInfo&) = default;
+	FIoStoreTocChunkInfo(FIoStoreTocChunkInfo&&) = default;
+	FIoStoreTocChunkInfo& operator=(FIoStoreTocChunkInfo&) = default;
+	FIoStoreTocChunkInfo& operator=(FIoStoreTocChunkInfo&&) = default;
+	PRAGMA_ENABLE_DEPRECATION_WARNINGS
 };
 
 struct FIoStoreTocCompressedBlockInfo
diff --git a/Engine/Source/Runtime/Core/Public/IO/PackageStore.h b/Engine/Source/Runtime/Core/Public/IO/PackageStore.h
--- a/Engine/Source/Runtime/Core/Public/IO/PackageStore.h
+++ b/Engine/Source/Runtime/Core/Public/IO/PackageStore.h
@@ -15,7 +15,7 @@
 #include "UObject/NameTypes.h"
 
 class FArchive;
-class FCbObject;
+class FCbObjectView;
 class FCbWriter;
 class FStructuredArchiveSlot;
 
@@ -96,7 +96,7 @@ struct FPackageStoreEntryResource
 	
 	CORE_API friend FCbWriter& operator<<(FCbWriter& Writer, const FPackageStoreEntryResource& PackageStoreEntry);
 	
-	CORE_API static FPackageStoreEntryResource FromCbObject(const FCbObject& Obj);
+	CORE_API static FPackageStoreEntryResource FromCbObject(FCbObjectView Obj);
 };
 
 /**
