UnrealPak: Reduce scheduler threads and their lifetime (mostly for improved visualization, analysis and understanding).
Remove superfluous FinishEncryptionAndSigningThread and set the FIoStoreWriteQueueEntry::CompressedSize in EncryptAndSign() instead.
Construct and initialize WriteRequestManager and start its threads when it is first used after ProcessShaderLibraries.
Start FIoStoreWriterContextImpl scheduler threads from Flush() instead of from Initialize().
Stop FIoStoreWriterContextImpl scheduler threads once they have drained their queue.
Add trace counters to analyze the overall progress of the different type of tasks.

#rnx
#tests Binary diff of Content/Paks

diff --git a/Engine/Source/Developer/IoStoreUtilities/Private/IoStoreUtilities.cpp b/Engine/Source/Developer/IoStoreUtilities/Private/IoStoreUtilities.cpp
--- a/Engine/Source/Developer/IoStoreUtilities/Private/IoStoreUtilities.cpp
+++ b/Engine/Source/Developer/IoStoreUtilities/Private/IoStoreUtilities.cpp
@@ -84,7 +84,9 @@ IMPLEMENT_MODULE(FDefaultModuleImpl, IoStoreUtilities);
 #define IOSTORE_CPU_SCOPE(NAME) TRACE_CPUPROFILER_EVENT_SCOPE(IoStore##NAME);
 #define IOSTORE_CPU_SCOPE_DATA(NAME, DATA) TRACE_CPUPROFILER_EVENT_SCOPE(IoStore##NAME);
 
-TRACE_DECLARE_MEMORY_COUNTER(IoStoreUsedFileBufferMemory, TEXT("IoStore/UsedFileBufferMemory"));
+TRACE_DECLARE_MEMORY_COUNTER(IoStoreSourceReadsUsedBufferMemory, TEXT("IoStoreWriter/SourceReadsUsedBufferMemory"));
+TRACE_DECLARE_ATOMIC_INT_COUNTER(IoStoreSourceReadsInflight, TEXT("IoStoreWriter/SourceReadsInflight"));
+TRACE_DECLARE_ATOMIC_INT_COUNTER(IoStoreSourceReadsDone, TEXT("IoStoreWriter/SourceReadsDone"));
 
 // Helper to format numbers with comma separators to help readability: 1,234 vs 1234.
 static FString NumberString(uint64 N) { return FText::AsNumber(N).ToString(); }
@@ -3420,6 +3422,8 @@ private:
 
 		void OnSourceBufferLoaded()
 		{
+			TRACE_COUNTER_DECREMENT(IoStoreSourceReadsInflight);
+			TRACE_COUNTER_INCREMENT(IoStoreSourceReadsDone);
 			QueueEntry->ReleaseRef(Manager);
 			CompletionEvent->DispatchSubsequents();
 		}
@@ -3642,7 +3646,8 @@ private:
 		}
 
 		UsedBufferMemory.AddExchange(SourceBufferSize);
-		TRACE_COUNTER_ADD(IoStoreUsedFileBufferMemory, SourceBufferSize);
+		TRACE_COUNTER_INCREMENT(IoStoreSourceReadsInflight);
+		TRACE_COUNTER_ADD(IoStoreSourceReadsUsedBufferMemory, SourceBufferSize);
 		QueueEntry->WriteRequest->LoadSourceBufferAsync();
 	}
 
@@ -3657,11 +3662,11 @@ private:
 		delete QueueEntry;
 	}
 
-	void OnBufferMemoryFreed(uint64 Count)
+	void OnBufferMemoryFreed(uint64 Size)
 	{
-		uint64 OldValue = UsedBufferMemory.SubExchange(Count);
-		check(OldValue >= Count);
-		TRACE_COUNTER_SUBTRACT(IoStoreUsedFileBufferMemory, Count);
+		uint64 OldSize = UsedBufferMemory.SubExchange(Size);
+		check(OldSize >= Size);
+		TRACE_COUNTER_SUBTRACT(IoStoreSourceReadsUsedBufferMemory, Size);
 		MemoryAvailableEvent->Trigger();
 	}
 
@@ -5369,7 +5374,6 @@ int32 CreateTarget(const FIoStoreArguments& Arguments, const FIoStoreWriterSetti
 
 	FPackageStoreOptimizer PackageStoreOptimizer;
 	PackageStoreOptimizer.Initialize(*Arguments.ScriptObjects);
-	FIoStoreWriteRequestManager WriteRequestManager(PackageStoreOptimizer, Arguments.PackageStore.Get());
 
 	TArray<FContainerTargetSpec*> ContainerTargets;
 	UE_LOG(LogIoStore, Display, TEXT("Creating container targets..."));
@@ -5533,6 +5537,8 @@ int32 CreateTarget(const FIoStoreArguments& Arguments, const FIoStoreWriterSetti
 	FShaderAssociationInfo ShaderAssocInfo;
 	ProcessShaderLibraries(Arguments, ContainerTargets, Shaders, ShaderAssocInfo);
 
+	FIoStoreWriteRequestManager WriteRequestManager(PackageStoreOptimizer, Arguments.PackageStore.Get());
+
 	auto AppendTargetFileChunk = [&WriteRequestManager](FContainerTargetSpec* ContainerTarget, const FContainerTargetFile& TargetFile)
 	{
 		FIoWriteOptions WriteOptions;
diff --git a/Engine/Source/Developer/IoStoreUtilities/Private/IoStoreWriter.cpp b/Engine/Source/Developer/IoStoreUtilities/Private/IoStoreWriter.cpp
--- a/Engine/Source/Developer/IoStoreUtilities/Private/IoStoreWriter.cpp
+++ b/Engine/Source/Developer/IoStoreUtilities/Private/IoStoreWriter.cpp
@@ -20,6 +20,15 @@
 #include "Serialization/MemoryWriter.h"
 #include "Templates/UniquePtr.h"
 
+TRACE_DECLARE_MEMORY_COUNTER(IoStoreCompressionUsedBufferMemory, TEXT("IoStoreWriter/CompressionUsedBufferMemory"));
+TRACE_DECLARE_INT_COUNTER(IoStoreCompressionAvailableBuffers, TEXT("IoStoreWriter/CompressionAvailableBuffers"));
+TRACE_DECLARE_ATOMIC_INT_COUNTER(IoStoreCompressionInflight, TEXT("IoStoreWriter/CompressionInflight"));
+TRACE_DECLARE_ATOMIC_INT_COUNTER(IoStoreRefDbInflight, TEXT("IoStoreWriter/RefDbInFlight"));
+TRACE_DECLARE_ATOMIC_INT_COUNTER(IoStoreRefDbDone, TEXT("IoStoreWriter/RefDbDone"));
+TRACE_DECLARE_INT_COUNTER(IoStoreBeginCompressionCount, TEXT("IoStoreWriter/BeginCompression"));
+TRACE_DECLARE_INT_COUNTER(IoStoreBeginEncryptionAndSigningCount, TEXT("IoStoreWriter/BeginEncryptionAndSigning"));
+TRACE_DECLARE_INT_COUNTER(IoStoreBeginWriteCount, TEXT("IoStoreWriter/BeginWrite"));
+
 struct FChunkBlock
 {
 	const uint8* UncompressedData = nullptr;
@@ -66,7 +75,6 @@ struct FIoStoreWriteQueueEntry
 	FGraphEventRef HashTask;
 	FGraphEventRef BeginCompressionBarrier;
 	FGraphEventRef FinishCompressionBarrier;
-	FGraphEventRef FinishEncryptionAndSigningBarrier;
 	FGraphEventRef BeginWriteBarrier;
 	FGraphEventRef WriteFinishedEvent;
 	TAtomic<int32> CompressedBlocksCount{ 0 };
@@ -175,11 +183,9 @@ public:
 	{
 		BeginCompressionQueue.CompleteAdding();
 		BeginEncryptionAndSigningQueue.CompleteAdding();
-		FinishEncryptionAndSigningQueue.CompleteAdding();
 		WriterQueue.CompleteAdding();
 		BeginCompressionThread.Wait();
 		BeginEncryptionAndSigningThread.Wait();
-		FinishEncryptionAndSigningThread.Wait();
 		WriterThread.Wait();
 		if (CompressionBufferAvailableEvent)
 		{
@@ -216,12 +222,7 @@ public:
 		{
 			AvailableCompressionBuffers.Add(new FIoBuffer(CompressionBufferSize));
 		}
-		//TRACE_COUNTER_SET(IoStoreAvailableCompressionBuffers, AvailableCompressionBuffers.Num());
-
-		BeginCompressionThread = Async(EAsyncExecution::Thread, [this]() { BeginCompressionThreadFunc(); });
-		BeginEncryptionAndSigningThread = Async(EAsyncExecution::Thread, [this]() { BeginEncryptionAndSigningThreadFunc(); });
-		FinishEncryptionAndSigningThread = Async(EAsyncExecution::Thread, [this]() { FinishEncryptionAndSigningThreadFunc(); });
-		WriterThread = Async(EAsyncExecution::Thread, [this]() { WriterThreadFunc(); });
+		TRACE_COUNTER_SET(IoStoreCompressionAvailableBuffers, AvailableCompressionBuffers.Num());
 
 		return FIoStatus::Ok;
 	}
@@ -290,6 +291,7 @@ public:
 	{
 		if (TotalEntryChunkBlocksCount > TotalCompressionBufferCount)
 		{
+			TRACE_COUNTER_ADD(IoStoreCompressionUsedBufferMemory, CompressionBufferSize);
 			return new FIoBuffer(CompressionBufferSize);
 		}
 		FIoBuffer* AllocatedBuffer = nullptr;
@@ -300,7 +302,8 @@ public:
 				if (AvailableCompressionBuffers.Num())
 				{
 					AllocatedBuffer = AvailableCompressionBuffers.Pop();
-					//TRACE_COUNTER_DECREMENT(IoStoreAvailableCompressionBuffers);
+					TRACE_COUNTER_ADD(IoStoreCompressionUsedBufferMemory, CompressionBufferSize);
+					TRACE_COUNTER_DECREMENT(IoStoreCompressionAvailableBuffers);
 				}
 			}
 			if (!AllocatedBuffer)
@@ -316,6 +319,7 @@ public:
 	{
 		if (TotalEntryChunkBlocksCount > TotalCompressionBufferCount)
 		{
+			TRACE_COUNTER_SUBTRACT(IoStoreCompressionUsedBufferMemory, CompressionBufferSize);
 			delete Buffer;
 			return;
 		}
@@ -324,7 +328,8 @@ public:
 			FScopeLock Lock(&AvailableCompressionBuffersCritical);
 			bTriggerEvent = AvailableCompressionBuffers.Num() == 0;
 			AvailableCompressionBuffers.Push(Buffer);
-			//TRACE_COUNTER_INCREMENT(IoStoreAvailableCompressionBuffers);
+			TRACE_COUNTER_SUBTRACT(IoStoreCompressionUsedBufferMemory, CompressionBufferSize);
+			TRACE_COUNTER_INCREMENT(IoStoreCompressionAvailableBuffers);
 		}
 		if (bTriggerEvent)
 		{
@@ -335,7 +340,6 @@ public:
 private:
 	void BeginCompressionThreadFunc();
 	void BeginEncryptionAndSigningThreadFunc();
-	void FinishEncryptionAndSigningThreadFunc();
 	void WriterThreadFunc();
 
 	FIoStoreWriterSettings WriterSettings;
@@ -343,11 +347,9 @@ private:
 	FEvent* CompressionBufferAvailableEvent = nullptr;
 	TFuture<void> BeginCompressionThread;
 	TFuture<void> BeginEncryptionAndSigningThread;
-	TFuture<void> FinishEncryptionAndSigningThread;
 	TFuture<void> WriterThread;
 	FIoStoreWriteQueue BeginCompressionQueue;
 	FIoStoreWriteQueue BeginEncryptionAndSigningQueue;
-	FIoStoreWriteQueue FinishEncryptionAndSigningQueue;
 	FIoStoreWriteQueue WriterQueue;
 	TAtomic<uint64> TotalChunksCount{ 0 };
 	TAtomic<uint64> HashedChunksCount{ 0 };
@@ -642,7 +644,6 @@ public:
 		Entry->Request = Request;		
 		Entry->BeginCompressionBarrier = FGraphEvent::CreateGraphEvent();
 		Entry->FinishCompressionBarrier = FGraphEvent::CreateGraphEvent();
-		Entry->FinishEncryptionAndSigningBarrier = FGraphEvent::CreateGraphEvent();
 		Entry->BeginWriteBarrier = FGraphEvent::CreateGraphEvent();
 		Entry->WriteFinishedEvent = FGraphEvent::CreateGraphEvent();
 
@@ -1357,6 +1358,7 @@ private:
 			if (Entry->NumChunkBlocksFromRefDb == 0)
 			{
 				Entry->FinishCompressionBarrier->DispatchSubsequents();
+				TRACE_COUNTER_INCREMENT(IoStoreRefDbDone);
 				return;
 			}
 
@@ -1375,6 +1377,7 @@ private:
 			// Valid chunks must create the same decompressed bits, but can have different compressed bits.
 			// Since we are on a lightweight dispatch thread, the actual read is async, as is the processing
 			// of the results.
+			TRACE_COUNTER_INCREMENT(IoStoreRefDbInflight);
 			bool bChunkExists = ReferenceChunkDatabase->RetrieveChunk(ContainerSettings.ContainerId, Entry->ChunkHash, Entry->ChunkId, [this, Entry](TIoStatusOr<FIoStoreCompressedReadResult> InReadResult)
 			{
 
@@ -1419,6 +1422,8 @@ private:
 
 				Entry->UncompressedSize.Emplace(TotalUncompressedSize);
 				Entry->FinishCompressionBarrier->DispatchSubsequents();
+				TRACE_COUNTER_DECREMENT(IoStoreRefDbInflight);
+				TRACE_COUNTER_INCREMENT(IoStoreRefDbDone);
 			});
 
 			check(bChunkExists); // Sanity - should never return false as we can only get here if ChunkExists() returns true.
@@ -1517,6 +1522,7 @@ private:
 
 	void ScheduleCompressionTasks(FIoStoreWriteQueueEntry* Entry)
 	{
+		TRACE_COUNTER_INCREMENT(IoStoreCompressionInflight);
 		for (FChunkBlock& Block : Entry->ChunkBlocks)
 		{
 			WriterContext->ScheduledCompressionTasksCount.IncrementExchange();
@@ -1530,6 +1536,7 @@ private:
 					WriterContext->CompressedChunksByType[(int8)Entry->ChunkId.GetChunkType()].IncrementExchange();
 					WriterContext->CompressedChunksCount.IncrementExchange();
 					Entry->FinishCompressionBarrier->DispatchSubsequents();
+					TRACE_COUNTER_DECREMENT(IoStoreCompressionInflight);
 				}
 			}, TStatId(), nullptr, ENamedThreads::AnyHiPriThreadNormalTask);
 		}
@@ -1542,13 +1549,13 @@ private:
 			FFunctionGraphTask::CreateAndDispatchWhenReady([this, Entry]()
 			{
 				EncryptAndSign(Entry);
-				Entry->FinishEncryptionAndSigningBarrier->DispatchSubsequents();
+				Entry->BeginWriteBarrier->DispatchSubsequents();
 			}, TStatId(), nullptr, ENamedThreads::AnyHiPriThreadHiPriTask);
 		}
 		else
 		{
 			EncryptAndSign(Entry);
-			Entry->FinishEncryptionAndSigningBarrier->DispatchSubsequents();
+			Entry->BeginWriteBarrier->DispatchSubsequents();
 		}
 	}
 
@@ -1581,6 +1588,11 @@ private:
 				FSHA1::HashBuffer(Block.IoBuffer->Data(), Block.Size, Block.Signature.Hash);
 			}
 		}
+		Entry->CompressedSize = 0;
+		for (const FChunkBlock& ChunkBlock : Entry->ChunkBlocks)
+		{
+			Entry->CompressedSize += ChunkBlock.Size;
+		}
 	}
 
 	void WriteEntry(FIoStoreWriteQueueEntry* Entry)
@@ -1839,31 +1851,38 @@ void FIoStoreWriterContextImpl::Flush()
 		AllEntries.Append(IoStoreWriter->Entries);
 	}
 
-	double WritesStart = FPlatformTime::Seconds();
-
-	for (FIoStoreWriteQueueEntry* Entry : AllEntries)
+	// Start scheduler threads, queue all entries, and wait for them to finish
 	{
-		ScheduleCompression(Entry);
-	}
+		double WritesStart = FPlatformTime::Seconds();
 
-	{
-		TRACE_CPUPROFILER_EVENT_SCOPE(WaitForWritesToComplete);
-		for (int32 EntryIndex = AllEntries.Num() - 1; EntryIndex >= 0; --EntryIndex)
+		BeginCompressionThread = Async(EAsyncExecution::Thread, [this]() { BeginCompressionThreadFunc(); });
+		BeginEncryptionAndSigningThread = Async(EAsyncExecution::Thread, [this]() { BeginEncryptionAndSigningThreadFunc(); });
+		WriterThread = Async(EAsyncExecution::Thread, [this]() { WriterThreadFunc(); });
+
+		for (FIoStoreWriteQueueEntry* Entry : AllEntries)
 		{
-			AllEntries[EntryIndex]->WriteFinishedEvent->Wait();
+			ScheduleCompression(Entry);
 		}
-	}
+		BeginCompressionQueue.CompleteAdding();
 
-	double WritesEnd = FPlatformTime::Seconds();
-	double WritesSeconds = FPlatformTime::ToSeconds64(WriteCycleCount.Load());
-	UE_LOG(LogIoStore, Display, TEXT("Writing and compressing took %.2f seconds, writes to disk took %.2f seconds for %s bytes @ %s bytes per second."), 
-		WritesEnd - WritesStart,
-		WritesSeconds,
-		*FText::AsNumber(WriteByteCount.Load()).ToString(),
-		*FText::AsNumber((int64)((double)WriteByteCount.Load() / FMath::Max(.0001f, WritesSeconds))).ToString()
-		);
+		{
+			TRACE_CPUPROFILER_EVENT_SCOPE(WaitForWritesToComplete);
+			for (int32 EntryIndex = AllEntries.Num() - 1; EntryIndex >= 0; --EntryIndex)
+			{
+				AllEntries[EntryIndex]->WriteFinishedEvent->Wait();
+			}
+		}
 
-	AllEntries.Empty();
+		double WritesEnd = FPlatformTime::Seconds();
+		double WritesSeconds = FPlatformTime::ToSeconds64(WriteCycleCount.Load());
+		UE_LOG(LogIoStore, Display, TEXT("Writing and compressing took %.2lf seconds, writes to disk took %.2lf seconds for %s bytes @ %s bytes per second."), 
+			WritesEnd - WritesStart,
+			WritesSeconds,
+			*FText::AsNumber(WriteByteCount.Load()).ToString(),
+			*FText::AsNumber((int64)((double)WriteByteCount.Load() / FMath::Max(.0001f, WritesSeconds))).ToString()
+			);
+		AllEntries.Empty();
+	}
 
 	// Classically there were so few writers that this didn't need to be multi threaded, but it
 	// involves writing files, and with content on demand this ends up being thousands of iterations. 
@@ -1898,17 +1917,19 @@ void FIoStoreWriterContextImpl::BeginCompressionThreadFunc()
 		FIoStoreWriteQueueEntry* Entry = BeginCompressionQueue.DequeueOrWait();
 		if (!Entry)
 		{
-			return;
+			break;
 		}
 		while (Entry)
 		{
 			FIoStoreWriteQueueEntry* Next = Entry->Next;
 			Entry->BeginCompressionBarrier->Wait();
+			TRACE_COUNTER_INCREMENT(IoStoreBeginCompressionCount);
 			Entry->Writer->BeginCompress(Entry);
 			BeginEncryptionAndSigningQueue.Enqueue(Entry);
 			Entry = Next;
 		}
 	}
+	BeginEncryptionAndSigningQueue.CompleteAdding();
 }
 
 void FIoStoreWriterContextImpl::BeginEncryptionAndSigningThreadFunc()
@@ -1919,7 +1940,7 @@ void FIoStoreWriterContextImpl::BeginEncryptionAndSigningThreadFunc()
 		FIoStoreWriteQueueEntry* Entry = BeginEncryptionAndSigningQueue.DequeueOrWait();
 		if (!Entry)
 		{
-			return;
+			break;
 		}
 		while (Entry)
 		{
@@ -1939,38 +1960,13 @@ void FIoStoreWriterContextImpl::BeginEncryptionAndSigningThreadFunc()
 					DDC->Put(*Entry->DDCCacheKey, DDCData, Entry->Options.FileName);
 				}
 			}
-			FinishEncryptionAndSigningQueue.Enqueue(Entry);
+			TRACE_COUNTER_INCREMENT(IoStoreBeginEncryptionAndSigningCount);
 			Entry->Writer->BeginEncryptAndSign(Entry);
-			Entry = Next;
-		}
-	}
-}
-
-void FIoStoreWriterContextImpl::FinishEncryptionAndSigningThreadFunc()
-{
-	TRACE_CPUPROFILER_EVENT_SCOPE(FinishEncryptionAndSigningThread);
-	for (;;)
-	{
-		FIoStoreWriteQueueEntry* Entry = FinishEncryptionAndSigningQueue.DequeueOrWait();
-		if (!Entry)
-		{
-			return;
-		}
-		while (Entry)
-		{
-			FIoStoreWriteQueueEntry* Next = Entry->Next;
-			Entry->FinishEncryptionAndSigningBarrier->Wait();
 			WriterQueue.Enqueue(Entry);
-			Entry->CompressedSize = 0;
-			for (const FChunkBlock& ChunkBlock : Entry->ChunkBlocks)
-			{
-				Entry->CompressedSize += ChunkBlock.Size;
-			}
-			Entry->BeginWriteBarrier->DispatchSubsequents();
-
 			Entry = Next;
 		}
 	}
+	WriterQueue.CompleteAdding();
 }
 
 void FIoStoreWriterContextImpl::WriterThreadFunc()
@@ -1987,6 +1983,7 @@ void FIoStoreWriterContextImpl::WriterThreadFunc()
 		{
 			FIoStoreWriteQueueEntry* Next = Entry->Next;
 			Entry->BeginWriteBarrier->Wait();
+			TRACE_COUNTER_INCREMENT(IoStoreBeginWriteCount);
 			Entry->Writer->WriteEntry(Entry);
 			Entry->WriteFinishedEvent->DispatchSubsequents();
 			Entry = Next;
