GitHub 11732 : Fix bug where ContainerFlags may not be copied when creating an IoContainer

#11732

diff --git a/Engine/Source/Developer/IoStoreUtilities/Private/IoStoreUtilities.cpp b/Engine/Source/Developer/IoStoreUtilities/Private/IoStoreUtilities.cpp
--- a/Engine/Source/Developer/IoStoreUtilities/Private/IoStoreUtilities.cpp
+++ b/Engine/Source/Developer/IoStoreUtilities/Private/IoStoreUtilities.cpp
@@ -5430,9 +5430,10 @@ int32 CreateTarget(const FIoStoreArguments& Arguments, const FIoStoreWriterSetti
 			{
 				FIoContainerSettings ContainerSettings;
 				ContainerSettings.ContainerId = ContainerTarget->ContainerId;
+				ContainerSettings.ContainerFlags = ContainerTarget->ContainerFlags; 
 				if (Arguments.bCreateDirectoryIndex)
 				{
-					ContainerSettings.ContainerFlags = ContainerTarget->ContainerFlags | EIoContainerFlags::Indexed;
+					ContainerSettings.ContainerFlags |= EIoContainerFlags::Indexed;
 				}
 				if (EnumHasAnyFlags(ContainerTarget->ContainerFlags, EIoContainerFlags::Encrypted))
 				{
