UnrealPak: Minor cpu optimizations (by reducing allocations).
Reserve memory upfront for some big containers and pass EAllowShrinking::No when popping/removing items.
Optimize parsing of huge text files in ParsePakResponseFile and ParsePakOrderFile by using more efficient utility functions:
FFileHelper::LoadFileToStringArray -> FFileHelper::LoadFileToStringWithLineVisitor
FParse::Token -> UE::String::ParseTokens

#rnx
#tests binary diff of content/paks folder

diff --git a/Engine/Source/Developer/IoStoreUtilities/Private/IoStoreUtilities.cpp b/Engine/Source/Developer/IoStoreUtilities/Private/IoStoreUtilities.cpp
--- a/Engine/Source/Developer/IoStoreUtilities/Private/IoStoreUtilities.cpp
+++ b/Engine/Source/Developer/IoStoreUtilities/Private/IoStoreUtilities.cpp
@@ -78,6 +78,7 @@
 #include "HAL/FileManagerGeneric.h"
 #include "Serialization/CompactBinarySerialization.h"
 #include "Serialization/ZenPackageHeader.h"
+#include "String/ParseTokens.h"
 
 IMPLEMENT_MODULE(FDefaultModuleImpl, IoStoreUtilities);
 
@@ -1470,7 +1471,7 @@ void SortPackagesInLoadOrderRecursive(TArray<FCookedPackage*>& Result, FCookedPa
 				{
 					while (P.Top()->PreOrderNumber > Parent->PreOrderNumber)
 					{
-						P.Pop();
+						P.Pop(EAllowShrinking::No);
 					}
 				}
 			}
@@ -1482,11 +1483,11 @@ void SortPackagesInLoadOrderRecursive(TArray<FCookedPackage*>& Result, FCookedPa
 		do
 		{
 			InStronglyConnectedComponent = S.Top();
-			S.Pop();
+			S.Pop(EAllowShrinking::No);
 			InStronglyConnectedComponent->bPermanentMark = true;
 			Result.Add(InStronglyConnectedComponent);
 		} while (InStronglyConnectedComponent != Package);
-		P.Pop();
+		P.Pop(EAllowShrinking::No);
 	}
 }
 
@@ -1499,6 +1500,7 @@ void SortPackagesInLoadOrder(TArray<FCookedPackage*>& Packages, const TMap<FPack
 		});
 
 	TMap<FCookedPackage*, TArray<FCookedPackage*>> ReverseEdges;
+	ReverseEdges.Reserve(PackagesMap.Num());
 	for (FCookedPackage* Package : Packages)
 	{
 		for (FPackageId ImportedPackageId : Package->PackageStoreEntry.ImportedPackageIds)
@@ -1523,7 +1525,6 @@ void SortPackagesInLoadOrder(TArray<FCookedPackage*>& Packages, const TMap<FPack
 	// Path based strongly connected components + topological sort of the components
 	TArray<FCookedPackage*> Result;
 	Result.Reserve(Packages.Num());
-	TArray<int32> PackagePreOrderNumbers;
 	TArray<FCookedPackage*> S;
 	TArray<FCookedPackage*> P;
 	for (FCookedPackage* Package : Packages)
@@ -1629,10 +1630,6 @@ static void AssignPackagesDiskOrder(
 		}
 	};
 
-	TArray<FCluster*> Clusters;
-	TSet<FCookedPackage*> AssignedPackages;
-	TArray<FCookedPackage*> ProcessStack;
-
 	struct FPackageAndOrder
 	{
 		FCookedPackage* Package = nullptr;
@@ -1731,6 +1728,11 @@ static void AssignPackagesDiskOrder(
 
 	int32 ClusterSequence = 0;
 	TMap<FCookedPackage*, FCluster*> PackageToCluster;
+	TArray<FCluster*> Clusters;
+	TSet<FCookedPackage*> AssignedPackages;
+	TArray<FCookedPackage*> ProcessStack;
+	PackageToCluster.Reserve(SortedPackages.Num());
+	AssignedPackages.Reserve(SortedPackages.Num());
 	for (FPackageAndOrder& Entry : SortedPackages)
 	{
 		checkSlow(Entry.OrderMap); // Without this, Entry.OrderMap != LastBlameOrderMap convinces static analysis that Entry.OrderMap may be null
@@ -1768,9 +1770,7 @@ static void AssignPackagesDiskOrder(
 					AssignedUExpSize += PackageToProcess->UExpSize;
 					AssignedBulkSize += PackageToProcess->TotalBulkDataSize;
 					
-					TArray<FPackageId> AllReferencedPackageIds;
-					AllReferencedPackageIds.Append(PackageToProcess->PackageStoreEntry.ImportedPackageIds);
-					for (const FPackageId& ReferencedPackageId : AllReferencedPackageIds)
+					for (const FPackageId& ReferencedPackageId : PackageToProcess->PackageStoreEntry.ImportedPackageIds)
 					{
 						FCookedPackage* FindReferencedPackage = PackageIdMap.FindRef(ReferencedPackageId);
 						if (FindReferencedPackage)
@@ -1792,7 +1792,7 @@ static void AssignPackagesDiskOrder(
 				DepQueue.Push(Package);
 				while (DepQueue.Num() > 0)
 				{
-					FCookedPackage* Cursor = DepQueue.Pop();
+					FCookedPackage* Cursor = DepQueue.Pop(EAllowShrinking::No);
 					if( VisitedDeps.Contains(Cursor) == false)
 					{
 						VisitedDeps.Add(Cursor);
@@ -3069,6 +3069,17 @@ void InitializeContainerTargetsAndPackages(
 		return true;
 	};
 	
+	// Reserve memory for lookup maps
+	{
+		int32 NumSourceFiles = 0;
+		for (const FContainerSourceSpec& ContainerSource : Arguments.Containers)
+		{
+			NumSourceFiles += ContainerSource.SourceFiles.Num();
+		}
+		PackageNameMap.Reserve(NumSourceFiles);
+		PackageIdMap.Reserve(NumSourceFiles);
+	}
+
 	for (const FContainerSourceSpec& ContainerSource : Arguments.Containers)
 	{
 		FContainerTargetSpec* ContainerTarget = AddContainer(ContainerSource.Name, ContainerTargets);
@@ -9065,41 +9076,60 @@ bool SignIoStoreContainer(const TCHAR* InContainerFilename, const FRSAKeyHandle
 
 static bool ParsePakResponseFile(const TCHAR* FilePath, TArray<FContainerSourceFile>& OutFiles)
 {
-	TArray<FString> ResponseFileContents;
-	if (!FFileHelper::LoadFileToStringArray(ResponseFileContents, FilePath))
-	{
-		UE_LOG(LogIoStore, Error, TEXT("Failed to read response file '%s'."), FilePath);
-		return false;
-	}
+	TRACE_CPUPROFILER_EVENT_SCOPE(ParsePakResponseFile);
 
-	for (const FString& ResponseLine : ResponseFileContents)
+	TArray<FStringView> Tokens;
+	TArray<FString> SourceAndDest;
+
+	bool bParseError = false;
+	bool bReadSuccess = FFileHelper::LoadFileToStringWithLineVisitor(FilePath,
+		[&Tokens, &SourceAndDest, &bParseError, &OutFiles](FStringView ResponseLine)
 	{
-		TArray<FString> SourceAndDest;
-		TArray<FString> Switches;
+		Tokens.Reset();
+		SourceAndDest.Reset();
+
+		bool bNeedsCompression = false;
+		bool bNeedsEncryption = false;
 
-		FString NextToken;
-		const TCHAR* ResponseLinePtr = *ResponseLine;
-		while (FParse::Token(ResponseLinePtr, NextToken, false))
+		UE::String::ParseTokens(ResponseLine, ' ', Tokens,
+			UE::String::EParseTokensOptions::Trim | UE::String::EParseTokensOptions::SkipEmpty);
+
+		for (FStringView Token : Tokens)
 		{
-			if ((**NextToken == TCHAR('-')))
+			// Parse switches
+			if (Token.StartsWith('-'))
 			{
-				new(Switches) FString(NextToken.Mid(1));
+				if (Token == TEXT("-compress"))
+				{
+					bNeedsCompression = true;
+				}
+				else if (Token == TEXT("-encrypt"))
+				{
+					bNeedsEncryption = true;
+				}
 			}
+			// Parse SourceAndDest path arguments
 			else
 			{
-				new(SourceAndDest) FString(NextToken);
+				if (Token.Len() > 1 && Token.StartsWith('"') && Token.EndsWith('"'))
+				{
+					Token = Token.Mid(1, Token.Len() - 2);
+				}
+				SourceAndDest.Emplace(FString(Token));
 			}
 		}
 
 		if (SourceAndDest.Num() == 0)
 		{
-			continue;
+			return;
 		}
 
 		if (SourceAndDest.Num() != 2)
 		{
-			UE_LOG(LogIoStore, Error, TEXT("Invalid line in response file '%s'."), *ResponseLine);
-			return false;
+			UE_LOG(LogIoStore, Error, TEXT("Invalid line in response file '%.*s'."),
+				ResponseLine.Len(), ResponseLine.GetData());
+			bParseError = true;
+			return;
 		}
 
 		FPaths::NormalizeFilename(SourceAndDest[0]);
@@ -9107,51 +9137,52 @@ static bool ParsePakResponseFile(const TCHAR* FilePath, TArray<FContainerSourceF
 		FContainerSourceFile& FileEntry = OutFiles.AddDefaulted_GetRef();
 		FileEntry.NormalizedPath = MoveTemp(SourceAndDest[0]);
 		FileEntry.DestinationPath = MoveTemp(SourceAndDest[1]);
+		FileEntry.bNeedsCompression = bNeedsCompression;
+		FileEntry.bNeedsEncryption = bNeedsEncryption;
+	});
 
-		for (int32 Index = 0; Index < Switches.Num(); ++Index)
-		{
-			if (Switches[Index] == TEXT("compress"))
-			{
-				FileEntry.bNeedsCompression = true;
-			}
-			if (Switches[Index] == TEXT("encrypt"))
-			{
-				FileEntry.bNeedsEncryption = true;
-			}
-		}
-	}
-	return true;
+	const bool bSuccess = bReadSuccess && !bParseError;
+	return bSuccess;
 }
 
 static bool ParsePakOrderFile(const TCHAR* FilePath, FFileOrderMap& Map, const FIoStoreArguments& Arguments)
 {
 	IOSTORE_CPU_SCOPE(ParsePakOrderFile);
 
-	TArray<FString> OrderFileContents;
-	if (!FFileHelper::LoadFileToStringArray(OrderFileContents, FilePath))
-	{
-		UE_LOG(LogIoStore, Error, TEXT("Failed to read order file '%s'."), FilePath);
-		return false;
-	}
-
 	Map.Name = FPaths::GetCleanFilename(FilePath);
 	UE_LOG(LogIoStore, Display, TEXT("Order file %s (short name %s) priority %d"), FilePath, *Map.Name, Map.Priority);
+
+	TArray<FStringView> Tokens;
+	FString FullFileName;
 	int64 NextOrder = 0;
-	for (const FString& OrderLine : OrderFileContents)
+
+	bool bParseError = false;
+	bool bReadSuccess = FFileHelper::LoadFileToStringWithLineVisitor(FilePath,
+		[&Tokens, &FullFileName, &NextOrder, &bParseError, &Map, &Arguments](FStringView OrderLine)
 	{
-		const TCHAR* OrderLinePtr = *OrderLine;
-		FString PackageName;
+		Tokens.Reset();
 
 		// Skip comments
-		if (FCString::Strncmp(OrderLinePtr, TEXT("#"), 1) == 0 || FCString::Strncmp(OrderLinePtr, TEXT("//"), 2) == 0)
+		if (OrderLine.StartsWith('#') || OrderLine.StartsWith(TEXTVIEW("//")))
 		{
-			continue;
+			return;
 		}
 
-		if (!FParse::Token(OrderLinePtr, PackageName, false))
+		UE::String::ParseTokens(OrderLine, ' ', Tokens,
+			UE::String::EParseTokensOptions::Trim | UE::String::EParseTokensOptions::SkipEmpty);
+
+		if (Tokens.Num() == 0)
 		{
-			UE_LOG(LogIoStore, Error, TEXT("Invalid line in order file '%s'."), *OrderLine);
-			return false;
+			UE_LOG(LogIoStore, Error, TEXT("Invalid line in order file '%.*s'."),
+				OrderLine.Len(), OrderLine.GetData());
+			bParseError = true;
+			return;
+		}
+
+		FStringView PackageName = Tokens[0];
+		if (PackageName.Len() > 1 && PackageName.StartsWith('"') && PackageName.EndsWith('"'))
+		{
+			PackageName = PackageName.Mid(1, PackageName.Len() - 2);
 		}
 
 		FName PackageFName;
@@ -9159,9 +9190,9 @@ static bool ParsePakOrderFile(const TCHAR* FilePath, FFileOrderMap& Map, const F
 		{
 			PackageFName = FName(PackageName);
 		}
-		else if (PackageName.StartsWith(TEXT("../../../")))
+		else if (PackageName.StartsWith(TEXTVIEW("../../../")))
 		{
-			FString FullFileName = FPaths::Combine(Arguments.CookedDir, PackageName.RightChop(9));
+			FullFileName = FPaths::Combine(Arguments.CookedDir, PackageName.RightChop(9));
 			FPaths::NormalizeFilename(FullFileName);
 			PackageFName = Arguments.PackageStore->GetPackageNameFromFileName(FullFileName);
 		}
@@ -9170,6 +9201,12 @@ static bool ParsePakOrderFile(const TCHAR* FilePath, FFileOrderMap& Map, const F
 		{
 			Map.PackageNameToOrder.Emplace(PackageFName, NextOrder++);
 		}
+	});
+
+	if (!bReadSuccess || bParseError)
+	{
+		UE_LOG(LogIoStore, Error, TEXT("Failed to read order file '%s'."), FilePath);
+		return false;
 	}
 
 	UE_LOG(LogIoStore, Display, TEXT("Order file %s (short name %s) contained %d valid entries"), FilePath, *Map.Name, Map.PackageNameToOrder.Num());
