UnrealPak: Revert ParsePakResponseFile and ParsePakOrderFile optimizations in 33258481 to fix parsing errors for file paths with spaces in them.

#rnx
#tests binary diff of content/paks folder

diff --git a/Engine/Source/Developer/IoStoreUtilities/Private/IoStoreUtilities.cpp b/Engine/Source/Developer/IoStoreUtilities/Private/IoStoreUtilities.cpp
--- a/Engine/Source/Developer/IoStoreUtilities/Private/IoStoreUtilities.cpp
+++ b/Engine/Source/Developer/IoStoreUtilities/Private/IoStoreUtilities.cpp
@@ -9078,58 +9078,41 @@ static bool ParsePakResponseFile(const TCHAR* FilePath, TArray<FContainerSourceF
 {
 	TRACE_CPUPROFILER_EVENT_SCOPE(ParsePakResponseFile);
 
-	TArray<FStringView> Tokens;
-	TArray<FString> SourceAndDest;
-
-	bool bParseError = false;
-	bool bReadSuccess = FFileHelper::LoadFileToStringWithLineVisitor(FilePath,
-		[&Tokens, &SourceAndDest, &bParseError, &OutFiles](FStringView ResponseLine)
+	TArray<FString> ResponseFileContents;
+	if (!FFileHelper::LoadFileToStringArray(ResponseFileContents, FilePath))
 	{
-		Tokens.Reset();
-		SourceAndDest.Reset();
-
-		bool bNeedsCompression = false;
-		bool bNeedsEncryption = false;
+		UE_LOG(LogIoStore, Error, TEXT("Failed to read response file '%s'."), FilePath);
+		return false;
+	}
 
-		UE::String::ParseTokens(ResponseLine, ' ', Tokens,
-			UE::String::EParseTokensOptions::Trim | UE::String::EParseTokensOptions::SkipEmpty);
+	for (const FString& ResponseLine : ResponseFileContents)
+	{
+		TArray<FString> SourceAndDest;
+		TArray<FString> Switches;
 
-		for (FStringView Token : Tokens)
+		FString NextToken;
+		const TCHAR* ResponseLinePtr = *ResponseLine;
+		while (FParse::Token(ResponseLinePtr, NextToken, false))
 		{
-			// Parse switches
-			if (Token.StartsWith('-'))
+			if ((**NextToken == TCHAR('-')))
 			{
-				if (Token == TEXT("-compress"))
-				{
-					bNeedsCompression = true;
-				}
-				else if (Token == TEXT("-encrypt"))
-				{
-					bNeedsEncryption = true;
-				}
+				new(Switches) FString(NextToken.Mid(1));
 			}
-			// Parse SourceAndDest path arguments
 			else
 			{
-				if (Token.Len() > 1 && Token.StartsWith('"') && Token.EndsWith('"'))
-				{
-					Token = Token.Mid(1, Token.Len() - 2);
-				}
-				SourceAndDest.Emplace(FString(Token));
+				new(SourceAndDest) FString(NextToken);
 			}
 		}
 
 		if (SourceAndDest.Num() == 0)
 		{
-			return;
+			continue;
 		}
 
 		if (SourceAndDest.Num() != 2)
 		{
-			UE_LOG(LogIoStore, Error, TEXT("Invalid line in response file '%.*s'."),
-				ResponseLine.Len(), ResponseLine.GetData());
-			bParseError = true;
-			return;
+			UE_LOG(LogIoStore, Error, TEXT("Invalid line in response file '%s'."), *ResponseLine);
+			return false;
 		}
 
 		FPaths::NormalizeFilename(SourceAndDest[0]);
@@ -9137,52 +9120,51 @@ static bool ParsePakResponseFile(const TCHAR* FilePath, TArray<FContainerSourceF
 		FContainerSourceFile& FileEntry = OutFiles.AddDefaulted_GetRef();
 		FileEntry.NormalizedPath = MoveTemp(SourceAndDest[0]);
 		FileEntry.DestinationPath = MoveTemp(SourceAndDest[1]);
-		FileEntry.bNeedsCompression = bNeedsCompression;
-		FileEntry.bNeedsEncryption = bNeedsEncryption;
-	});
 
-	const bool bSuccess = bReadSuccess && !bParseError;
-	return bSuccess;
+		for (int32 Index = 0; Index < Switches.Num(); ++Index)
+		{
+			if (Switches[Index] == TEXT("compress"))
+			{
+				FileEntry.bNeedsCompression = true;
+			}
+			if (Switches[Index] == TEXT("encrypt"))
+			{
+				FileEntry.bNeedsEncryption = true;
+			}
+		}
+	}
+	return true;
 }
 
 static bool ParsePakOrderFile(const TCHAR* FilePath, FFileOrderMap& Map, const FIoStoreArguments& Arguments)
 {
 	IOSTORE_CPU_SCOPE(ParsePakOrderFile);
 
+	TArray<FString> OrderFileContents;
+	if (!FFileHelper::LoadFileToStringArray(OrderFileContents, FilePath))
+	{
+		UE_LOG(LogIoStore, Error, TEXT("Failed to read order file '%s'."), FilePath);
+		return false;
+	}
+
 	Map.Name = FPaths::GetCleanFilename(FilePath);
 	UE_LOG(LogIoStore, Display, TEXT("Order file %s (short name %s) priority %d"), FilePath, *Map.Name, Map.Priority);
-
-	TArray<FStringView> Tokens;
-	FString FullFileName;
 	int64 NextOrder = 0;
-
-	bool bParseError = false;
-	bool bReadSuccess = FFileHelper::LoadFileToStringWithLineVisitor(FilePath,
-		[&Tokens, &FullFileName, &NextOrder, &bParseError, &Map, &Arguments](FStringView OrderLine)
+	for (const FString& OrderLine : OrderFileContents)
 	{
-		Tokens.Reset();
+		const TCHAR* OrderLinePtr = *OrderLine;
+		FString PackageName;
 
 		// Skip comments
-		if (OrderLine.StartsWith('#') || OrderLine.StartsWith(TEXTVIEW("//")))
-		{
-			return;
-		}
-
-		UE::String::ParseTokens(OrderLine, ' ', Tokens,
-			UE::String::EParseTokensOptions::Trim | UE::String::EParseTokensOptions::SkipEmpty);
-
-		if (Tokens.Num() == 0)
+		if (FCString::Strncmp(OrderLinePtr, TEXT("#"), 1) == 0 || FCString::Strncmp(OrderLinePtr, TEXT("//"), 2) == 0)
 		{
-			UE_LOG(LogIoStore, Error, TEXT("Invalid line in order file '%.*s'."),
-				OrderLine.Len(), OrderLine.GetData());
-			bParseError = true;
-			return;
+			continue;
 		}
 
-		FStringView PackageName = Tokens[0];
-		if (PackageName.Len() > 1 && PackageName.StartsWith('"') && PackageName.EndsWith('"'))
+		if (!FParse::Token(OrderLinePtr, PackageName, false))
 		{
-			PackageName = PackageName.Mid(1, PackageName.Len() - 2);
+			UE_LOG(LogIoStore, Error, TEXT("Invalid line in order file '%s'."), *OrderLine);
+			return false;
 		}
 
 		FName PackageFName;
@@ -9190,9 +9172,9 @@ static bool ParsePakOrderFile(const TCHAR* FilePath, FFileOrderMap& Map, const F
 		{
 			PackageFName = FName(PackageName);
 		}
-		else if (PackageName.StartsWith(TEXTVIEW("../../../")))
+		else if (PackageName.StartsWith(TEXT("../../../")))
 		{
-			FullFileName = FPaths::Combine(Arguments.CookedDir, PackageName.RightChop(9));
+			FString FullFileName = FPaths::Combine(Arguments.CookedDir, PackageName.RightChop(9));
 			FPaths::NormalizeFilename(FullFileName);
 			PackageFName = Arguments.PackageStore->GetPackageNameFromFileName(FullFileName);
 		}
@@ -9201,12 +9183,6 @@ static bool ParsePakOrderFile(const TCHAR* FilePath, FFileOrderMap& Map, const F
 		{
 			Map.PackageNameToOrder.Emplace(PackageFName, NextOrder++);
 		}
-	});
-
-	if (!bReadSuccess || bParseError)
-	{
-		UE_LOG(LogIoStore, Error, TEXT("Failed to read order file '%s'."), FilePath);
-		return false;
 	}
 
 	UE_LOG(LogIoStore, Display, TEXT("Order file %s (short name %s) contained %d valid entries"), FilePath, *Map.Name, Map.PackageNameToOrder.Num());
