UnrealPak: Add UsedPhysical and UsedVirtual memory stats to progress output logging and always log at least one status line.

#rnx

diff --git a/Engine/Source/Developer/IoStoreUtilities/Private/IoStoreUtilities.cpp b/Engine/Source/Developer/IoStoreUtilities/Private/IoStoreUtilities.cpp
--- a/Engine/Source/Developer/IoStoreUtilities/Private/IoStoreUtilities.cpp
+++ b/Engine/Source/Developer/IoStoreUtilities/Private/IoStoreUtilities.cpp
@@ -9,6 +9,7 @@
 #include "CookMetadata.h"
 #include "HAL/FileManager.h"
 #include "HAL/PlatformFileManager.h"
+#include "HAL/PlatformMemory.h"
 #include "Hash/CityHash.h"
 #include "Hash/xxhash.h"
 #include "Interfaces/ITargetPlatformManagerModule.h"
@@ -5679,18 +5680,22 @@ int32 CreateTarget(const FIoStoreArguments& Arguments, const FIoStoreWriterSetti
 			IoStoreWriterContext->Flush();
 		});
 
-		while (!FlushTask.IsReady())
+		bool bPrintProgress = true;
+		while (bPrintProgress)
 		{
 			FlushTask.WaitFor(FTimespan::FromSeconds(2.0));
-			FIoStoreWriterContext::FProgress Progress = IoStoreWriterContext->GetProgress();
+			bPrintProgress = !FlushTask.IsReady();
+
 			TStringBuilder<1024> ProgressStringBuilder;
+
+			const FIoStoreWriterContext::FProgress Progress = IoStoreWriterContext->GetProgress();
 			if (Progress.SerializedChunksCount >= Progress.TotalChunksCount)
 			{
 				ProgressStringBuilder.Appendf(TEXT("Writing tocs..."));
 			}
 			else if (Progress.SerializedChunksCount)
 			{
-				ProgressStringBuilder.Appendf(TEXT("Writing chunks (%llu/%llu)..."), Progress.SerializedChunksCount, Progress.TotalChunksCount);
+				ProgressStringBuilder.Appendf(TEXT("Writing chunks %llu/%llu..."), Progress.SerializedChunksCount, Progress.TotalChunksCount);
 				if (Progress.CompressedChunksCount)
 				{
 					ProgressStringBuilder.Appendf(TEXT(" [%llu compressed]"), Progress.CompressedChunksCount);
@@ -5699,12 +5704,17 @@ int32 CreateTarget(const FIoStoreArguments& Arguments, const FIoStoreWriterSetti
 				{
 					ProgressStringBuilder.Appendf(TEXT(" [%llu compression tasks scheduled]"), Progress.ScheduledCompressionTasksCount);
 				}
-				UE_LOG(LogIoStore, Display, TEXT("%s"), *ProgressStringBuilder);
 			}
 			else
 			{
-			UE_LOG(LogIoStore, Display, TEXT("Hashing chunks (%llu/%llu)..."), Progress.HashedChunksCount, Progress.TotalChunksCount);
+				ProgressStringBuilder.Appendf(TEXT("Hashing %llu/%llu..."), Progress.HashedChunksCount, Progress.TotalChunksCount);
 			}
+
+			const FPlatformMemoryStats MemStats = FPlatformMemory::GetStats();
+			ProgressStringBuilder.Appendf(TEXT(" [Physical: %llu MiB, Virtual: %llu MiB]"),
+				MemStats.UsedPhysical >> 20, MemStats.UsedVirtual >> 20);
+
+			UE_LOG(LogIoStore, Display, TEXT("%s"), *ProgressStringBuilder);
 		}
 	}
 
