UnrealPak: Change memory throttling logic to prepare for async DDC get requests that can allocate compression buffers from worker threads.
Add IIoStoreWriteRequest::GetSourceBufferSizeEstimate() to get an upfront approximation of required memory for source reads and compression buffers.
Limit memory usage upfront in FIoStoreWriterContextImpl::ScheduleEntry() based on the estimate instead of during the actual allocations in the middle of the pipeline.
Remove wait from AllocCompressionBuffer and always let it return immediately with a valid compression buffer.
Rename FIoStoreWriteQueueEntry::NumChunkBlocksFromRefDb to NumChunkBlocks and use it for all entries.
Replace the unnecessary FIoStoreWriteQueueEntry::WriteFinishedEvent with a single WriterThread.Wait().
In FIoStoreWriterContextImpl::WriteEntry() reduce som memory growth by releasing the ChunkBlocks array and moving the filename when building the toc directory index.
Bump compression memory limit from 2 GB to 3 GB to reduce the need to (mostly) drain the queue to schedule large-ish entries.

#rnx
#tests identical binary diff of /content/paks folder

diff --git a/Engine/Source/Developer/IoStoreUtilities/Private/IoStoreLooseFiles.cpp b/Engine/Source/Developer/IoStoreUtilities/Private/IoStoreLooseFiles.cpp
--- a/Engine/Source/Developer/IoStoreUtilities/Private/IoStoreLooseFiles.cpp
+++ b/Engine/Source/Developer/IoStoreUtilities/Private/IoStoreLooseFiles.cpp
@@ -110,6 +110,11 @@ public:
 				return nullptr;
 			}
 
+			virtual uint64 GetSourceBufferSizeEstimate() override
+			{
+				return SourceBuffer.DataSize();
+			}
+
 			FIoBuffer SourceBuffer;
 			uint64 OrderHint;
 		};
diff --git a/Engine/Source/Developer/IoStoreUtilities/Private/IoStoreUtilities.cpp b/Engine/Source/Developer/IoStoreUtilities/Private/IoStoreUtilities.cpp
--- a/Engine/Source/Developer/IoStoreUtilities/Private/IoStoreUtilities.cpp
+++ b/Engine/Source/Developer/IoStoreUtilities/Private/IoStoreUtilities.cpp
@@ -329,7 +329,7 @@ public:
 
 	// Returns whether we expect to be able to load the chunk from the reference chunk database.
 	// This can be called from any thread, though in the presence of existing hashes it's single threaded.
-	virtual bool ChunkExists(const FIoContainerId& InContainerId, const FIoHash& InChunkHash, const FIoChunkId& InChunkId, uint32& OutNumChunkBlocks)
+	virtual bool ChunkExists(const FIoContainerId& InContainerId, const FIoHash& InChunkHash, const FIoChunkId& InChunkId, int32& OutNumChunkBlocks)
 	{
 		if (!bValid)
 		{
@@ -404,7 +404,7 @@ public:
 		// We match.
 		MatchCount.fetch_add(1, std::memory_order_relaxed);
 		ReaderChunks->UsedChunksByType[ChunkType].fetch_add(1, std::memory_order_relaxed);
-		OutNumChunkBlocks = ChunkInfo.NumCompressedBlocks;
+		OutNumChunkBlocks = IntCastChecked<int32>(ChunkInfo.NumCompressedBlocks);
 		return true;
 	}
 
@@ -3414,11 +3414,14 @@ private:
 		
 		virtual void FreeSourceBuffer() override
 		{
-			SourceBuffer = FIoBuffer();
-			Manager.OnBufferMemoryFreed(SourceBufferSize);
+			if (SourceBuffer.DataSize() > 0)
+			{
+				SourceBuffer = FIoBuffer();
+				Manager.OnBufferMemoryFreed(SourceBufferSize);
+			}
 		}
 
-		uint64 GetSourceBufferSize() const
+		virtual uint64 GetSourceBufferSizeEstimate() override
 		{
 			return SourceBufferSize;
 		}
@@ -3483,7 +3486,7 @@ private:
 		virtual void LoadSourceBufferAsync() override
 		{
 			Manager.LooseFileSourceReads[(int8)TargetFile.ChunkId.GetChunkType()].IncrementExchange();
-			SourceBuffer = FIoBuffer(GetSourceBufferSize());
+			SourceBuffer = FIoBuffer(GetSourceBufferSizeEstimate());
 
 			QueueEntry->FileHandle.Reset(
 				FPlatformFileManager::Get().GetPlatformFile().OpenAsyncRead(*TargetFile.NormalizedSourcePath));
@@ -3647,7 +3650,7 @@ private:
 
 	void Start(FQueueEntry* QueueEntry)
 	{
-		const uint64 SourceBufferSize = QueueEntry->WriteRequest->GetSourceBufferSize();
+		const uint64 SourceBufferSize = QueueEntry->WriteRequest->GetSourceBufferSizeEstimate();
 
 		uint64 LocalUsedBufferMemory = UsedBufferMemory.Load();
 		while (LocalUsedBufferMemory > 0 && LocalUsedBufferMemory + SourceBufferSize > BufferMemoryLimit)
diff --git a/Engine/Source/Developer/IoStoreUtilities/Private/IoStoreWriter.cpp b/Engine/Source/Developer/IoStoreUtilities/Private/IoStoreWriter.cpp
--- a/Engine/Source/Developer/IoStoreUtilities/Private/IoStoreWriter.cpp
+++ b/Engine/Source/Developer/IoStoreUtilities/Private/IoStoreWriter.cpp
@@ -20,8 +20,8 @@
 #include "Serialization/MemoryWriter.h"
 #include "Templates/UniquePtr.h"
 
-TRACE_DECLARE_MEMORY_COUNTER(IoStoreCompressionUsedBufferMemory, TEXT("IoStoreWriter/CompressionUsedBufferMemory"));
-TRACE_DECLARE_INT_COUNTER(IoStoreCompressionAvailableBuffers, TEXT("IoStoreWriter/CompressionAvailableBuffers"));
+TRACE_DECLARE_MEMORY_COUNTER(IoStoreCompressionMemoryUsed, TEXT("IoStoreWriter/CompressionMemoryUsed"));
+TRACE_DECLARE_MEMORY_COUNTER(IoStoreCompressionMemoryScheduled, TEXT("IoStoreWriter/CompressionMemoryScheduled"));
 TRACE_DECLARE_ATOMIC_INT_COUNTER(IoStoreCompressionInflight, TEXT("IoStoreWriter/CompressionInflight"));
 TRACE_DECLARE_ATOMIC_INT_COUNTER(IoStoreRefDbInflight, TEXT("IoStoreWriter/RefDbInFlight"));
 TRACE_DECLARE_ATOMIC_INT_COUNTER(IoStoreRefDbDone, TEXT("IoStoreWriter/RefDbDone"));
@@ -56,11 +56,12 @@ struct FIoStoreWriteQueueEntry
 	/** Hash of the block data as it would be found on disk after compression and encryption */
 	FIoHash ChunkDiskHash;
 
+	uint64 CompressionMemoryEstimate = 0;
 	uint64 Sequence = 0;
 	
 	// We make this optional because at the latest it might not be valid until FinishCompressionBarrior
 	// completes and we'd like to have a check() on that.
-	TOptional<uint64> UncompressedSize = 0;
+	TOptional<uint64> UncompressedSize;
 
 	// this is not filled out until after encryption completes and *includes the alignment padding for encryption*!
 	// think of this as "size on disk".
@@ -76,10 +77,10 @@ struct FIoStoreWriteQueueEntry
 	FGraphEventRef BeginCompressionBarrier;
 	FGraphEventRef FinishCompressionBarrier;
 	FGraphEventRef BeginWriteBarrier;
-	FGraphEventRef WriteFinishedEvent;
 	TAtomic<int32> CompressedBlocksCount{ 0 };
 	TAtomic<int32> FinishedBlocksCount{ 0 };
 	int32 PartitionIndex = -1;
+	int32 NumChunkBlocks = 0;
 	FString DDCCacheKey;
 	bool bAdded = false;
 	bool bModified = false;
@@ -87,9 +88,6 @@ struct FIoStoreWriteQueueEntry
 	
 	bool bCouldBeFromReferenceDb = false; // Whether the chunk is a valid candidate for the reference db.
 	bool bLoadingFromReferenceDb = false;
-	// When we know we're loading from the reference chunk db, we don't read the source
-	// buffer but we still need to know the number of chunks which we get from the refdb.
-	uint32 NumChunkBlocksFromRefDb = 0; 
 };
 
 class FIoStoreWriteQueue
@@ -173,7 +171,7 @@ private:
 
 class FIoStoreWriterContextImpl
 {
-	static constexpr uint64 DefaultMemoryLimit = 2ull << 30ull;
+	static constexpr uint64 CompressionMemorySoftLimit = 3ull << 30ull;
 public:
 	FIoStoreWriterContextImpl()
 	{
@@ -187,11 +185,6 @@ public:
 		BeginCompressionThread.Wait();
 		BeginEncryptionAndSigningThread.Wait();
 		WriterThread.Wait();
-		if (CompressionBufferAvailableEvent)
-		{
-			FPlatformProcess::ReturnSynchEventToPool(CompressionBufferAvailableEvent);
-		}
-		check(AvailableCompressionBuffers.Num() == TotalCompressionBufferCount);
 		for (FIoBuffer* IoBuffer : AvailableCompressionBuffers)
 		{
 			delete IoBuffer;
@@ -202,7 +195,6 @@ public:
 	{
 		TRACE_CPUPROFILER_EVENT_SCOPE(FIoStoreWriterContext::Initialize);
 		WriterSettings = InWriterSettings;
-		CompressionBufferAvailableEvent = FPlatformProcess::GetSynchEventFromPool(false);
 
 		if (WriterSettings.bCompressionEnableDDC)
 		{
@@ -216,13 +208,12 @@ public:
 		CompressionBufferSize = FMath::Max(CompressionBufferSize, static_cast<int32>(WriterSettings.CompressionBlockSize));
 		CompressionBufferSize = Align(CompressionBufferSize, FAES::AESBlockSize);
 
-		TotalCompressionBufferCount = int32(DefaultMemoryLimit / CompressionBufferSize);
-		AvailableCompressionBuffers.Reserve(TotalCompressionBufferCount);
-		for (int32 BufferIndex = 0; BufferIndex < TotalCompressionBufferCount; ++BufferIndex)
+		const int32 InitialCompressionBufferCount = int32(CompressionMemorySoftLimit / CompressionBufferSize);
+		AvailableCompressionBuffers.Reserve(InitialCompressionBufferCount);
+		for (int32 BufferIndex = 0; BufferIndex < InitialCompressionBufferCount; ++BufferIndex)
 		{
 			AvailableCompressionBuffers.Add(new FIoBuffer(CompressionBufferSize));
 		}
-		TRACE_COUNTER_SET(IoStoreCompressionAvailableBuffers, AvailableCompressionBuffers.Num());
 
 		return FIoStatus::Ok;
 	}
@@ -269,82 +260,40 @@ public:
 		return WriterSettings;
 	}
 
-	void ScheduleCompression(FIoStoreWriteQueueEntry* QueueEntry)
-	{
-		BeginCompressionQueue.Enqueue(QueueEntry);
-		if (QueueEntry->bLoadingFromReferenceDb == false)
-		{
-			QueueEntry->Request->PrepareSourceBufferAsync(QueueEntry->BeginCompressionBarrier);
-		}
-		else
-		{
-			// We don't need to wait on a read so we can kick directly.
-			QueueEntry->BeginCompressionBarrier->DispatchSubsequents();
-		}
-	}
-
-	//
-	// This must be called prior to the consumer being dispatched in order to prevent resource
-	// contention deadlock (even if using retracting waits)
-	//
-	FIoBuffer* AllocCompressionBuffer(int32 TotalEntryChunkBlocksCount)
+	FIoBuffer* AllocCompressionBuffer()
 	{
-		if (TotalEntryChunkBlocksCount > TotalCompressionBufferCount)
-		{
-			TRACE_COUNTER_ADD(IoStoreCompressionUsedBufferMemory, CompressionBufferSize);
-			return new FIoBuffer(CompressionBufferSize);
-		}
 		FIoBuffer* AllocatedBuffer = nullptr;
-		while (!AllocatedBuffer)
 		{
+			FScopeLock Lock(&AvailableCompressionBuffersCritical);
+			if (AvailableCompressionBuffers.Num() > 0)
 			{
-				FScopeLock Lock(&AvailableCompressionBuffersCritical);
-				if (AvailableCompressionBuffers.Num())
-				{
-					AllocatedBuffer = AvailableCompressionBuffers.Pop();
-					TRACE_COUNTER_ADD(IoStoreCompressionUsedBufferMemory, CompressionBufferSize);
-					TRACE_COUNTER_DECREMENT(IoStoreCompressionAvailableBuffers);
-				}
-			}
-			if (!AllocatedBuffer)
-			{
-				TRACE_CPUPROFILER_EVENT_SCOPE(WaitForCompressionBuffer);
-				CompressionBufferAvailableEvent->Wait();
+				AllocatedBuffer = AvailableCompressionBuffers.Pop();
 			}
+			TRACE_COUNTER_ADD(IoStoreCompressionMemoryUsed, CompressionBufferSize);
+		}
+		if (!AllocatedBuffer)
+		{
+			AllocatedBuffer = new FIoBuffer(CompressionBufferSize);
 		}
 		return AllocatedBuffer;
 	}
 
-	void FreeCompressionBuffer(FIoBuffer* Buffer, int32 TotalEntryChunkBlocksCount)
+	void FreeCompressionBuffer(FIoBuffer* Buffer)
 	{
-		if (TotalEntryChunkBlocksCount > TotalCompressionBufferCount)
-		{
-			TRACE_COUNTER_SUBTRACT(IoStoreCompressionUsedBufferMemory, CompressionBufferSize);
-			delete Buffer;
-			return;
-		}
-		bool bTriggerEvent;
-		{
-			FScopeLock Lock(&AvailableCompressionBuffersCritical);
-			bTriggerEvent = AvailableCompressionBuffers.Num() == 0;
-			AvailableCompressionBuffers.Push(Buffer);
-			TRACE_COUNTER_SUBTRACT(IoStoreCompressionUsedBufferMemory, CompressionBufferSize);
-			TRACE_COUNTER_INCREMENT(IoStoreCompressionAvailableBuffers);
-		}
-		if (bTriggerEvent)
-		{
-			CompressionBufferAvailableEvent->Trigger();
-		}
+		FScopeLock Lock(&AvailableCompressionBuffersCritical);
+		AvailableCompressionBuffers.Push(Buffer);
+		TRACE_COUNTER_SUBTRACT(IoStoreCompressionMemoryUsed, CompressionBufferSize);
 	}
 
 private:
+	void ScheduleEntry(FIoStoreWriteQueueEntry* Entry);
 	void BeginCompressionThreadFunc();
 	void BeginEncryptionAndSigningThreadFunc();
 	void WriterThreadFunc();
 
 	FIoStoreWriterSettings WriterSettings;
 	FDerivedDataCacheInterface* DDC = nullptr;
-	FEvent* CompressionBufferAvailableEvent = nullptr;
+	FEventRef CompressionMemoryReleasedEvent;
 	TFuture<void> BeginCompressionThread;
 	TFuture<void> BeginEncryptionAndSigningThread;
 	TFuture<void> WriterThread;
@@ -366,10 +315,10 @@ private:
 	TAtomic<uint64> ScheduledCompressionTasksCount{ 0 };
 	TAtomic<uint64> CompressionDDCHitCount{ 0 };
 	TAtomic<uint64> CompressionDDCMissCount{ 0 };
+	TAtomic<uint64> ScheduledCompressionMemory{ 0 };
 	FCriticalSection AvailableCompressionBuffersCritical;
 	TArray<FIoBuffer*> AvailableCompressionBuffers;
 	int32 CompressionBufferSize = -1;
-	int32 TotalCompressionBufferCount = -1;
 	TArray<TSharedPtr<FIoStoreWriter>> IoStoreWriters;
 
 	friend class FIoStoreWriter;
@@ -460,9 +409,9 @@ public:
 		return 1 + uint8(Toc.CompressionMethods.Add(CompressionMethod));
 	}
 
-	void AddToFileIndex(const FIoChunkId& ChunkId, const FString& FileName)
+	void AddToFileIndex(const FIoChunkId& ChunkId, FString&& FileName)
 	{
-		ChunkIdToFileName.Emplace(ChunkId, FileName);
+		ChunkIdToFileName.Emplace(ChunkId, MoveTemp(FileName));
 	}
 
 	FIoStoreTocResource& GetTocResource()
@@ -641,11 +590,11 @@ public:
 		Entry->ChunkId = ChunkId;
 		Entry->Options = WriteOptions;
 		Entry->CompressionMethod = CompressionMethodForEntry(WriteOptions);
+		Entry->CompressionMemoryEstimate = CalculateCompressionBufferMemory(Request->GetSourceBufferSizeEstimate());
 		Entry->Request = Request;		
 		Entry->BeginCompressionBarrier = FGraphEvent::CreateGraphEvent();
 		Entry->FinishCompressionBarrier = FGraphEvent::CreateGraphEvent();
 		Entry->BeginWriteBarrier = FGraphEvent::CreateGraphEvent();
-		Entry->WriteFinishedEvent = FGraphEvent::CreateGraphEvent();
 
 		// If we can get the hash without reading the whole thing and hashing it, do so to avoid the IO.
 		if (const FIoHash* ChunkHash = Request->GetChunkHash(); ChunkHash != nullptr)
@@ -661,7 +610,7 @@ public:
 
 				if (ReferenceChunkDatabase.IsValid() && Entry->CompressionMethod != NAME_None)
 				{
-					Entry->bLoadingFromReferenceDb = ReferenceChunkDatabase->ChunkExists(ContainerSettings.ContainerId, Entry->ChunkHash, Entry->ChunkId, Entry->NumChunkBlocksFromRefDb);
+					Entry->bLoadingFromReferenceDb = ReferenceChunkDatabase->ChunkExists(ContainerSettings.ContainerId, Entry->ChunkHash, Entry->ChunkId, Entry->NumChunkBlocks);
 					Entry->bCouldBeFromReferenceDb = true;
 				}
 				return;
@@ -690,7 +639,7 @@ public:
 
 			if (ReferenceChunkDatabase.IsValid() && Entry->CompressionMethod != NAME_None)
 			{
-				Entry->bLoadingFromReferenceDb = ReferenceChunkDatabase->ChunkExists(ContainerSettings.ContainerId, Entry->ChunkHash, Entry->ChunkId, Entry->NumChunkBlocksFromRefDb);
+				Entry->bLoadingFromReferenceDb = ReferenceChunkDatabase->ChunkExists(ContainerSettings.ContainerId, Entry->ChunkHash, Entry->ChunkId, Entry->NumChunkBlocks);
 				Entry->bCouldBeFromReferenceDb = true;
 			}
 
@@ -745,6 +694,11 @@ public:
 				return nullptr;
 			}
 
+			virtual uint64 GetSourceBufferSizeEstimate() override
+			{
+				return SourceBuffer.DataSize();
+			}
+
 			FIoBuffer SourceBuffer;
 			uint64 OrderHint;
 		};
@@ -1347,7 +1301,18 @@ private:
 		return CompressionMethod;
 	}
 
+	int32 CalculateNumChunkBlocks(uint64 ChunkSize) const
+	{
+		const uint64 BlockSize = WriterContext->WriterSettings.CompressionBlockSize;
+		const uint64 NumChunkBlocks64 = Align(ChunkSize, BlockSize) / BlockSize;
+		return IntCastChecked<int32>(NumChunkBlocks64);
+	}
 
+	uint64 CalculateCompressionBufferMemory(uint64 ChunkSize)
+	{
+		int32 NumBlocks = CalculateNumChunkBlocks(ChunkSize);
+		return WriterContext->CompressionBufferSize * NumBlocks;
+	}
 
 	void BeginCompress(FIoStoreWriteQueueEntry* Entry)
 	{
@@ -1355,22 +1320,21 @@ private:
 		
 		if (Entry->bLoadingFromReferenceDb)
 		{
-			if (Entry->NumChunkBlocksFromRefDb == 0)
+			if (Entry->NumChunkBlocks == 0)
 			{
 				Entry->FinishCompressionBarrier->DispatchSubsequents();
 				TRACE_COUNTER_INCREMENT(IoStoreRefDbDone);
 				return;
 			}
 
-			// Allocate resources before launching the read tasks to prevent deadlock. Note this will
+			// Allocate resources before launching the read tasks to reduce contention. Note this will
 			// allocate iobuffers big enough for uncompressed size, when we only actually need it for
 			// compressed size.
-			int32 NumChunkBlocks32 = IntCastChecked<int32>(Entry->NumChunkBlocksFromRefDb);
-			Entry->ChunkBlocks.SetNum(NumChunkBlocks32);
-			for (int32 BlockIndex = 0; BlockIndex < NumChunkBlocks32; ++BlockIndex)
+			Entry->ChunkBlocks.SetNum(Entry->NumChunkBlocks);
+			for (int32 BlockIndex = 0; BlockIndex < Entry->NumChunkBlocks; ++BlockIndex)
 			{
 				FChunkBlock& Block = Entry->ChunkBlocks[BlockIndex];
-				Block.IoBuffer = WriterContext->AllocCompressionBuffer(NumChunkBlocks32);
+				Block.IoBuffer = WriterContext->AllocCompressionBuffer();
 				// Everything else in a block gets filled out from the refdb.
 			}
 
@@ -1438,26 +1402,23 @@ private:
 
 		const FIoBuffer* SourceBuffer = Entry->Request->GetSourceBuffer();
 		Entry->UncompressedSize.Emplace(SourceBuffer->DataSize());
+		Entry->NumChunkBlocks = CalculateNumChunkBlocks(Entry->UncompressedSize.GetValue());
 
-		check(WriterSettings.CompressionBlockSize > 0);
-		const uint64 NumChunkBlocks64 = Align(Entry->UncompressedSize.GetValue(), WriterSettings.CompressionBlockSize) / WriterSettings.CompressionBlockSize;
-		if (NumChunkBlocks64 == 0)
+		if (Entry->NumChunkBlocks == 0)
 		{
 			Entry->FinishCompressionBarrier->DispatchSubsequents();
 			return;
 		}
 
-		int32 NumChunkBlocks = IntCastChecked<int32>(NumChunkBlocks64);
-		
-		Entry->ChunkBlocks.SetNum(NumChunkBlocks);
+		Entry->ChunkBlocks.SetNum(Entry->NumChunkBlocks);
 		{
 			// We must allocate resources for our tasks up front to prevent resource deadlock.
 			uint64 BytesToProcess = Entry->UncompressedSize.GetValue();
 			const uint8* UncompressedData = SourceBuffer->Data();
-			for (int32 BlockIndex = 0; BlockIndex < NumChunkBlocks; ++BlockIndex)
+			for (int32 BlockIndex = 0; BlockIndex < Entry->NumChunkBlocks; ++BlockIndex)
 			{
 				FChunkBlock& Block = Entry->ChunkBlocks[BlockIndex];
-				Block.IoBuffer = WriterContext->AllocCompressionBuffer(NumChunkBlocks);
+				Block.IoBuffer = WriterContext->AllocCompressionBuffer();
 				Block.CompressionMethod = Entry->CompressionMethod;
 				Block.UncompressedSize = FMath::Min(BytesToProcess, WriterSettings.CompressionBlockSize);
 				Block.UncompressedData = UncompressedData;
@@ -1598,6 +1559,21 @@ private:
 	void WriteEntry(FIoStoreWriteQueueEntry* Entry)
 	{
 		TRACE_CPUPROFILER_EVENT_SCOPE(WriteEntry);
+		ON_SCOPE_EXIT
+		{
+			TRACE_CPUPROFILER_EVENT_SCOPE(FreeBlocks);
+			for (FChunkBlock& ChunkBlock : Entry->ChunkBlocks)
+			{
+				WriterContext->FreeCompressionBuffer(ChunkBlock.IoBuffer);
+			}
+			Entry->ChunkBlocks.Empty();
+			delete Entry->Request;
+			Entry->Request = nullptr;
+
+			WriterContext->ScheduledCompressionMemory.SubExchange(Entry->CompressionMemoryEstimate);
+			WriterContext->CompressionMemoryReleasedEvent->Trigger();
+			TRACE_COUNTER_SET(IoStoreCompressionMemoryScheduled, WriterContext->ScheduledCompressionMemory.Load());
+		};
 		const int32* FindExistingIndex = TocBuilder.GetTocEntryIndex(Entry->ChunkId);
 		if (FindExistingIndex)
 		{
@@ -1613,11 +1589,6 @@ private:
 				);
 
 			checkf(TocBuilder.GetTocResource().ChunkMetas[*FindExistingIndex].ChunkHash == Entry->ChunkHash, TEXT("Chunk id has already been added with different content"));
-			for (FChunkBlock& ChunkBlock : Entry->ChunkBlocks)
-			{
-				WriterContext->FreeCompressionBuffer(ChunkBlock.IoBuffer, Entry->ChunkBlocks.Num());
-				ChunkBlock.IoBuffer = nullptr;
-			}
 			return;
 		}
 
@@ -1734,7 +1705,7 @@ private:
 
 		if (ContainerSettings.IsIndexed() && Entry->Options.FileName.Len() > 0)
 		{
-			TocBuilder.AddToFileIndex(Entry->ChunkId, Entry->Options.FileName);
+			TocBuilder.AddToFileIndex(Entry->ChunkId, MoveTemp(Entry->Options.FileName));
 		}
 
 		const uint64 RegionStartOffset = TargetPartition->Offset;
@@ -1746,8 +1717,6 @@ private:
 		{
 			FFileRegion::AccumulateFileRegions(TargetPartition->AllFileRegions, RegionStartOffset, RegionStartOffset, TargetPartition->Offset, Entry->Request->GetRegions());
 		}
-		delete Entry->Request;
-		Entry->Request = nullptr;
 		uint64 WriteStartCycles = FPlatformTime::Cycles64();
 		uint64 WriteBytes = 0;
 		if (Entry->Padding > 0)
@@ -1765,13 +1734,9 @@ private:
 		check(Entry->Offset == TargetPartition->ContainerFileHandle->Tell());
 		for (FChunkBlock& ChunkBlock : Entry->ChunkBlocks)
 		{
-			{
-				TRACE_CPUPROFILER_EVENT_SCOPE(WriteBlockToContainer);
-				TargetPartition->ContainerFileHandle->Serialize(ChunkBlock.IoBuffer->Data(), ChunkBlock.Size);
-				WriteBytes += ChunkBlock.Size;
-			}
-			WriterContext->FreeCompressionBuffer(ChunkBlock.IoBuffer, Entry->ChunkBlocks.Num());
-			ChunkBlock.IoBuffer = nullptr;
+			TRACE_CPUPROFILER_EVENT_SCOPE(WriteBlockToContainer);
+			TargetPartition->ContainerFileHandle->Serialize(ChunkBlock.IoBuffer->Data(), ChunkBlock.Size);
+			WriteBytes += ChunkBlock.Size;
 		}
 		uint64 WriteEndCycles = FPlatformTime::Cycles64();
 		WriterContext->WriteCycleCount.AddExchange(WriteEndCycles - WriteStartCycles);
@@ -1861,16 +1826,14 @@ void FIoStoreWriterContextImpl::Flush()
 
 		for (FIoStoreWriteQueueEntry* Entry : AllEntries)
 		{
-			ScheduleCompression(Entry);
+			BeginCompressionQueue.Enqueue(Entry);
+			ScheduleEntry(Entry);
 		}
 		BeginCompressionQueue.CompleteAdding();
 
 		{
 			TRACE_CPUPROFILER_EVENT_SCOPE(WaitForWritesToComplete);
-			for (int32 EntryIndex = AllEntries.Num() - 1; EntryIndex >= 0; --EntryIndex)
-			{
-				AllEntries[EntryIndex]->WriteFinishedEvent->Wait();
-			}
+			WriterThread.Wait();
 		}
 
 		double WritesEnd = FPlatformTime::Seconds();
@@ -1909,6 +1872,32 @@ void FIoStoreWriterContextImpl::Flush()
 		);
 }
 
+void FIoStoreWriterContextImpl::ScheduleEntry(FIoStoreWriteQueueEntry* Entry)
+{
+	uint64 LocalScheduledCompressionMemory = ScheduledCompressionMemory.Load();
+	
+	while (LocalScheduledCompressionMemory > 0 &&
+		LocalScheduledCompressionMemory + Entry->CompressionMemoryEstimate > CompressionMemorySoftLimit)
+	{
+		TRACE_CPUPROFILER_EVENT_SCOPE(WaitForCompressionMemory);
+		CompressionMemoryReleasedEvent->Wait();
+		LocalScheduledCompressionMemory = ScheduledCompressionMemory.Load();
+	}
+
+	ScheduledCompressionMemory.AddExchange(Entry->CompressionMemoryEstimate);
+	TRACE_COUNTER_SET(IoStoreCompressionMemoryScheduled, ScheduledCompressionMemory.Load());
+
+	if (Entry->bLoadingFromReferenceDb)
+	{
+		// We don't need to wait on a source read so we can kick directly.
+		Entry->BeginCompressionBarrier->DispatchSubsequents();
+	}
+	else
+	{
+		Entry->Request->PrepareSourceBufferAsync(Entry->BeginCompressionBarrier);
+	}
+}
+
 void FIoStoreWriterContextImpl::BeginCompressionThreadFunc()
 {
 	TRACE_CPUPROFILER_EVENT_SCOPE(BeginCompressionThread);
@@ -1946,10 +1935,7 @@ void FIoStoreWriterContextImpl::BeginEncryptionAndSigningThreadFunc()
 		{
 			FIoStoreWriteQueueEntry* Next = Entry->Next;
 			Entry->FinishCompressionBarrier->Wait();
-			if (Entry->bLoadingFromReferenceDb == false)
-			{
-				Entry->Request->FreeSourceBuffer();
-			}
+			Entry->Request->FreeSourceBuffer();
 			if (Entry->bStoreCompressedDataInDDC)
 			{
 				TRACE_CPUPROFILER_EVENT_SCOPE(StoreInDDC);
@@ -1985,7 +1971,6 @@ void FIoStoreWriterContextImpl::WriterThreadFunc()
 			Entry->BeginWriteBarrier->Wait();
 			TRACE_COUNTER_INCREMENT(IoStoreBeginWriteCount);
 			Entry->Writer->WriteEntry(Entry);
-			Entry->WriteFinishedEvent->DispatchSubsequents();
 			Entry = Next;
 		}
 	}
diff --git a/Engine/Source/Developer/IoStoreUtilities/Private/IoStoreWriter.h b/Engine/Source/Developer/IoStoreUtilities/Private/IoStoreWriter.h
--- a/Engine/Source/Developer/IoStoreUtilities/Private/IoStoreWriter.h
+++ b/Engine/Source/Developer/IoStoreUtilities/Private/IoStoreWriter.h
@@ -110,6 +110,8 @@ public:
 
 	// Can't be called between PrepareSourceBufferAsync and its completion!
 	virtual void FreeSourceBuffer() = 0;
+
+	virtual uint64 GetSourceBufferSizeEstimate() = 0;
 };
 
 class IIoStoreWriterReferenceChunkDatabase
@@ -142,7 +144,7 @@ public:
 	* the necessary structures without needing to read the source data for the chunk. This might be called from
 	* multiple threads as it has to happen after we have the source hash computed.
 	*/
-	virtual bool ChunkExists(const FIoContainerId& InContainerId, const FIoHash& InChunkHash, const FIoChunkId& InChunkId, uint32& OutNumChunkBlocks) = 0;
+	virtual bool ChunkExists(const FIoContainerId& InContainerId, const FIoHash& InChunkHash, const FIoChunkId& InChunkId, int32& OutNumChunkBlocks) = 0;
 
 	/*
 	* Returns the compression block size that was used to break up the IoChunks in the source containers. If this is different than what we want, 
