Add SSL certificates to the zen manifest if the configuration would mean they're part of the staged build.

diff --git a/Engine/Source/Developer/IoStoreUtilities/Private/ZenFileSystemManifest.cpp b/Engine/Source/Developer/IoStoreUtilities/Private/ZenFileSystemManifest.cpp
--- a/Engine/Source/Developer/IoStoreUtilities/Private/ZenFileSystemManifest.cpp
+++ b/Engine/Source/Developer/IoStoreUtilities/Private/ZenFileSystemManifest.cpp
@@ -7,6 +7,7 @@
 #include "Interfaces/ITargetPlatform.h"
 #include "Interfaces/IPluginManager.h"
 #include "Misc/App.h"
+#include "Misc/ConfigCacheIni.h"
 #include "Misc/FileHelper.h"
 #include "Misc/Paths.h"
 #include "Misc/DataDrivenPlatformInfoRegistry.h"
@@ -426,6 +427,44 @@ int32 FZenFileSystemManifest::Generate()
 	
 	AddFilesFromExtensionDirectories(TEXT("Content/Localization"), &LocalizationFilter);
 
+	bool bSSLCertificatesWillStage = false;
+	FConfigCacheIni* TargetPlatformConfig = TargetPlatform.GetConfigSystem();
+	if (TargetPlatformConfig)
+	{
+		GConfig->GetBool(TEXT("/Script/Engine.NetworkSettings"), TEXT("n.VerifyPeer"), bSSLCertificatesWillStage, GEngineIni);
+	}
+	if (bSSLCertificatesWillStage)
+	{
+		FString ProjectCertFile = FPaths::Combine(ProjectDir, TEXT("Content"), TEXT("Certificates"), TEXT("cacert.pem"));
+		if (FPaths::FileExists(ProjectCertFile))
+		{
+			const TCHAR* ClientProjectCertFile = TEXT("/{project}/Content/Certificates/cacert.pem");
+			const FIoChunkId FileChunkId = CreateExternalFileChunkId(ClientProjectCertFile);
+			FPaths::MakePathRelativeTo(ProjectCertFile, *FPaths::RootDir());
+			AddManifestEntry(
+				FileChunkId,
+				ProjectCertFile,
+				ClientProjectCertFile);
+		}
+		else
+		{
+			FString EngineCertFile = FPaths::Combine(EngineDir, TEXT("Content"), TEXT("Certificates"), TEXT("ThirdParty"), TEXT("cacert.pem"));
+			if (FPaths::FileExists(EngineCertFile))
+			{
+				const TCHAR* ClientEngineCertFile = TEXT("/{engine}/Content/Certificates/ThirdParty/cacert.pem");
+				const FIoChunkId FileChunkId = CreateExternalFileChunkId(ClientEngineCertFile);
+				FPaths::MakePathRelativeTo(EngineCertFile, *FPaths::RootDir());
+				AddManifestEntry(
+					FileChunkId,
+					EngineCertFile,
+					ClientEngineCertFile);
+			}
+		}
+		FFileFilter CertificateFilter = FFileFilter()
+			.IncludeExtension(TEXT("pem"));
+		AddFilesFromDirectory(TEXT("/{project}/Certificates"), FPaths::Combine(ProjectDir, TEXT("Certificates")), true, &CertificateFilter);
+	}
+
 	FFileFilter ContentFilter = FFileFilter()
 		.ExcludeExtension(TEXT("uasset"))
 		.ExcludeExtension(TEXT("ubulk"))
