Adding asset types and package names to the IoStore output csv file

diff --git a/Engine/Source/Developer/IoStoreUtilities/Private/IoStoreUtilities.cpp b/Engine/Source/Developer/IoStoreUtilities/Private/IoStoreUtilities.cpp
--- a/Engine/Source/Developer/IoStoreUtilities/Private/IoStoreUtilities.cpp
+++ b/Engine/Source/Developer/IoStoreUtilities/Private/IoStoreUtilities.cpp
@@ -100,6 +100,11 @@ struct FIoStoreChunkSource
 	UE::Cook::EPluginSizeTypes SizeType;
 };
 
+static FPackageId PackageIdFromChunkId(const FIoChunkId& ChunkId)
+{
+	return FPackageId::FromValue(*(int64*)(ChunkId.GetData()));
+}
+
 static const FName DefaultCompressionMethod = NAME_Zlib;
 static const uint64 DefaultCompressionBlockSize = 64 << 10;
 static const uint64 DefaultCompressionBlockAlignment = 64 << 10;
@@ -1422,17 +1427,17 @@ public:
 		OutputArchive.Reset(IFileManager::Get().CreateFileWriter(OutputFilename));
 		if (OutputArchive.IsValid())
 		{
-			OutputArchive->Logf(TEXT("OrderInContainer, ChunkId, PackageId, PackageName, Filename, ContainerName, Offset, OffsetOnDisk, Size, CompressedSize, Hash, ChunkType"));
+			OutputArchive->Logf(TEXT("OrderInContainer, ChunkId, PackageId, PackageName, Filename, ContainerName, Offset, OffsetOnDisk, Size, CompressedSize, Hash, ChunkType, ClassType"));
 		}
 	}
 
-	void AddChunk(const FString& ContainerName, int32 Index, const FIoStoreTocChunkInfo& Info, FPackageId PackageId, const FString& PackageName)
+	void AddChunk(const FString& ContainerName, int32 Index, const FIoStoreTocChunkInfo& Info, FPackageId PackageId, const FString& PackageName, const FString& ClassType)
 	{
 		if (OutputArchive.IsValid())
 		{
-			OutputArchive->Logf(TEXT("%d, %s, 0x%llX, %s, %s, %s, %lld, %lld, %lld, %lld, 0x%s, %s"),
+			OutputArchive->Logf(TEXT("%d, %s, 0x%llX, %s, %s, %s, %lld, %lld, %lld, %lld, 0x%s, %s, %s"),
 				Index,
-				*BytesToHex(Info.Id.GetData(), Info.Id.GetSize()),
+				*LexToString(Info.Id),
 				PackageId.ValueForDebugging(),
 				*PackageName,
 				*Info.FileName,
@@ -1442,7 +1447,8 @@ public:
 				Info.Size,
 				Info.CompressedSize,
 				*LexToString(Info.ChunkHash),
-				*LexToString(Info.ChunkType)
+				*LexToString(Info.ChunkType),
+				*ClassType
 			);
 		}
 	}	
@@ -5057,7 +5063,7 @@ static bool DoAssetRegistryWritebackDuringStage(
 
 				// Shader code chunks don't have the package in their chunk id, so we have to use other data to look
 				// it up and find it.
-				FPackageId PackageId = FPackageId::FromValue(*(int64*)(ChunkInfo.Id.GetData()));
+				FPackageId PackageId = PackageIdFromChunkId(ChunkInfo.Id);
 				if (ChunkInfo.ChunkType == EIoChunkType::ShaderCode)
 				{
 					// Update size info for the shader.
@@ -5907,7 +5913,7 @@ int32 CreateTarget(const FIoStoreArguments& Arguments, const FIoStoreWriterSetti
 							}
 						}
 
-						Out->AddChunk(Result.ContainerName, EntryIndex, ChunkInfo, PackageId, PackageName);
+						Out->AddChunk(Result.ContainerName, EntryIndex, ChunkInfo, PackageId, PackageName, "");
 					}
 				}
 			}
@@ -6275,6 +6281,27 @@ int32 CreateContentPatch(const FIoStoreArguments& Arguments, const FIoStoreWrite
 	return 0;
 }
 
+void Split(FStringView Line, FStringView::ViewType Sep, TArray<FStringView>& OutValues)
+{
+	OutValues.Reset();
+
+	int32 Start = 0;
+	int32 End = 0;
+	for (;;)
+	{
+		End = Line.Find(TEXTVIEW(","), Start);
+		if (End < 0)
+		{
+			break;
+		}
+		FStringView Value = Line.Mid(Start, End - Start);
+		OutValues.Add(Value);
+		Start = End + 1;
+	}
+	FStringView Value = Line.Mid(Start);
+	OutValues.Add(Value);
+}
+
 int32 ListContainer(
 	const FKeyChain& KeyChain,
 	const FString& ContainerPathOrWildcard,
@@ -6283,6 +6310,30 @@ int32 ListContainer(
 	IOSTORE_CPU_SCOPE(ListContainer);
 	TArray<FString> ContainerFilePaths;
 
+	FString AllChunksInfoFilename;
+	TMap<FPackageId, TPair<FName,FName>> AllChunksInfoMap;
+	if (FParse::Value(FCommandLine::Get(), TEXT("AllChunksInfo="), AllChunksInfoFilename))
+	{
+		TArray<FStringView> Values;
+		bool bSkipFirstLine = true;
+		auto Visitor = [&AllChunksInfoMap,&Values,&bSkipFirstLine](FStringView Line)
+		{
+			const int32 PackageNameIndex = 1;
+			const int32 ClassTypeIndex = 2;
+			if (bSkipFirstLine)
+			{
+				bSkipFirstLine = false;
+				return;
+			}
+			Split(Line, TEXTVIEW(","), Values);
+			FName PackageName(Values[PackageNameIndex]);
+			FPackageId PackageId = FPackageId::FromName(PackageName);
+			AllChunksInfoMap.Add(PackageId, TPair<FName,FName>(PackageName, FName(Values[ClassTypeIndex])));
+		};
+
+		FFileHelper::LoadFileToStringWithLineVisitor(*AllChunksInfoFilename, Visitor);
+	}
+
 	if (IFileManager::Get().FileExists(*ContainerPathOrWildcard))
 	{
 		ContainerFilePaths.Add(ContainerPathOrWildcard);
@@ -6382,18 +6433,27 @@ int32 ListContainer(
 		{
 			IOSTORE_CPU_SCOPE(WriteCsvFile);
 			FString PackageName;
-			for(int32 Index=0; Index < Chunks.Num(); ++Index)
+			FString ClassType;
+			for (int32 Index=0; Index < Chunks.Num(); ++Index)
 			{
 				const FIoStoreTocChunkInfo& ChunkInfo = Chunks[Index];
 
-				FPackageId PackageId;
+				FPackageId PackageId = PackageIdFromChunkId(ChunkInfo.Id);
 				PackageName.Reset();
+				ClassType.Reset();
 				if (ChunkInfo.bHasValidFileName && FPackageName::TryConvertFilenameToLongPackageName(ChunkInfo.FileName, PackageName, nullptr))
 				{
 					PackageId = FPackageId::FromName(FName(*PackageName));
 				}
 
-				Out->AddChunk(ContainerName, Index, ChunkInfo, PackageId, PackageName);
+				TPair<FName,FName>* PackageInfo = AllChunksInfoMap.Find(PackageId);
+				if (PackageInfo)
+				{
+					PackageName = PackageInfo->Get<0>().ToString();
+					ClassType = PackageInfo->Get<1>().ToString();
+				}
+
+				Out->AddChunk(ContainerName, Index, ChunkInfo, PackageId, PackageName, ClassType);
 			}
 		}
 	}
