Updating missing directory index 'warning' to 'display' as indexing is optional and global .utoc is not indexed

diff --git a/Engine/Source/Developer/IoStoreUtilities/Private/IoStoreUtilities.cpp b/Engine/Source/Developer/IoStoreUtilities/Private/IoStoreUtilities.cpp
--- a/Engine/Source/Developer/IoStoreUtilities/Private/IoStoreUtilities.cpp
+++ b/Engine/Source/Developer/IoStoreUtilities/Private/IoStoreUtilities.cpp
@@ -6407,7 +6407,7 @@ int32 ListContainer(
 
 		if (!EnumHasAnyFlags(Reader->GetContainerFlags(), EIoContainerFlags::Indexed))
 		{
-			UE_LOG(LogIoStore, Warning, TEXT("Missing directory index for container '%s'"), *ContainerFilePath);
+			UE_LOG(LogIoStore, Display, TEXT("No directory index for container '%s'"), *ContainerFilePath);
 		}
 
 		UE_LOG(LogIoStore, Display, TEXT("Listing container '%s'"), *ContainerFilePath);
