Cooker: ZenStore oplog attachments: Fix corrupted attachments in incremental cooks when we have more than 1 attachment. Part of the attachment data from the previous cook was not cleared when recording the results from the new version of the package.
#rnx

diff --git a/Engine/Source/Developer/IoStoreUtilities/Private/ZenStoreWriter.cpp b/Engine/Source/Developer/IoStoreUtilities/Private/ZenStoreWriter.cpp
--- a/Engine/Source/Developer/IoStoreUtilities/Private/ZenStoreWriter.cpp
+++ b/Engine/Source/Developer/IoStoreUtilities/Private/ZenStoreWriter.cpp
@@ -868,7 +868,8 @@ void FZenStoreWriter::CommitPackageInternal(FZenCommitInfo&& ZenCommitInfo)
 
 		const int32 NumAttachments = CommitInfo.Attachments.Num();
 		TArray<FCbAttachment, TInlineAllocator<2>> CbAttachments;
-		
+		CookInfo.Attachments.Empty(NumAttachments);
+
 		if (NumAttachments)
 		{
 			TArray<const FCommitAttachmentInfo*, TInlineAllocator<2>> SortedAttachments;
