UnrealPak: Convert IoStore from old TaskGraph API to the new UE::Tasks API.
Enables retraction when waiting on some of the scheduler threads.
Add MaxConcurrentSourceReads throttling since the ZenServerHttp is limited to 32 blocking requests (that should be enough for normal file io as well).
Bump default MaxSourceBufferMemory to 4GiB so that it is higher than the 3 GiB CompressionMemorySoftLimit to make sure a started queue entry can never get blocked.
Batch up to 4 compression blocks in each task to reduce task overhead with neglectable impact on tail wait times.

#rnx

diff --git a/Engine/Source/Developer/IoStoreUtilities/Private/IoStoreLooseFiles.cpp b/Engine/Source/Developer/IoStoreUtilities/Private/IoStoreLooseFiles.cpp
--- a/Engine/Source/Developer/IoStoreUtilities/Private/IoStoreLooseFiles.cpp
+++ b/Engine/Source/Developer/IoStoreUtilities/Private/IoStoreLooseFiles.cpp
@@ -81,9 +81,9 @@ public:
 
 			virtual ~FWriteRequest() = default;
 
-			void PrepareSourceBufferAsync(FGraphEventRef CompletionEvent) override
+			void PrepareSourceBufferAsync(UE::Tasks::FTaskEvent& CompletionEvent) override
 			{
-				CompletionEvent->DispatchSubsequents();
+				CompletionEvent.Trigger();
 			}
 
 			const FIoBuffer* GetSourceBuffer() override
@@ -137,11 +137,11 @@ public:
 					
 					// Setup the task pipe when holding the lock to make it easy to flush all pending task pipe(s)
 
-					FGraphEventRef Event = FGraphEvent::CreateGraphEvent();
+					UE::Tasks::FTaskEvent Event { UE_SOURCE_LOCATION };
 					PendingWrite->WriteRequest->PrepareSourceBufferAsync(Event);
 					UE::Tasks::FTask ReadChunkTask = PendingWrite->TaskPipe.Launch(TEXT("ReadChunk"), [PendingWrite, Event]() mutable
 					{
-						Event->Wait();
+						Event.Wait();
 					});
 
 					UE::Tasks::FTask WriteChunkTask = PendingWrite->TaskPipe.Launch(TEXT("WriteChunk"), [this, PendingWrite]() mutable
diff --git a/Engine/Source/Developer/IoStoreUtilities/Private/IoStoreUtilities.cpp b/Engine/Source/Developer/IoStoreUtilities/Private/IoStoreUtilities.cpp
--- a/Engine/Source/Developer/IoStoreUtilities/Private/IoStoreUtilities.cpp
+++ b/Engine/Source/Developer/IoStoreUtilities/Private/IoStoreUtilities.cpp
@@ -80,6 +80,7 @@
 #include "Serialization/CompactBinarySerialization.h"
 #include "Serialization/ZenPackageHeader.h"
 #include "String/ParseTokens.h"
+#include "Tasks/Task.h"
 
 IMPLEMENT_MODULE(FDefaultModuleImpl, IoStoreUtilities);
 
@@ -1292,37 +1293,13 @@ public:
 		return ZenStoreClient->ReadChunk(ChunkId, ReadOptions.GetOffset(), ReadOptions.GetSize());
 	}
 
-	void ReadChunkAsync(const FIoChunkId& ChunkId, TFunction<void(TIoStatusOr<FIoBuffer>)>&& Callback)
+	UE::Tasks::FTask ReadChunkAsync(const FIoChunkId& ChunkId, TFunction<void(TIoStatusOr<FIoBuffer>)>&& Callback)
 	{
-		class FReadChunkTask
-			: public FNonAbandonableTask
+		return UE::Tasks::Launch(TEXT("ReadChunkAsync"), [this, ChunkId, Callback = MoveTemp(Callback)]()
 		{
-		public:
-			FReadChunkTask(UE::FZenStoreHttpClient* InZenStoreClient, const FIoChunkId& InChunkId, TFunction<void(TIoStatusOr<FIoBuffer>)>&& InCallback)
-				: ZenStoreClient(InZenStoreClient)
-				, ChunkId(InChunkId)
-				, Callback(MoveTemp(InCallback))
-			{
-			}
-
-			void DoWork()
-			{
-				FIoReadOptions ReadOptions;
-				Callback(ZenStoreClient->ReadChunk(ChunkId, ReadOptions.GetOffset(), ReadOptions.GetSize()));
-			}
-
-			TStatId GetStatId() const
-			{
-				return TStatId();
-			}
-
-		private:
-			UE::FZenStoreHttpClient* ZenStoreClient;
-			FIoChunkId ChunkId;
-			TFunction<void(TIoStatusOr<FIoBuffer>)> Callback;
-		};
-
-		(new FAutoDeleteAsyncTask<FReadChunkTask>(ZenStoreClient.Get(), ChunkId, MoveTemp(Callback)))->StartBackgroundTask();
+			FIoReadOptions ReadOptions;
+			Callback(ZenStoreClient->ReadChunk(ChunkId, ReadOptions.GetOffset(), ReadOptions.GetSize()));
+		}, UE::Tasks::ETaskPriority::Normal);
 	}
 
 private:
@@ -3348,10 +3325,18 @@ public:
 	FIoStoreWriteRequestManager(FPackageStoreOptimizer& InPackageStoreOptimizer, FCookedPackageStore* InPackageStore)
 		: PackageStoreOptimizer(InPackageStoreOptimizer)
 		, PackageStore(InPackageStore)
-		, MemoryAvailableEvent(FPlatformProcess::GetSynchEventFromPool(false))
 	{
 		InitiatorThread = Async(EAsyncExecution::Thread, [this]() { InitiatorThreadFunc(); });
 		RetirerThread = Async(EAsyncExecution::Thread, [this]() { RetirerThreadFunc(); });
+
+		MaxSourceBufferMemory = 4ull << 30;
+		FParse::Value(FCommandLine::Get(), TEXT("MaxSourceBufferMemory="), MaxSourceBufferMemory);
+
+		MaxConcurrentSourceReads = uint32(FMath::Clamp(FPlatformMisc::NumberOfCoresIncludingHyperthreads()/2, 4, 32));
+		FParse::Value(FCommandLine::Get(), TEXT("MaxConcurrentSourceReads="), MaxConcurrentSourceReads);
+
+		UE_LOG(LogIoStore, Display, TEXT("Initialized WriteRequestManager with MaxConcurrentSourceReads=%d (%d cores), MaxSourceBufferMemory=%lluMiB"),
+			MaxConcurrentSourceReads, FPlatformMisc::NumberOfCoresIncludingHyperthreads(), MaxSourceBufferMemory >> 20);
 	}
 
 	~FIoStoreWriteRequestManager()
@@ -3360,7 +3345,6 @@ public:
 		RetirerQueue.CompleteAdding();
 		InitiatorThread.Wait();
 		RetirerThread.Wait();
-		FPlatformProcess::ReturnSynchEventToPool(MemoryAvailableEvent);
 	}
 
 	IIoStoreWriteRequest* Read(const FContainerTargetFile& InTargetFile)
@@ -3407,9 +3391,9 @@ private:
 			return TargetFile.ChunkHash.IsZero() ? nullptr : &TargetFile.ChunkHash;
 		}
 		
-		virtual void PrepareSourceBufferAsync(FGraphEventRef InCompletionEvent) override
+		virtual void PrepareSourceBufferAsync(UE::Tasks::FTaskEvent& InCompletionEvent) override
 		{
-			CompletionEvent = InCompletionEvent;
+			CompletionEvent.Emplace(InCompletionEvent);
 			Manager.ScheduleLoad(this);
 		}
 
@@ -3441,12 +3425,15 @@ private:
 			, FileRegions(TargetFile.FileRegions)
 			, SourceBufferSize(TargetFile.SourceSize) { }
 
-		void OnSourceBufferLoaded()
+		void OnSourceBufferLoaded(bool bTriggerCompletionEvent)
 		{
 			TRACE_COUNTER_DECREMENT(IoStoreSourceReadsInflight);
 			TRACE_COUNTER_INCREMENT(IoStoreSourceReadsDone);
 			QueueEntry->ReleaseRef(Manager);
-			CompletionEvent->DispatchSubsequents();
+			if (bTriggerCompletionEvent)
+			{
+				CompletionEvent.GetValue().Trigger();
+			}
 		}
 
 		FIoStoreWriteRequestManager& Manager;
@@ -3457,7 +3444,7 @@ private:
 		// used for IO, however it's not necessarily the size of the resulting input to iostore as
 		// the buffer can be post-processed after i/o (e.g. CreateOptimizedPackage).
 		uint64 SourceBufferSize;
-		FGraphEventRef CompletionEvent;
+		TOptional<UE::Tasks::FTaskEvent> CompletionEvent;
 		FIoBuffer SourceBuffer;
 		FQueueEntry* QueueEntry = nullptr;
 	};
@@ -3474,7 +3461,7 @@ private:
 			Manager.MemorySourceReads[(int8)TargetFile.ChunkId.GetChunkType()].IncrementExchange();
 			SourceBuffer = TargetFile.SourceBuffer.GetValue();
 			Manager.MemorySourceBytes[(int8)TargetFile.ChunkId.GetChunkType()] += SourceBuffer.DataSize();
-			OnSourceBufferLoaded();
+			OnSourceBufferLoaded(/*bTriggerCompletionEvent*/ true);
 		}
 	};
 
@@ -3511,7 +3498,7 @@ private:
 					check(Package->OptimizedOptionalSegmentPackage);
 					SourceBuffer = Manager.PackageStoreOptimizer.CreatePackageBuffer(Package->OptimizedOptionalSegmentPackage, SourceBuffer);
 				}
-				OnSourceBufferLoaded();
+				OnSourceBufferLoaded(/*bTriggerCompletionEvent*/ true);
 			};
 
 			QueueEntry->ReadRequest.Reset(
@@ -3533,14 +3520,16 @@ private:
 		virtual void LoadSourceBufferAsync() override
 		{
 			Manager.ZenSourceReads[(int8)TargetFile.ChunkId.GetChunkType()].IncrementExchange();
-			Manager.PackageStore->ReadChunkAsync(
+			UE::Tasks::FTask ReadTask = Manager.PackageStore->ReadChunkAsync(
 				TargetFile.ChunkId,
 				[this](TIoStatusOr<FIoBuffer> Status)
 				{
 					SourceBuffer = Status.ConsumeValueOrDie();
 					Manager.ZenSourceBytes[(int8)TargetFile.ChunkId.GetChunkType()] += SourceBuffer.DataSize();
-					OnSourceBufferLoaded();
+					OnSourceBufferLoaded(/*bTriggerCompletionEvent*/ false);
 				});
+			CompletionEvent.GetValue().AddPrerequisites(ReadTask);
+			CompletionEvent.GetValue().Trigger();
 		}
 	};
 
@@ -3651,6 +3640,8 @@ private:
 
 	void ScheduleRetire(FQueueEntry* QueueEntry)
 	{
+		--NumConcurrentSourceReads;
+		SourceReadCompletedEvent->Trigger();
 		RetirerQueue.Enqueue(QueueEntry);
 	}
 
@@ -3658,8 +3649,14 @@ private:
 	{
 		const uint64 SourceBufferSize = QueueEntry->WriteRequest->GetSourceBufferSizeEstimate();
 
+		while (NumConcurrentSourceReads >= MaxConcurrentSourceReads)
+		{
+			TRACE_CPUPROFILER_EVENT_SCOPE(WaitForSourceReads);
+			SourceReadCompletedEvent->Wait();
+		}
+
 		uint64 LocalUsedBufferMemory = UsedBufferMemory.Load();
-		while (LocalUsedBufferMemory > 0 && LocalUsedBufferMemory + SourceBufferSize > BufferMemoryLimit)
+		while (LocalUsedBufferMemory > 0 && LocalUsedBufferMemory + SourceBufferSize > MaxSourceBufferMemory)
 		{
 			TRACE_CPUPROFILER_EVENT_SCOPE(WaitForBufferMemory);
 			MemoryAvailableEvent->Wait();
@@ -3670,6 +3667,7 @@ private:
 		TRACE_COUNTER_INCREMENT(IoStoreSourceReadsInflight);
 		TRACE_COUNTER_ADD(IoStoreSourceReadsUsedBufferMemory, SourceBufferSize);
 		QueueEntry->WriteRequest->LoadSourceBufferAsync();
+		++NumConcurrentSourceReads;
 	}
 
 	void Retire(FQueueEntry* QueueEntry)
@@ -3735,8 +3733,12 @@ private:
 	TFuture<void> RetirerThread;
 	FQueue InitiatorQueue;
 	FQueue RetirerQueue;
+	uint64 MaxSourceBufferMemory = 0;
 	TAtomic<uint64> UsedBufferMemory { 0 };
-	FEvent* MemoryAvailableEvent;
+	int32 MaxConcurrentSourceReads = 0;
+	TAtomic<int32> NumConcurrentSourceReads { 0 };
+	FEventRef MemoryAvailableEvent;
+	FEventRef SourceReadCompletedEvent;
 
 public:
 	TAtomic<uint64> ZenSourceReads[(int8)EIoChunkType::MAX] { 0 };
@@ -3745,8 +3747,6 @@ public:
 	TAtomic<uint64> MemorySourceBytes[(int8)EIoChunkType::MAX]{ 0 };
 	TAtomic<uint64> LooseFileSourceReads[(int8)EIoChunkType::MAX]{ 0 };
 	TAtomic<uint64> LooseFileSourceBytes[(int8)EIoChunkType::MAX]{ 0 };
-
-	static constexpr uint64 BufferMemoryLimit = 2ull << 30;
 };
 
 static bool WriteUtf8StringView(FUtf8StringView InView, const FString& InFilename)
diff --git a/Engine/Source/Developer/IoStoreUtilities/Private/IoStoreWriter.cpp b/Engine/Source/Developer/IoStoreUtilities/Private/IoStoreWriter.cpp
--- a/Engine/Source/Developer/IoStoreUtilities/Private/IoStoreWriter.cpp
+++ b/Engine/Source/Developer/IoStoreUtilities/Private/IoStoreWriter.cpp
@@ -18,6 +18,7 @@
 #include "ProfilingDebugging/CountersTrace.h"
 #include "Serialization/LargeMemoryReader.h"
 #include "Serialization/MemoryWriter.h"
+#include "Tasks/Task.h"
 #include "Templates/UniquePtr.h"
 
 TRACE_DECLARE_MEMORY_COUNTER(IoStoreCompressionMemoryUsed, TEXT("IoStoreWriter/CompressionMemoryUsed"));
@@ -72,13 +73,11 @@ struct FIoStoreWriteQueueEntry
 	TArray<FChunkBlock> ChunkBlocks;
 	FIoWriteOptions Options;
 	FName CompressionMethod = NAME_None;
-	FGraphEventRef HashBarrier;
-	FGraphEventRef HashTask;
-	FGraphEventRef BeginCompressionBarrier;
-	FGraphEventRef FinishCompressionBarrier;
-	FGraphEventRef BeginWriteBarrier;
+	UE::Tasks::FTask HashTask;
+	UE::Tasks::FTaskEvent BeginCompressionBarrier{ TEXT("BeginCompression") };
+	UE::Tasks::FTaskEvent FinishCompressionBarrier{ TEXT("FinishCompression") };
+	UE::Tasks::FTaskEvent BeginWriteBarrier{ TEXT("BeginWrite") };
 	TAtomic<int32> CompressedBlocksCount{ 0 };
-	TAtomic<int32> FinishedBlocksCount{ 0 };
 	int32 PartitionIndex = -1;
 	int32 NumChunkBlocks = 0;
 	FString DDCCacheKey;
@@ -592,9 +591,6 @@ public:
 		Entry->CompressionMethod = CompressionMethodForEntry(WriteOptions);
 		Entry->CompressionMemoryEstimate = CalculateCompressionBufferMemory(Request->GetSourceBufferSizeEstimate());
 		Entry->Request = Request;		
-		Entry->BeginCompressionBarrier = FGraphEvent::CreateGraphEvent();
-		Entry->FinishCompressionBarrier = FGraphEvent::CreateGraphEvent();
-		Entry->BeginWriteBarrier = FGraphEvent::CreateGraphEvent();
 
 		// If we can get the hash without reading the whole thing and hashing it, do so to avoid the IO.
 		if (const FIoHash* ChunkHash = Request->GetChunkHash(); ChunkHash != nullptr)
@@ -618,9 +614,8 @@ public:
 			// If we are validating run the normal path to verify it.
 		}
 		// Otherwise, we have to do the load & hash
-
-		Entry->HashBarrier = FGraphEvent::CreateGraphEvent();
-		Entry->HashTask = FFunctionGraphTask::CreateAndDispatchWhenReady([this, Entry]()
+		UE::Tasks::FTaskEvent HashEvent{ TEXT("HashEvent") };
+		Entry->HashTask = UE::Tasks::Launch(TEXT("HashChunk"), [this, Entry]()
 		{
 			TRACE_CPUPROFILER_EVENT_SCOPE(HashChunk);
 			const FIoBuffer* SourceBuffer = Entry->Request->GetSourceBuffer();
@@ -645,10 +640,10 @@ public:
 
 			// Release the source data buffer, it will be reloaded later when we start compressing the chunk
 			Entry->Request->FreeSourceBuffer();
-		}, TStatId(), Entry->HashBarrier, ENamedThreads::AnyHiPriThreadHiPriTask);
+		}, HashEvent, UE::Tasks::ETaskPriority::High);
 
 		// Kick off the source buffer read to run the hash task
-		Entry->Request->PrepareSourceBufferAsync(Entry->HashBarrier);
+		Entry->Request->PrepareSourceBufferAsync(HashEvent);
 	}
 
 	virtual void Append(const FIoChunkId& ChunkId, FIoBuffer Chunk, const FIoWriteOptions& WriteOptions, uint64 OrderHint) override
@@ -665,9 +660,9 @@ public:
 
 			virtual ~FWriteRequest() = default;
 
-			void PrepareSourceBufferAsync(FGraphEventRef CompletionEvent) override
+			void PrepareSourceBufferAsync(UE::Tasks::FTaskEvent& CompletionEvent) override
 			{
-				CompletionEvent->DispatchSubsequents();
+				CompletionEvent.Trigger();
 			}
 
 			const FIoBuffer* GetSourceBuffer() override
@@ -1322,7 +1317,7 @@ private:
 		{
 			if (Entry->NumChunkBlocks == 0)
 			{
-				Entry->FinishCompressionBarrier->DispatchSubsequents();
+				Entry->FinishCompressionBarrier.Trigger();
 				TRACE_COUNTER_INCREMENT(IoStoreRefDbDone);
 				return;
 			}
@@ -1385,7 +1380,7 @@ private:
 				}
 
 				Entry->UncompressedSize.Emplace(TotalUncompressedSize);
-				Entry->FinishCompressionBarrier->DispatchSubsequents();
+				Entry->FinishCompressionBarrier.Trigger();
 				TRACE_COUNTER_DECREMENT(IoStoreRefDbInflight);
 				TRACE_COUNTER_INCREMENT(IoStoreRefDbDone);
 			});
@@ -1406,7 +1401,7 @@ private:
 
 		if (Entry->NumChunkBlocks == 0)
 		{
-			Entry->FinishCompressionBarrier->DispatchSubsequents();
+			Entry->FinishCompressionBarrier.Trigger();
 			return;
 		}
 
@@ -1435,7 +1430,7 @@ private:
 				Block.CompressedSize = Block.UncompressedSize;
 				FMemory::Memcpy(Block.IoBuffer->Data(), Block.UncompressedData, Block.UncompressedSize);
 			}
-			Entry->FinishCompressionBarrier->DispatchSubsequents();
+			Entry->FinishCompressionBarrier.Trigger();
 			return;
 		}
 
@@ -1468,7 +1463,7 @@ private:
 			if (bFoundInDDC)
 			{
 				WriterContext->CompressionDDCHitCount.IncrementExchange();
-				Entry->FinishCompressionBarrier->DispatchSubsequents();
+				Entry->FinishCompressionBarrier.Trigger();
 				return;
 			}
 			else
@@ -1484,39 +1479,49 @@ private:
 	void ScheduleCompressionTasks(FIoStoreWriteQueueEntry* Entry)
 	{
 		TRACE_COUNTER_INCREMENT(IoStoreCompressionInflight);
-		for (FChunkBlock& Block : Entry->ChunkBlocks)
+		constexpr int32 BatchSize = 4;
+		const int32 NumBatches = 1 + (Entry->ChunkBlocks.Num() / BatchSize);
+		for (int32 BatchIndex = 0; BatchIndex < NumBatches; ++BatchIndex)
 		{
+			const int32 BeginIndex = BatchIndex * BatchSize;
+			const int32 EndIndex = FMath::Min(BeginIndex + BatchSize, Entry->ChunkBlocks.Num());
 			WriterContext->ScheduledCompressionTasksCount.IncrementExchange();
-			FFunctionGraphTask::CreateAndDispatchWhenReady([this, Entry, BlockPtr = &Block]()
+			UE::Tasks::FTask CompressTask = UE::Tasks::Launch(TEXT("CompressBlocks"), [this, Entry, BeginIndex, EndIndex]()
 			{
-				CompressBlock(BlockPtr);
-				WriterContext->ScheduledCompressionTasksCount.DecrementExchange();
-				int32 CompressedBlocksCount = Entry->CompressedBlocksCount.IncrementExchange();
-				if (CompressedBlocksCount + 1 == Entry->ChunkBlocks.Num())
+				for (int32 Index = BeginIndex; Index < EndIndex; ++Index)
 				{
-					WriterContext->CompressedChunksByType[(int8)Entry->ChunkId.GetChunkType()].IncrementExchange();
-					WriterContext->CompressedChunksCount.IncrementExchange();
-					Entry->FinishCompressionBarrier->DispatchSubsequents();
-					TRACE_COUNTER_DECREMENT(IoStoreCompressionInflight);
+					FChunkBlock* BlockPtr = &Entry->ChunkBlocks[Index];
+					CompressBlock(BlockPtr);
+					int32 CompressedBlocksCount = Entry->CompressedBlocksCount.IncrementExchange();
+					if (CompressedBlocksCount + 1 == Entry->ChunkBlocks.Num())
+					{
+						WriterContext->CompressedChunksByType[(int8)Entry->ChunkId.GetChunkType()].IncrementExchange();
+						WriterContext->CompressedChunksCount.IncrementExchange();
+						TRACE_COUNTER_DECREMENT(IoStoreCompressionInflight);
+					}
 				}
-			}, TStatId(), nullptr, ENamedThreads::AnyHiPriThreadNormalTask);
+				WriterContext->ScheduledCompressionTasksCount.DecrementExchange();
+			}, UE::Tasks::ETaskPriority::High);
+			Entry->FinishCompressionBarrier.AddPrerequisites(CompressTask);
 		}
+		Entry->FinishCompressionBarrier.Trigger();
 	}
 
 	void BeginEncryptAndSign(FIoStoreWriteQueueEntry* Entry)
 	{
 		if (ContainerSettings.IsEncrypted() || ContainerSettings.IsSigned())
 		{
-			FFunctionGraphTask::CreateAndDispatchWhenReady([this, Entry]()
+			UE::Tasks::FTask EncryptAndSignTask = UE::Tasks::Launch(TEXT("EncryptAndSign"), [this, Entry]()
 			{
 				EncryptAndSign(Entry);
-				Entry->BeginWriteBarrier->DispatchSubsequents();
-			}, TStatId(), nullptr, ENamedThreads::AnyHiPriThreadHiPriTask);
+			}, UE::Tasks::ETaskPriority::High);
+			Entry->BeginWriteBarrier.AddPrerequisites(EncryptAndSignTask);
+			Entry->BeginWriteBarrier.Trigger();
 		}
 		else
 		{
 			EncryptAndSign(Entry);
-			Entry->BeginWriteBarrier->DispatchSubsequents();
+			Entry->BeginWriteBarrier.Trigger();
 		}
 	}
 
@@ -1799,7 +1804,7 @@ void FIoStoreWriterContextImpl::Flush()
 		TRACE_CPUPROFILER_EVENT_SCOPE(WaitForChunkHashes);
 		for (int32 EntryIndex = AllEntries.Num() - 1; EntryIndex >= 0; --EntryIndex)
 		{
-			FTaskGraphInterface::Get().WaitUntilTaskCompletes(AllEntries[EntryIndex]->HashTask);
+			AllEntries[EntryIndex]->HashTask.Wait();
 		}
 	}
 	for (TSharedPtr<FIoStoreWriter> IoStoreWriter : IoStoreWriters)
@@ -1890,7 +1895,7 @@ void FIoStoreWriterContextImpl::ScheduleEntry(FIoStoreWriteQueueEntry* Entry)
 	if (Entry->bLoadingFromReferenceDb)
 	{
 		// We don't need to wait on a source read so we can kick directly.
-		Entry->BeginCompressionBarrier->DispatchSubsequents();
+		Entry->BeginCompressionBarrier.Trigger();
 	}
 	else
 	{
@@ -1911,7 +1916,7 @@ void FIoStoreWriterContextImpl::BeginCompressionThreadFunc()
 		while (Entry)
 		{
 			FIoStoreWriteQueueEntry* Next = Entry->Next;
-			Entry->BeginCompressionBarrier->Wait();
+			Entry->BeginCompressionBarrier.Wait();
 			TRACE_COUNTER_INCREMENT(IoStoreBeginCompressionCount);
 			Entry->Writer->BeginCompress(Entry);
 			BeginEncryptionAndSigningQueue.Enqueue(Entry);
@@ -1934,7 +1939,7 @@ void FIoStoreWriterContextImpl::BeginEncryptionAndSigningThreadFunc()
 		while (Entry)
 		{
 			FIoStoreWriteQueueEntry* Next = Entry->Next;
-			Entry->FinishCompressionBarrier->Wait();
+			Entry->FinishCompressionBarrier.Wait();
 			Entry->Request->FreeSourceBuffer();
 			if (Entry->bStoreCompressedDataInDDC)
 			{
@@ -1968,7 +1973,7 @@ void FIoStoreWriterContextImpl::WriterThreadFunc()
 		while (Entry)
 		{
 			FIoStoreWriteQueueEntry* Next = Entry->Next;
-			Entry->BeginWriteBarrier->Wait();
+			Entry->BeginWriteBarrier.Wait();
 			TRACE_COUNTER_INCREMENT(IoStoreBeginWriteCount);
 			Entry->Writer->WriteEntry(Entry);
 			Entry = Next;
diff --git a/Engine/Source/Developer/IoStoreUtilities/Private/IoStoreWriter.h b/Engine/Source/Developer/IoStoreUtilities/Private/IoStoreWriter.h
--- a/Engine/Source/Developer/IoStoreUtilities/Private/IoStoreWriter.h
+++ b/Engine/Source/Developer/IoStoreUtilities/Private/IoStoreWriter.h
@@ -103,7 +103,7 @@ public:
 	virtual const FIoHash* GetChunkHash() = 0;
 
 	// Launches any async operations necessary in order to access the buffer. CompletionEvent is set once it's ready, which may be immediate.
-	virtual void PrepareSourceBufferAsync(FGraphEventRef CompletionEvent) = 0;
+	virtual void PrepareSourceBufferAsync(UE::Tasks::FTaskEvent& CompletionEvent) = 0;
 
 	// Only valid after the completion event passed to PrepareSourceBufferAsync has fired.
 	virtual const FIoBuffer* GetSourceBuffer() = 0;
