UnrealPak: Convert RetrieveChunk in ReferenceChunkDatabase from the old TaskGraph API to the new UE::Tasks API.
Break out a separate LoadFromReferenceDb function from BeginCompress and run it earlier from ScheduleEntry instead.
Enable task retraction for the compressed read when waiting on the BeginCompressionBarrier.

#rnx

diff --git a/Engine/Source/Developer/IoStoreUtilities/Private/IoStoreUtilities.cpp b/Engine/Source/Developer/IoStoreUtilities/Private/IoStoreUtilities.cpp
--- a/Engine/Source/Developer/IoStoreUtilities/Private/IoStoreUtilities.cpp
+++ b/Engine/Source/Developer/IoStoreUtilities/Private/IoStoreUtilities.cpp
@@ -417,19 +417,20 @@ public:
 
 	// This function was not written to be thread safe as it's only ever called from
 	// the iostore begindispatch thread (i.e. is not async)
-	virtual bool RetrieveChunk(const FIoContainerId& InContainerId, const FIoHash& InChunkHash, const FIoChunkId& InChunkId, TUniqueFunction<void(TIoStatusOr<FIoStoreCompressedReadResult>)> InCompleteCallback)
+	virtual UE::Tasks::FTask RetrieveChunk(const FIoContainerId& InContainerId, const FIoHash& InChunkHash, const FIoChunkId& InChunkId, TUniqueFunction<void(TIoStatusOr<FIoStoreCompressedReadResult>)> InCompleteCallback)
 	{
 		if (!bValid)
 		{
-			return false;
+			InCompleteCallback(FIoStatus(EIoErrorCode::InvalidCode, TEXT("IoStoreChunkDatabase not initialized")));
+			return UE::Tasks::MakeCompletedTask<void>();
 		}
 
 		TUniquePtr<FReaderChunks>* ReaderChunksPtr = ChunkDatabase.Find(InContainerId);
 		if (ReaderChunksPtr == nullptr)
 		{
 			// This should never happen now as we wrap this in a ChunkExists call.
-			UE_LOG(LogIoStore, Warning, TEXT("RetrieveChunk can't find the container  invariant violated!"));
-			return false;
+			InCompleteCallback(FIoStatus(EIoErrorCode::InvalidCode, TEXT("RetrieveChunk can't find the container - invariant violated!")));
+			return UE::Tasks::MakeCompletedTask<void>();
 		}
 
 		FReaderChunks* ReaderChunks = ReaderChunksPtr->Get();
@@ -438,8 +439,8 @@ public:
 		if (ChunkIndex == nullptr)
 		{
 			// This should never happen now as we wrap this in a ChunkExists call.
-			UE_LOG(LogIoStore, Warning, TEXT("RetrieveChunk can't find the chunk - invariant violated!"));
-			return false;
+			InCompleteCallback(FIoStatus(EIoErrorCode::InvalidCode, TEXT("RetrieveChunk can't find the chunk - invariant violated!")));
+			return UE::Tasks::MakeCompletedTask<void>();
 		}
 
 		FIoStoreTocChunkInfo& ChunkInfo = ReaderChunks->ChunkInfos[*ChunkIndex];
@@ -462,13 +463,11 @@ public:
 		//
 		// At this point we know we can use the block so we can go async.
 		//
-		FFunctionGraphTask::CreateAndDispatchWhenReady([this, Id = InChunkId, ReaderIndex = ReaderChunks->ReaderIndex, CompleteCallback = MoveTemp(InCompleteCallback)]()
+		return UE::Tasks::Launch(TEXT("ReadCompressed"), [this, Id = InChunkId, ReaderIndex = ReaderChunks->ReaderIndex, CompleteCallback = MoveTemp(InCompleteCallback)]()
 		{
 			TIoStatusOr<FIoStoreCompressedReadResult> Result = Readers[ReaderIndex]->ReadCompressed(Id, FIoReadOptions());
 			CompleteCallback(Result);
-		}, TStatId(), nullptr, ENamedThreads::AnyHiPriThreadNormalTask);
-
-		return true;
+		}, UE::Tasks::ETaskPriority::Normal);
 	}
 
 	void WriteCSV(const FString& InOutputFileName, const TArray<struct FCookedPackage*>& InPackages);
diff --git a/Engine/Source/Developer/IoStoreUtilities/Private/IoStoreWriter.cpp b/Engine/Source/Developer/IoStoreUtilities/Private/IoStoreWriter.cpp
--- a/Engine/Source/Developer/IoStoreUtilities/Private/IoStoreWriter.cpp
+++ b/Engine/Source/Developer/IoStoreUtilities/Private/IoStoreWriter.cpp
@@ -1309,87 +1309,94 @@ private:
 		return WriterContext->CompressionBufferSize * NumBlocks;
 	}
 
-	void BeginCompress(FIoStoreWriteQueueEntry* Entry)
+	void LoadFromReferenceDb(FIoStoreWriteQueueEntry* Entry)
 	{
-		WriterContext->BeginCompressChunksByType[(int8)Entry->ChunkId.GetChunkType()].IncrementExchange();
-		
-		if (Entry->bLoadingFromReferenceDb)
+		if (Entry->NumChunkBlocks == 0)
 		{
-			if (Entry->NumChunkBlocks == 0)
-			{
-				Entry->FinishCompressionBarrier.Trigger();
-				TRACE_COUNTER_INCREMENT(IoStoreRefDbDone);
-				return;
-			}
-
-			// Allocate resources before launching the read tasks to reduce contention. Note this will
-			// allocate iobuffers big enough for uncompressed size, when we only actually need it for
-			// compressed size.
-			Entry->ChunkBlocks.SetNum(Entry->NumChunkBlocks);
-			for (int32 BlockIndex = 0; BlockIndex < Entry->NumChunkBlocks; ++BlockIndex)
-			{
-				FChunkBlock& Block = Entry->ChunkBlocks[BlockIndex];
-				Block.IoBuffer = WriterContext->AllocCompressionBuffer();
-				// Everything else in a block gets filled out from the refdb.
-			}
-
-			// Valid chunks must create the same decompressed bits, but can have different compressed bits.
-			// Since we are on a lightweight dispatch thread, the actual read is async, as is the processing
-			// of the results.
-			TRACE_COUNTER_INCREMENT(IoStoreRefDbInflight);
-			bool bChunkExists = ReferenceChunkDatabase->RetrieveChunk(ContainerSettings.ContainerId, Entry->ChunkHash, Entry->ChunkId, [this, Entry](TIoStatusOr<FIoStoreCompressedReadResult> InReadResult)
-			{
+			Entry->BeginCompressionBarrier.Trigger();
+			TRACE_COUNTER_INCREMENT(IoStoreRefDbDone);
+			return;
+		}
 
-				// If we fail here, in order to recover we effectively need to re-kick this chunk's
-				// BeginCompress() as well as source buffer read... however, this is just a direct read and should only fail
-				// in catastrophic scenarios (loss of connection on a network drive?).
-				FIoStoreCompressedReadResult ReadResult = InReadResult.ValueOrDie();
+		// Allocate resources before launching the read tasks to reduce contention. Note this will
+		// allocate iobuffers big enough for uncompressed size, when we only actually need it for
+		// compressed size.
+		Entry->ChunkBlocks.SetNum(Entry->NumChunkBlocks);
+		for (int32 BlockIndex = 0; BlockIndex < Entry->NumChunkBlocks; ++BlockIndex)
+		{
+			FChunkBlock& Block = Entry->ChunkBlocks[BlockIndex];
+			Block.IoBuffer = WriterContext->AllocCompressionBuffer();
+			// Everything else in a block gets filled out from the refdb.
+		}
 
-				uint64 TotalUncompressedSize = 0;
-				uint8* ReferenceData = ReadResult.IoBuffer.GetData();
-				uint64 TotalAlignedSize = 0;
-				for (int32 BlockIndex = 0; BlockIndex < ReadResult.Blocks.Num(); ++BlockIndex)
-				{
-					FIoStoreCompressedBlockInfo& ReferenceBlock = ReadResult.Blocks[BlockIndex];
-					FChunkBlock& Block = Entry->ChunkBlocks[BlockIndex];
-					Block.CompressionMethod = ReferenceBlock.CompressionMethod;
-					Block.CompressedSize = ReferenceBlock.CompressedSize;
-					Block.UncompressedSize = ReferenceBlock.UncompressedSize;
-					TotalUncompressedSize += ReferenceBlock.UncompressedSize;
-
-					// Future optimization: ReadCompressed returns the memory ready to encrypt in one
-					// large contiguous buffer (i.e. padded). We could use the FIoBuffer functionality of referencing a 
-					// sub block from a parent buffer, however this would mean that we need to add support
-					// for tracking the memory usage in order to remain within our prescribed limits. To do this
-					// requires releasing the entire chunk's memory at once after WriteEntry.
-					// As it stands, we temporarily use untracked memory in the ReadCompressed call (in RetrieveChunk),
-					// then immediately copy it to tracked memory. There's some waste as tracked memory is mod CompressionBlockSize
-					// and we are post compression, so with the average 50% compression rate, we're using double the memory
-					// we "could".
-					FMemory::Memcpy(Block.IoBuffer->GetData(), ReferenceData, Block.CompressedSize);
-					ReferenceData += ReferenceBlock.AlignedSize;
-					TotalAlignedSize += ReferenceBlock.AlignedSize;
-				}
+		// Valid chunks must create the same decompressed bits, but can have different compressed bits.
+		// Since we are on a lightweight dispatch thread, the actual read is async, as is the processing
+		// of the results.
+		TRACE_COUNTER_INCREMENT(IoStoreRefDbInflight);
+		UE::Tasks::FTask RetrieveChunkTask = ReferenceChunkDatabase->RetrieveChunk(
+			ContainerSettings.ContainerId, Entry->ChunkHash, Entry->ChunkId,
+			[this, Entry](TIoStatusOr<FIoStoreCompressedReadResult> InReadResult)
+		{
 
-				if (TotalAlignedSize != ReadResult.IoBuffer.GetSize())
-				{
-					// If we hit this, we might have read garbage memory above! This is very bad.
-					UE_LOG(LogIoStore, Error, TEXT("Block aligned size does not match iobuffer source size! Blocks: %s source size: %s"),
-						*FText::AsNumber(TotalAlignedSize).ToString(),
-						*FText::AsNumber(ReadResult.IoBuffer.GetSize()).ToString());
-				}
+			// If we fail here, in order to recover we effectively need to re-kick this chunk's
+			// BeginCompress() as well as source buffer read... however, this is just a direct read and should only fail
+			// in catastrophic scenarios (loss of connection on a network drive?).
+			UE_CLOG(!InReadResult.IsOk(), LogIoStore, Error, TEXT("RetrieveChunk from ReferenceChunkDatabase failed: %s"),
+				*InReadResult.Status().ToString());
+			FIoStoreCompressedReadResult ReadResult = InReadResult.ValueOrDie();
 
-				Entry->UncompressedSize.Emplace(TotalUncompressedSize);
-				Entry->FinishCompressionBarrier.Trigger();
-				TRACE_COUNTER_DECREMENT(IoStoreRefDbInflight);
-				TRACE_COUNTER_INCREMENT(IoStoreRefDbDone);
-			});
+			uint64 TotalUncompressedSize = 0;
+			uint8* ReferenceData = ReadResult.IoBuffer.GetData();
+			uint64 TotalAlignedSize = 0;
+			for (int32 BlockIndex = 0; BlockIndex < ReadResult.Blocks.Num(); ++BlockIndex)
+			{
+				FIoStoreCompressedBlockInfo& ReferenceBlock = ReadResult.Blocks[BlockIndex];
+				FChunkBlock& Block = Entry->ChunkBlocks[BlockIndex];
+				Block.CompressionMethod = ReferenceBlock.CompressionMethod;
+				Block.CompressedSize = ReferenceBlock.CompressedSize;
+				Block.UncompressedSize = ReferenceBlock.UncompressedSize;
+				TotalUncompressedSize += ReferenceBlock.UncompressedSize;
+
+				// Future optimization: ReadCompressed returns the memory ready to encrypt in one
+				// large contiguous buffer (i.e. padded). We could use the FIoBuffer functionality of referencing a 
+				// sub block from a parent buffer, however this would mean that we need to add support
+				// for tracking the memory usage in order to remain within our prescribed limits. To do this
+				// requires releasing the entire chunk's memory at once after WriteEntry.
+				// As it stands, we temporarily use untracked memory in the ReadCompressed call (in RetrieveChunk),
+				// then immediately copy it to tracked memory. There's some waste as tracked memory is mod CompressionBlockSize
+				// and we are post compression, so with the average 50% compression rate, we're using double the memory
+				// we "could".
+				FMemory::Memcpy(Block.IoBuffer->GetData(), ReferenceData, Block.CompressedSize);
+				ReferenceData += ReferenceBlock.AlignedSize;
+				TotalAlignedSize += ReferenceBlock.AlignedSize;
+			}
+
+			if (TotalAlignedSize != ReadResult.IoBuffer.GetSize())
+			{
+				// If we hit this, we might have read garbage memory above! This is very bad.
+				UE_LOG(LogIoStore, Error, TEXT("Block aligned size does not match iobuffer source size! Blocks: %s source size: %s"),
+					*FText::AsNumber(TotalAlignedSize).ToString(),
+					*FText::AsNumber(ReadResult.IoBuffer.GetSize()).ToString());
+			}
+
+			Entry->UncompressedSize.Emplace(TotalUncompressedSize);
+			TRACE_COUNTER_DECREMENT(IoStoreRefDbInflight);
+			TRACE_COUNTER_INCREMENT(IoStoreRefDbDone);
+		});
+		Entry->BeginCompressionBarrier.AddPrerequisites(RetrieveChunkTask);
+		Entry->BeginCompressionBarrier.Trigger();
 
-			check(bChunkExists); // Sanity - should never return false as we can only get here if ChunkExists() returns true.
+		WriterContext->RefDbChunksCount.IncrementExchange();
+		WriterContext->RefDbChunksByType[(int8)Entry->ChunkId.GetChunkType()].IncrementExchange();
+	}
 
-			WriterContext->RefDbChunksCount.IncrementExchange();
-			WriterContext->RefDbChunksByType[(int8)Entry->ChunkId.GetChunkType()].IncrementExchange();
-			// Lambda handles dispatch subsequents
+	void BeginCompress(FIoStoreWriteQueueEntry* Entry)
+	{
+		WriterContext->BeginCompressChunksByType[(int8)Entry->ChunkId.GetChunkType()].IncrementExchange();
+		
+		if (Entry->bLoadingFromReferenceDb)
+		{
+			Entry->FinishCompressionBarrier.Trigger();
 			return;
 		}
 
@@ -1894,8 +1901,7 @@ void FIoStoreWriterContextImpl::ScheduleEntry(FIoStoreWriteQueueEntry* Entry)
 
 	if (Entry->bLoadingFromReferenceDb)
 	{
-		// We don't need to wait on a source read so we can kick directly.
-		Entry->BeginCompressionBarrier.Trigger();
+		Entry->Writer->LoadFromReferenceDb(Entry);
 	}
 	else
 	{
diff --git a/Engine/Source/Developer/IoStoreUtilities/Private/IoStoreWriter.h b/Engine/Source/Developer/IoStoreUtilities/Private/IoStoreWriter.h
--- a/Engine/Source/Developer/IoStoreUtilities/Private/IoStoreWriter.h
+++ b/Engine/Source/Developer/IoStoreUtilities/Private/IoStoreWriter.h
@@ -137,7 +137,7 @@ public:
 	* Chunks provided *MUST* decompress to bits that hash to the exact value provided in InChunkKey (i.e. be exactly the same bits),
 	* and also be the same number of blocks (i.e. same CompressionBlockSize)
 	*/
-	virtual bool RetrieveChunk(const FIoContainerId& InContainerId, const FIoHash& InChunkHash, const FIoChunkId& InChunkId, TUniqueFunction<void(TIoStatusOr<FIoStoreCompressedReadResult>)> InCompletionCallback) = 0;
+	virtual UE::Tasks::FTask RetrieveChunk(const FIoContainerId& InContainerId, const FIoHash& InChunkHash, const FIoChunkId& InChunkId, TUniqueFunction<void(TIoStatusOr<FIoStoreCompressedReadResult>)> InCompletionCallback) = 0;
 
 	/* 
 	* Quick synchronous existence check that returns the number of blocks for the chunk. This is used to set up
