Make sure we can read large files when saving loose files to non-local remote zenserver

diff --git a/Engine/Source/Developer/IoStoreUtilities/Private/ZenStoreWriter.cpp b/Engine/Source/Developer/IoStoreUtilities/Private/ZenStoreWriter.cpp
--- a/Engine/Source/Developer/IoStoreUtilities/Private/ZenStoreWriter.cpp
+++ b/Engine/Source/Developer/IoStoreUtilities/Private/ZenStoreWriter.cpp
@@ -1453,20 +1453,26 @@ void FZenStoreWriter::CreateProjectMetaData(FCbPackage& Pkg, FCbWriter& PackageO
 				else
 				{
 					const FString AbsPath = ZenFileSystemManifest->ServerRootPath() / NewEntry.ServerPath;
-					TArray<uint8> FileBuffer;
-					FFileHelper::LoadFileToArray(FileBuffer, *AbsPath);
-					if (FileBuffer.Num())
+					TUniquePtr<FArchive> Reader(IFileManager::Get().CreateFileReader(*AbsPath, 0));
+					if (Reader)
 					{
-						FCbAttachment FileAttachment = CreateAttachment(FIoBuffer(FIoBuffer::Clone, FileBuffer.GetData(), FileBuffer.Num()));
-
-						PackageObj.BeginObject();
-						PackageObj << "id" << FileOid;
-						PackageObj << "data" << FileAttachment;
-						PackageObj << "serverpath" << NewEntry.ServerPath;
-						PackageObj << "clientpath" << NewEntry.ClientPath;
-						PackageObj.EndObject();
-
-						Pkg.AddAttachment(FileAttachment);
+						int64 TotalSize = Reader->TotalSize();
+						if (TotalSize > 0)
+						{
+							FIoBuffer FileBuffer(TotalSize);
+							Reader->Serialize(FileBuffer.GetData(), TotalSize);
+							bool Success = Reader->Close();
+							FCbAttachment FileAttachment = CreateAttachment(MoveTemp(FileBuffer));
+
+							PackageObj.BeginObject();
+							PackageObj << "id" << FileOid;
+							PackageObj << "data" << FileAttachment;
+							PackageObj << "serverpath" << NewEntry.ServerPath;
+							PackageObj << "clientpath" << NewEntry.ClientPath;
+							PackageObj.EndObject();
+
+							Pkg.AddAttachment(FileAttachment);
+						}
 					}
 				}
 			}
