Add recovery of zenserver from project store operations in addition to the DDC recovery logic

diff --git a/Engine/Source/Developer/IoStoreUtilities/Internal/ZenStoreHttpClient.h b/Engine/Source/Developer/IoStoreUtilities/Internal/ZenStoreHttpClient.h
--- a/Engine/Source/Developer/IoStoreUtilities/Internal/ZenStoreHttpClient.h
+++ b/Engine/Source/Developer/IoStoreUtilities/Internal/ZenStoreHttpClient.h
@@ -19,6 +19,8 @@ class FCbObject;
 namespace UE {
 	namespace Zen {
 		struct FZenHttpRequestPool;
+		class FZenHttpRequest;
+		enum class EContentType;
 	}
 
 /**
@@ -68,6 +70,11 @@ public:
 
 private:
 	TIoStatusOr<FIoBuffer> ReadOpLogUri(FStringBuilderBase& ChunkUri, uint64 Offset = 0, uint64 Size = ~0ull);
+	bool Download(Zen::FZenHttpRequest& Request, FStringView Uri, TArray64<uint8>* Buffer, Zen::EContentType AcceptType);
+	bool Post(Zen::FZenHttpRequest& Request, FStringView Uri, FCbObjectView Obj);
+	bool Post(Zen::FZenHttpRequest& Request, FStringView Uri, FMemoryView Payload);
+	bool Delete(Zen::FZenHttpRequest& Request, FStringView Uri);
+	bool ShouldRecoverAndRetry(Zen::FZenHttpRequest& Request);
 
 	static const uint32 PoolEntryCount;
 	struct SaltGenerator
diff --git a/Engine/Source/Developer/IoStoreUtilities/Private/ZenStoreHttpClient.cpp b/Engine/Source/Developer/IoStoreUtilities/Private/ZenStoreHttpClient.cpp
--- a/Engine/Source/Developer/IoStoreUtilities/Private/ZenStoreHttpClient.cpp
+++ b/Engine/Source/Developer/IoStoreUtilities/Private/ZenStoreHttpClient.cpp
@@ -81,12 +81,11 @@ FZenStoreHttpClient::TryCreateProject(FStringView InProjectId,
 		TStringBuilder<128> ProjectUri;
 		ProjectUri << "/prj/" << InProjectId;
 		TArray64<uint8> GetBuffer;
-		UE::Zen::FZenHttpRequest::Result Res = Request->PerformBlockingDownload(ProjectUri, &GetBuffer, Zen::EContentType::CbObject);
 
 		// TODO: how to handle failure here? This is probably the most likely point of failure
 		// if the service is not up or not responding
 
-		if (Res == Zen::FZenHttpRequest::Result::Success && Request->GetResponseCode() == 200)
+		if (Download(*Request, ProjectUri, &GetBuffer, Zen::EContentType::CbObject) && Request->GetResponseCode() == 200)
 		{
 			UE_LOG(LogZenStore, Display, TEXT("Zen project '%s' already exists"), *FString(InProjectId));
 			bConnectionSucceeded = true;
@@ -104,9 +103,7 @@ FZenStoreHttpClient::TryCreateProject(FStringView InProjectId,
 			ProjInfo << "projectfile" << ProjectFilePath;
 			ProjInfo.EndObject();
 
-			Res = Request->PerformBlockingPost(ProjectUri, ProjInfo.Save().AsObject());
-
-			if (Res != Zen::FZenHttpRequest::Result::Success)
+			if (!Post(*Request, ProjectUri, ProjInfo.Save().AsObject()))
 			{
 				UE_LOG(LogZenStore, Display, TEXT("Zen project '%s' creation FAILED"), *FString(InProjectId));
 				bConnectionSucceeded = false;
@@ -148,15 +145,13 @@ bool FZenStoreHttpClient::TryCreateOplog(FStringView InProjectId, FStringView In
 	if (bFullBuild)
 	{
 		UE_LOG(LogZenStore, Display, TEXT("Deleting oplog '%s/%s' if it exists"), *FString(InProjectId), *FString(InOplogId));
-		Request->PerformBlockingDelete(OplogPath);
+		Delete(*Request, OplogPath);
 		Request->Reset();
 	}
 
 	TArray64<uint8> GetBuffer;
-	UE::Zen::FZenHttpRequest::Result Res = Request->PerformBlockingDownload(OplogPath, &GetBuffer, Zen::EContentType::CbObject);
 	FCbObjectView OplogInfo;
-
-	if (Res == Zen::FZenHttpRequest::Result::Success && Request->GetResponseCode() == 200)
+	if (Download(*Request, OplogPath, &GetBuffer, Zen::EContentType::CbObject) && Request->GetResponseCode() == 200)
 	{
 		UE_LOG(LogZenStore, Display, TEXT("Zen oplog '%s/%s' already exists"), *FString(InProjectId), *FString(InOplogId));
 
@@ -173,9 +168,7 @@ bool FZenStoreHttpClient::TryCreateOplog(FStringView InProjectId, FStringView In
 
 		Request->Reset();
 
-		Res = Request->PerformBlockingPost(OplogPath, OplogCreateInfo.AsObjectView());
-
-		if (Res != Zen::FZenHttpRequest::Result::Success)
+		if (!Post(*Request, OplogPath, OplogCreateInfo.AsObjectView()))
 		{
 			UE_LOG(LogZenStore, Error, TEXT("Zen oplog '%s/%s' creation FAILED"), *FString(InProjectId), *FString(InOplogId));
 			// Demote the connection status back to not connected
@@ -195,9 +188,7 @@ bool FZenStoreHttpClient::TryCreateOplog(FStringView InProjectId, FStringView In
 
 		GetBuffer.Reset();
 		Request->Reset();
-		Res = Request->PerformBlockingDownload(OplogPath, &GetBuffer, Zen::EContentType::CbObject);
-
-		if (Res == Zen::FZenHttpRequest::Result::Success && Request->GetResponseCode() == 200)
+		if (Download(*Request, OplogPath, &GetBuffer, Zen::EContentType::CbObject) && Request->GetResponseCode() == 200)
 		{
 			OplogInfo = FCbObjectView(GetBuffer.GetData());
 		}
@@ -219,12 +210,11 @@ void FZenStoreHttpClient::InitializeReadOnly(FStringView InProjectId, FStringVie
 		UE::Zen::FZenScopedRequestPtr Request(RequestPool.Get());
 
 		TArray64<uint8> GetBuffer;
-		UE::Zen::FZenHttpRequest::Result Res = Request->PerformBlockingDownload(WriteToString<128>("/prj/", InProjectId), &GetBuffer, Zen::EContentType::CbObject);
 
 		// TODO: how to handle failure here? This is probably the most likely point of failure
 		// if the service is not up or not responding
 
-		if (Res != Zen::FZenHttpRequest::Result::Success || Request->GetResponseCode() != 200)
+		if (!Download(*Request, WriteToString<128>("/prj/", InProjectId), &GetBuffer, Zen::EContentType::CbObject) || Request->GetResponseCode() != 200)
 		{
 			UE_LOG(LogZenStore, Fatal, TEXT("Zen project '%s' not found"), *FString(InProjectId));
 			bConnectionSucceeded = false;
@@ -244,9 +234,7 @@ void FZenStoreHttpClient::InitializeReadOnly(FStringView InProjectId, FStringVie
 		OplogPath = WriteToString<128>("/prj/", InProjectId, "/oplog/", InOplogId);
 
 		TArray64<uint8> GetBuffer;
-		UE::Zen::FZenHttpRequest::Result Res = Request->PerformBlockingDownload(OplogPath, &GetBuffer, Zen::EContentType::CbObject);
-
-		if (Res != Zen::FZenHttpRequest::Result::Success || Request->GetResponseCode() != 200)
+		if (!Download(*Request, OplogPath, &GetBuffer, Zen::EContentType::CbObject) || Request->GetResponseCode() != 200)
 		{
 			UE_LOG(LogZenStore, Fatal, TEXT("Zen oplog '%s'/'%s' not found"), *FString(InProjectId), *FString(InOplogId));
 		}
@@ -255,6 +243,71 @@ void FZenStoreHttpClient::InitializeReadOnly(FStringView InProjectId, FStringVie
 	bAllowRead = true;
 }
 
+bool FZenStoreHttpClient::Download(Zen::FZenHttpRequest& Request, FStringView Uri, TArray64<uint8>* Buffer, Zen::EContentType AcceptType)
+{
+	Zen::FZenHttpRequest::Result Res = Request.PerformBlockingDownload(Uri, Buffer, AcceptType);
+	if (ShouldRecoverAndRetry(Request))
+	{
+		Request.Reset();
+		Res = Request.PerformBlockingDownload(Uri, Buffer, AcceptType);
+	}
+	return Res == Zen::FZenHttpRequest::Result::Success;
+}
+
+bool FZenStoreHttpClient::Post(Zen::FZenHttpRequest& Request, FStringView Uri, FCbObjectView Obj)
+{
+	Zen::FZenHttpRequest::Result Res = Request.PerformBlockingPost(Uri, Obj);
+	if (ShouldRecoverAndRetry(Request))
+	{
+		Request.Reset();
+		Res = Request.PerformBlockingPost(Uri, Obj);
+	}
+	return Res == Zen::FZenHttpRequest::Result::Success;
+}
+
+bool FZenStoreHttpClient::Post(Zen::FZenHttpRequest& Request, FStringView Uri, FMemoryView Payload)
+{
+	Zen::FZenHttpRequest::Result Res = Request.PerformBlockingPost(Uri, Payload);
+	if (ShouldRecoverAndRetry(Request))
+	{
+		Request.Reset();
+		Res = Request.PerformBlockingPost(Uri, Payload);
+	}
+	return Res == Zen::FZenHttpRequest::Result::Success;
+}
+
+bool FZenStoreHttpClient::Delete(Zen::FZenHttpRequest& Request, FStringView Uri)
+{
+	Zen::FZenHttpRequest::Result Res = Request.PerformBlockingDelete(Uri);
+	if (ShouldRecoverAndRetry(Request))
+	{
+		Request.Reset();
+		Res = Request.PerformBlockingDelete(Uri);
+	}
+	return Res == Zen::FZenHttpRequest::Result::Success;
+}
+
+bool FZenStoreHttpClient::ShouldRecoverAndRetry(Zen::FZenHttpRequest& Request)
+{
+	if (!ZenService.GetInstance().IsServiceRunningLocally())
+	{
+		return false;
+	}
+
+	const int ConnectError = 7;
+	const int SSLConnectError = 35;
+	const int TimeoutError = 28;
+
+	if ((Request.GetResponseCode() == ConnectError) ||
+		(Request.GetResponseCode() == SSLConnectError) ||
+		(Request.GetResponseCode() == TimeoutError))
+	{
+		return true;
+	}
+
+	return false;
+}
+
 TIoStatusOr<uint64> FZenStoreHttpClient::AppendOp(FCbPackage OpEntry)
 {
 	check(bAllowEdit);
@@ -300,15 +353,12 @@ TIoStatusOr<uint64> FZenStoreHttpClient::AppendOp(FCbPackage OpEntry)
 
 				FCbFieldIterator Prep = Writer.Save();
 
-				bool IsOk = false;
-				
-				const Zen::FZenHttpRequest::Result Res = Request->PerformBlockingPost(OplogPrepNewEntryPath, Prep.AsObjectView());
-
-				if (Res == Zen::FZenHttpRequest::Result::Success)
+				bool IsOk = false;			
+				if (Post(*Request, OplogPrepNewEntryPath, Prep.AsObjectView()))
 				{
 					FCbObjectView NeedObject;
 
-					if (Res == Zen::FZenHttpRequest::Result::Success && Request->GetResponseCode() == 200)
+					if (Request->GetResponseCode() == 200)
 					{
 						NeedObject = FCbObjectView(Request->GetResponseBuffer().GetData());
 
@@ -422,7 +472,7 @@ TIoStatusOr<uint64> FZenStoreHttpClient::AppendOp(FCbPackage OpEntry)
 
 		Request->Reset();
 
-		if (UE::Zen::FZenHttpRequest::Result::Success == Request->PerformBlockingPost(NewOpPostUri, SerializedPackage.GetView()))
+		if (Post(*Request, NewOpPostUri, SerializedPackage.GetView()))
 		{
 			return TIoStatusOr<uint64>(SerializedPackage.TotalSize());
 		}
@@ -442,8 +492,7 @@ TIoStatusOr<uint64> FZenStoreHttpClient::GetChunkSize(const FIoChunkId& Id)
 	UE::Zen::FZenScopedRequestPtr Request(RequestPool.Get());
 	TStringBuilder<128> ChunkUri;
 	ChunkUri << OplogPath << '/' << Id << "/info";
-	UE::Zen::FZenHttpRequest::Result Res = Request->PerformBlockingDownload(ChunkUri, nullptr, Zen::EContentType::CbObject);
-	if (Res == Zen::FZenHttpRequest::Result::Success && Request->GetResponseCode() == 200)
+	if (Download(*Request, ChunkUri, nullptr, Zen::EContentType::CbObject) && Request->GetResponseCode() == 200)
 	{
 		FCbObjectView ResponseObj = Request->GetResponseAsObject();
 		const uint64 ChunkSize = ResponseObj["size"].AsUInt64(0);
@@ -501,9 +550,7 @@ TIoStatusOr<FIoBuffer> FZenStoreHttpClient::ReadOpLogUri(FStringBuilderBase& Chu
 		ChunkUri.Appendf(TEXT("size=%" UINT64_FMT), Size);
 	}
 
-	UE::Zen::FZenHttpRequest::Result Res = Request->PerformBlockingDownload(ChunkUri, &GetBuffer, Zen::EContentType::CompressedBinary);
-
-	if (Res == Zen::FZenHttpRequest::Result::Success && Request->GetResponseCode() == 200)
+	if (Download(*Request, ChunkUri, &GetBuffer, Zen::EContentType::CompressedBinary) && Request->GetResponseCode() == 200)
 	{
 		if (FCompressedBuffer Compressed = FCompressedBuffer::FromCompressed(FSharedBuffer::MakeView(GetBuffer.GetData(), GetBuffer.Num())))
 		{
@@ -555,9 +602,7 @@ TFuture<TIoStatusOr<FCbObject>> FZenStoreHttpClient::GetOplog()
 		Uri << OplogPath << "/entries";
 
 		TArray64<uint8> GetBuffer;
-		UE::Zen::FZenHttpRequest::Result Res = Request->PerformBlockingDownload(Uri, &GetBuffer, Zen::EContentType::CbObject);
-
-		if (Res == Zen::FZenHttpRequest::Result::Success && Request->GetResponseCode() == 200)
+		if (Download(*Request, Uri, &GetBuffer, Zen::EContentType::CbObject) && Request->GetResponseCode() == 200)
 		{
 			FCbObjectView Response(GetBuffer.GetData());
 			return TIoStatusOr<FCbObject>(FCbObject::Clone(Response));
@@ -584,9 +629,7 @@ TFuture<TIoStatusOr<FCbObject>> FZenStoreHttpClient::GetFiles()
 		Uri << OplogPath << "/files";
 
 		TArray64<uint8> GetBuffer;
-		UE::Zen::FZenHttpRequest::Result Res = Request->PerformBlockingDownload(Uri, &GetBuffer, Zen::EContentType::CbObject);
-
-		if (Res == Zen::FZenHttpRequest::Result::Success && Request->GetResponseCode() == 200)
+		if (Download(*Request, Uri, &GetBuffer, Zen::EContentType::CbObject) && Request->GetResponseCode() == 200)
 		{
 			FCbObjectView Response(GetBuffer.GetData());
 			return TIoStatusOr<FCbObject>(FCbObject::Clone(Response));
@@ -613,9 +656,7 @@ TFuture<TIoStatusOr<FCbObject>> FZenStoreHttpClient::GetChunkInfos()
 		Uri << OplogPath << "/chunkinfos";
 
 		TArray64<uint8> GetBuffer;
-		UE::Zen::FZenHttpRequest::Result Res = Request->PerformBlockingDownload(Uri, &GetBuffer, Zen::EContentType::CbObject);
-
-		if (Res == Zen::FZenHttpRequest::Result::Success && Request->GetResponseCode() == 200)
+		if (Download(*Request, Uri, &GetBuffer, Zen::EContentType::CbObject) && Request->GetResponseCode() == 200)
 		{
 			FCbObjectView Response(GetBuffer.GetData());
 			return TIoStatusOr<FCbObject>(FCbObject::Clone(Response));
@@ -650,8 +691,7 @@ FZenStoreHttpClient::EndBuildPass(FCbPackage OpEntry)
 	UE::Zen::FZenScopedRequestPtr Request(RequestPool.Get());
 
 	FMemoryView Payload { SerializedPackage.GetData(), (uint64)SerializedPackage.TotalSize()};
-	
-	if (UE::Zen::FZenHttpRequest::Result::Success == Request->PerformBlockingPost(OplogNewEntryPath, Payload))
+	if (Post(*Request, OplogNewEntryPath, Payload))
 	{
 		return static_cast<uint64>(Payload.GetSize());
 	}
diff --git a/Engine/Source/Developer/Zen/Internal/ZenServerHttp.h b/Engine/Source/Developer/Zen/Internal/ZenServerHttp.h
--- a/Engine/Source/Developer/Zen/Internal/ZenServerHttp.h
+++ b/Engine/Source/Developer/Zen/Internal/ZenServerHttp.h
@@ -316,6 +316,12 @@ public:
 		return Request;
 	}
 
+	FZenHttpRequest& operator*()
+	{
+		check(IsValid());
+		return *Request;
+	}
+
 private:
 	FZenHttpRequest*		Request;
 	FZenHttpRequestPool*	Pool;
