UnrealPak: Trivial rename of some IoStoreWriter.cpp variables.
From FChunkBlock::Size to FChunkBlock::DiskSize.
From FIoStoreWriteQueueEntry::CompressedSize to FIoStoreWriteQueueEntry::DiskSize.

#tests identical binary output
#rnx

diff --git a/Engine/Source/Developer/IoStoreUtilities/Private/IoStoreWriter.cpp b/Engine/Source/Developer/IoStoreUtilities/Private/IoStoreWriter.cpp
--- a/Engine/Source/Developer/IoStoreUtilities/Private/IoStoreWriter.cpp
+++ b/Engine/Source/Developer/IoStoreUtilities/Private/IoStoreWriter.cpp
@@ -37,7 +37,7 @@ struct FChunkBlock
 
 	// This is the size of the actual block after encryption alignment, and is
 	// set in EncryptAndSign. This happens whether or not the container is encrypted.
-	uint64 Size = 0;
+	uint64 DiskSize = 0;
 	uint64 CompressedSize = 0;
 	uint64 UncompressedSize = 0;
 	FName CompressionMethod = NAME_None;
@@ -65,8 +65,7 @@ struct FIoStoreWriteQueueEntry
 	TOptional<uint64> UncompressedSize;
 
 	// this is not filled out until after encryption completes and *includes the alignment padding for encryption*!
-	// think of this as "size on disk".
-	uint64 CompressedSize = 0; 
+	uint64 DiskSize = 0; 
 
 	uint64 Padding = 0;
 	uint64 Offset = 0;
@@ -1029,12 +1028,12 @@ public:
 				if (Entry->bModified)
 				{
 					++Result.ModifiedChunksCount;
-					Result.ModifiedChunksSize += Entry->CompressedSize;
+					Result.ModifiedChunksSize += Entry->DiskSize;
 				}
 				else if (Entry->bAdded)
 				{
 					++Result.AddedChunksCount;
-					Result.AddedChunksSize += Entry->CompressedSize;
+					Result.AddedChunksSize += Entry->DiskSize;
 				}
 				delete Entry;
 			}
@@ -1538,33 +1537,33 @@ private:
 		for (FChunkBlock& Block : Entry->ChunkBlocks)
 		{
 			// Always align each compressed block to AES block size but store the compressed block size in the TOC
-			Block.Size = Block.CompressedSize;
-			if (!IsAligned(Block.Size, FAES::AESBlockSize))
+			Block.DiskSize = Block.CompressedSize;
+			if (!IsAligned(Block.DiskSize, FAES::AESBlockSize))
 			{
-				uint64 AlignedCompressedBlockSize = Align(Block.Size, FAES::AESBlockSize);
+				uint64 AlignedCompressedBlockSize = Align(Block.DiskSize, FAES::AESBlockSize);
 				uint8* CompressedData = Block.IoBuffer->Data();
-				for (uint64 FillIndex = Block.Size; FillIndex < AlignedCompressedBlockSize; ++FillIndex)
+				for (uint64 FillIndex = Block.DiskSize; FillIndex < AlignedCompressedBlockSize; ++FillIndex)
 				{
 					check(FillIndex < Block.IoBuffer->DataSize());
-					CompressedData[FillIndex] = CompressedData[(FillIndex - Block.Size) % Block.Size];
+					CompressedData[FillIndex] = CompressedData[(FillIndex - Block.DiskSize) % Block.DiskSize];
 				}
-				Block.Size = AlignedCompressedBlockSize;
+				Block.DiskSize = AlignedCompressedBlockSize;
 			}
 
 			if (ContainerSettings.IsEncrypted())
 			{
-				FAES::EncryptData(Block.IoBuffer->Data(), static_cast<uint32>(Block.Size), ContainerSettings.EncryptionKey);
+				FAES::EncryptData(Block.IoBuffer->Data(), static_cast<uint32>(Block.DiskSize), ContainerSettings.EncryptionKey);
 			}
 
 			if (ContainerSettings.IsSigned())
 			{
-				FSHA1::HashBuffer(Block.IoBuffer->Data(), Block.Size, Block.Signature.Hash);
+				FSHA1::HashBuffer(Block.IoBuffer->Data(), Block.DiskSize, Block.Signature.Hash);
 			}
 		}
-		Entry->CompressedSize = 0;
+		Entry->DiskSize = 0;
 		for (const FChunkBlock& ChunkBlock : Entry->ChunkBlocks)
 		{
-			Entry->CompressedSize += ChunkBlock.Size;
+			Entry->DiskSize += ChunkBlock.DiskSize;
 		}
 	}
 
@@ -1609,9 +1608,9 @@ private:
 		if (Entry->PartitionIndex >= 0)
 		{
 			TargetPartition = &Partitions[Entry->PartitionIndex];
-			if (TargetPartition->ReservedSpace > Entry->CompressedSize)
+			if (TargetPartition->ReservedSpace > Entry->DiskSize)
 			{
-				TargetPartition->ReservedSpace -= Entry->CompressedSize;
+				TargetPartition->ReservedSpace -= Entry->DiskSize;
 			}
 			else
 			{
@@ -1624,7 +1623,7 @@ private:
 		bHasMemoryMappedEntry |= Entry->Options.bIsMemoryMapped;
 		const uint64 ChunkAlignment = Entry->Options.bIsMemoryMapped ? WriterSettings.MemoryMappingAlignment : 0;
 		const uint64 PartitionSizeLimit = WriterSettings.MaxPartitionSize > 0 ? WriterSettings.MaxPartitionSize : MAX_uint64;
-		checkf(Entry->CompressedSize <= PartitionSizeLimit, TEXT("Chunk is too large, increase max partition size!"));
+		checkf(Entry->DiskSize <= PartitionSizeLimit, TEXT("Chunk is too large, increase max partition size!"));
 		for (;;)
 		{
 			uint64 OffsetBeforePadding = TargetPartition->Offset;
@@ -1638,14 +1637,14 @@ private:
 				// small entries from causing multiple file system block reads afaict. Large entries necesarily get
 				// aligned to prevent things like a blocksize + 2 entry being at alignment -1, causing 3 low level reads.
 				// ...I think.
-				bool bCrossesBlockBoundary = Align(TargetPartition->Offset, WriterSettings.CompressionBlockAlignment) != Align(TargetPartition->Offset + Entry->CompressedSize - 1, WriterSettings.CompressionBlockAlignment);
+				bool bCrossesBlockBoundary = Align(TargetPartition->Offset, WriterSettings.CompressionBlockAlignment) != Align(TargetPartition->Offset + Entry->DiskSize - 1, WriterSettings.CompressionBlockAlignment);
 				if (bCrossesBlockBoundary)
 				{
 					TargetPartition->Offset = Align(TargetPartition->Offset, WriterSettings.CompressionBlockAlignment);
 				}
 			}
 
-			if (TargetPartition->Offset + Entry->CompressedSize + TargetPartition->ReservedSpace > PartitionSizeLimit)
+			if (TargetPartition->Offset + Entry->DiskSize + TargetPartition->ReservedSpace > PartitionSizeLimit)
 			{
 				TargetPartition->Offset = OffsetBeforePadding;
 				while (Partitions.Num() <= NextPartitionIndexToTry)
@@ -1687,7 +1686,7 @@ private:
 		{
 			FIoStoreTocCompressedBlockEntry& BlockEntry = TocBuilder.AddCompressionBlockEntry();
 			BlockEntry.SetOffset(TargetPartition->Index * WriterSettings.MaxPartitionSize + TargetPartition->Offset + OffsetInChunk);
-			OffsetInChunk += ChunkBlock.Size;
+			OffsetInChunk += ChunkBlock.DiskSize;
 			BlockEntry.SetCompressedSize(uint32(ChunkBlock.CompressedSize));
 			BlockEntry.SetUncompressedSize(uint32(ChunkBlock.UncompressedSize));
 			BlockEntry.SetCompressionMethodIndex(TocBuilder.AddCompressionMethodEntry(ChunkBlock.CompressionMethod));
@@ -1721,7 +1720,7 @@ private:
 		}
 
 		const uint64 RegionStartOffset = TargetPartition->Offset;
-		TargetPartition->Offset += Entry->CompressedSize;
+		TargetPartition->Offset += Entry->DiskSize;
 		UncompressedFileOffset += Align(Entry->UncompressedSize.GetValue(), WriterSettings.CompressionBlockSize);
 		TotalEntryUncompressedSize += Entry->UncompressedSize.GetValue();
 
@@ -1747,8 +1746,8 @@ private:
 		for (FChunkBlock& ChunkBlock : Entry->ChunkBlocks)
 		{
 			TRACE_CPUPROFILER_EVENT_SCOPE(WriteBlockToContainer);
-			TargetPartition->ContainerFileHandle->Serialize(ChunkBlock.IoBuffer->Data(), ChunkBlock.Size);
-			WriteBytes += ChunkBlock.Size;
+			TargetPartition->ContainerFileHandle->Serialize(ChunkBlock.IoBuffer->Data(), ChunkBlock.DiskSize);
+			WriteBytes += ChunkBlock.DiskSize;
 		}
 		uint64 WriteEndCycles = FPlatformTime::Cycles64();
 		WriterContext->WriteCycleCount.AddExchange(WriteEndCycles - WriteStartCycles);
