[TaskGraph/Misc]
 - Replace BusyWait API calls to simple Wait

#rnx

diff --git a/Engine/Plugins/Editor/ContentBrowser/ContentBrowserAssetDataSource/Source/ContentBrowserAssetDataSource/Private/ContentBrowserAssetDataSource.cpp b/Engine/Plugins/Editor/ContentBrowser/ContentBrowserAssetDataSource/Source/ContentBrowserAssetDataSource/Private/ContentBrowserAssetDataSource.cpp
--- a/Engine/Plugins/Editor/ContentBrowser/ContentBrowserAssetDataSource/Source/ContentBrowserAssetDataSource/Private/ContentBrowserAssetDataSource.cpp
+++ b/Engine/Plugins/Editor/ContentBrowser/ContentBrowserAssetDataSource/Source/ContentBrowserAssetDataSource/Private/ContentBrowserAssetDataSource.cpp
@@ -2186,7 +2186,7 @@ void UContentBrowserAssetDataSource::EnumerateItemsMatchingFilter(const FContent
 			}
 		}
 
-		DiskTask.BusyWait();
+		DiskTask.Wait();
 		TArray<FAssetData> DiskAssets = MoveTemp(DiskTask.GetResult());
 		ProduceAssets(DiskAssets, IgnorePackages);
 	}
diff --git a/Engine/Plugins/Experimental/MeshModelingToolsetExp/Source/GeometryProcessingAdapters/Private/CombineMeshInstancesImpl.cpp b/Engine/Plugins/Experimental/MeshModelingToolsetExp/Source/GeometryProcessingAdapters/Private/CombineMeshInstancesImpl.cpp
--- a/Engine/Plugins/Experimental/MeshModelingToolsetExp/Source/GeometryProcessingAdapters/Private/CombineMeshInstancesImpl.cpp
+++ b/Engine/Plugins/Experimental/MeshModelingToolsetExp/Source/GeometryProcessingAdapters/Private/CombineMeshInstancesImpl.cpp
@@ -3156,7 +3156,7 @@ void ProcessCombinedLODChain(
 
 			if (bVerbose)
 			{
-				RemoveHiddenTask.BusyWait();
+				RemoveHiddenTask.Wait();
 			}
 		}
 	}
diff --git a/Engine/Plugins/Experimental/Mutable/Source/CustomizableObject/Private/MuCO/FMutableTaskGraph.cpp b/Engine/Plugins/Experimental/Mutable/Source/CustomizableObject/Private/MuCO/FMutableTaskGraph.cpp
--- a/Engine/Plugins/Experimental/Mutable/Source/CustomizableObject/Private/MuCO/FMutableTaskGraph.cpp
+++ b/Engine/Plugins/Experimental/Mutable/Source/CustomizableObject/Private/MuCO/FMutableTaskGraph.cpp
@@ -86,15 +86,7 @@ void FMutableTaskGraph::WaitForMutableTasks()
 
 	if (LastMutableTask.IsValid())
 	{
-		if (CVarTaskGraphBusyWait->GetBool())
-		{
-			LastMutableTask.BusyWait();
-		}
-		else
-		{
-			LastMutableTask.Wait();
-		}
-
+		LastMutableTask.Wait();
 		LastMutableTask = {};
 	}
 }
diff --git a/Engine/Plugins/Experimental/Mutable/Source/MutableRuntime/Private/MuR/CodeRunnerTasked.cpp b/Engine/Plugins/Experimental/Mutable/Source/MutableRuntime/Private/MuR/CodeRunnerTasked.cpp
--- a/Engine/Plugins/Experimental/Mutable/Source/MutableRuntime/Private/MuR/CodeRunnerTasked.cpp
+++ b/Engine/Plugins/Experimental/Mutable/Source/MutableRuntime/Private/MuR/CodeRunnerTasked.cpp
@@ -555,14 +555,7 @@ namespace mu
 					{
 						if (IssuedTasks[IssuedIndex]->Event.IsValid())
 						{
-							if (CVarTaskGraphBusyWait->GetBool())
-							{
-								IssuedTasks[IssuedIndex]->Event.BusyWait();							
-							}
-							else
-							{
-								IssuedTasks[IssuedIndex]->Event.Wait();
-							}
+							IssuedTasks[IssuedIndex]->Event.Wait();
 
 							break;
 						}
diff --git a/Engine/Plugins/Experimental/Mutable/Source/MutableRuntime/Private/MuR/System.cpp b/Engine/Plugins/Experimental/Mutable/Source/MutableRuntime/Private/MuR/System.cpp
--- a/Engine/Plugins/Experimental/Mutable/Source/MutableRuntime/Private/MuR/System.cpp
+++ b/Engine/Plugins/Experimental/Mutable/Source/MutableRuntime/Private/MuR/System.cpp
@@ -36,12 +36,6 @@
 #include "ProfilingDebugging/CountersTrace.h"
 #include "PackedNormal.h"
 
-TAutoConsoleVariable<bool> CVarTaskGraphBusyWait(
-	TEXT("mutable.TaskGraphBusyWait"),
-	true,
-	TEXT("Use TaskGraph BusyWait instead of simple Wait. Required to avoid hangs on platforms with low number of cores."),
-	ECVF_Default);
-
 namespace 
 {
 
diff --git a/Engine/Plugins/Experimental/Mutable/Source/MutableRuntime/Public/MuR/System.h b/Engine/Plugins/Experimental/Mutable/Source/MutableRuntime/Public/MuR/System.h
--- a/Engine/Plugins/Experimental/Mutable/Source/MutableRuntime/Public/MuR/System.h
+++ b/Engine/Plugins/Experimental/Mutable/Source/MutableRuntime/Public/MuR/System.h
@@ -39,10 +39,6 @@ enum class ETextureCompressionStrategy : uint8
 	NeverCompress
 };
 
-
-MUTABLERUNTIME_API extern TAutoConsoleVariable<bool> CVarTaskGraphBusyWait;
-
-
 namespace mu
 {
 	// Forward references
diff --git a/Engine/Plugins/Runtime/Metasound/Source/MetasoundFrontend/Private/MetasoundFrontendRegistryContainerImpl.cpp b/Engine/Plugins/Runtime/Metasound/Source/MetasoundFrontend/Private/MetasoundFrontendRegistryContainerImpl.cpp
--- a/Engine/Plugins/Runtime/Metasound/Source/MetasoundFrontend/Private/MetasoundFrontendRegistryContainerImpl.cpp
+++ b/Engine/Plugins/Runtime/Metasound/Source/MetasoundFrontend/Private/MetasoundFrontendRegistryContainerImpl.cpp
@@ -1222,7 +1222,7 @@ namespace Metasound::Frontend
 			METASOUND_TRACE_CPUPROFILER_EVENT_SCOPE(Metasound::FRegistryContainerImpl::WaitForRegistrationTaskToComplete);
 			if (Task.IsValid())
 			{
-				Task.BusyWait();
+				Task.Wait();
 			}
 		}
 	}
diff --git a/Engine/Plugins/Runtime/Metasound/Source/MetasoundGenerator/Private/MetasoundOperatorCache.cpp b/Engine/Plugins/Runtime/Metasound/Source/MetasoundGenerator/Private/MetasoundOperatorCache.cpp
--- a/Engine/Plugins/Runtime/Metasound/Source/MetasoundGenerator/Private/MetasoundOperatorCache.cpp
+++ b/Engine/Plugins/Runtime/Metasound/Source/MetasoundGenerator/Private/MetasoundOperatorCache.cpp
@@ -334,7 +334,7 @@ namespace Metasound
 				FTask& TaskToCancel = Pair.Value;
 				if (!TaskToCancel.IsCompleted())
 				{
-					TaskToCancel.BusyWait();
+					TaskToCancel.Wait();
 				}
 			}
 		}
diff --git a/Engine/Source/Developer/IoStoreUtilities/Private/IoStoreUtilities.cpp b/Engine/Source/Developer/IoStoreUtilities/Private/IoStoreUtilities.cpp
--- a/Engine/Source/Developer/IoStoreUtilities/Private/IoStoreUtilities.cpp
+++ b/Engine/Source/Developer/IoStoreUtilities/Private/IoStoreUtilities.cpp
@@ -6992,7 +6992,7 @@ int32 ProfileReadSpeed(const TCHAR* InCommandLine, const FKeyChain& InKeyChain)
 			case EReadType::ReadAsync:
 				{
 					UE::Tasks::TTask<TIoStatusOr<FIoBuffer>> Tasks = Reader->ReadAsync(Id, FIoReadOptions());
-					Tasks.BusyWait();
+					Tasks.Wait();
 					FIoBuffer Result = Tasks.GetResult().ValueOrDie();
 
 					BytesRead += Result.GetSize();
@@ -7037,7 +7037,7 @@ int32 ProfileReadSpeed(const TCHAR* InCommandLine, const FKeyChain& InKeyChain)
 
 	{
 		double WaitStartTime = FPlatformTime::Seconds();
-		CompletedEvent.BusyWait();
+		CompletedEvent.Wait();
 		UE_LOG(LogIoStore, Display, TEXT("Waited %.1f seconds"), FPlatformTime::Seconds() - WaitStartTime);
 	}
 
diff --git a/Engine/Source/Runtime/Core/Private/IO/IoStore.cpp b/Engine/Source/Runtime/Core/Private/IO/IoStore.cpp
--- a/Engine/Source/Runtime/Core/Private/IO/IoStore.cpp
+++ b/Engine/Source/Runtime/Core/Private/IO/IoStore.cpp
@@ -737,7 +737,7 @@ public:
 			// Now, wait for _our_ block's IO
 			{
 				TRACE_CPUPROFILER_EVENT_SCOPE(WaitForIo);
-				ReadRequest.BusyWait();
+				ReadRequest.Wait();
 			}
 
 			if (AsyncReadSucceeded[OurBufferIndex] == false)
@@ -869,7 +869,7 @@ public:
 
 			{
 				TRACE_CPUPROFILER_EVENT_SCOPE(WaitForIo);
-				ReadTask.BusyWait();
+				ReadTask.Wait();
 			}
 
 			if (bReadSucceeded == false)
diff --git a/Engine/Source/Runtime/Engine/Private/SaveGameSystem.cpp b/Engine/Source/Runtime/Engine/Private/SaveGameSystem.cpp
--- a/Engine/Source/Runtime/Engine/Private/SaveGameSystem.cpp
+++ b/Engine/Source/Runtime/Engine/Private/SaveGameSystem.cpp
@@ -292,6 +292,6 @@ void FBaseAsyncSaveGameSystem::WaitForAsyncTask(UE::Tasks::FTask AsyncSaveTask)
 	else
 	{
 		// not running on the game thread, so just block until the async operation comes back
-		AsyncSaveTask.BusyWait();
+		AsyncSaveTask.Wait();
 	}
 }
