Add metadata files (including the DevelopmentAssetRegistry) to the zen filesystem manifest so that snapshots can include the data needed not just to run a game, but also package and iteratively cook one.

diff --git a/Engine/Source/Developer/IoStoreUtilities/Private/ZenFileSystemManifest.cpp b/Engine/Source/Developer/IoStoreUtilities/Private/ZenFileSystemManifest.cpp
--- a/Engine/Source/Developer/IoStoreUtilities/Private/ZenFileSystemManifest.cpp
+++ b/Engine/Source/Developer/IoStoreUtilities/Private/ZenFileSystemManifest.cpp
@@ -233,7 +233,6 @@ int32 FZenFileSystemManifest::Generate()
 
 	FFileFilter CookedFilter = FFileFilter()
 		.ExcludeDirectory(TEXT("Metadata"))
-		.ExcludeExtension(TEXT("json"))
 		.ExcludeExtension(TEXT("uasset"))
 		.ExcludeExtension(TEXT("ubulk"))
 		.ExcludeExtension(TEXT("uexp"))
@@ -241,8 +240,15 @@ int32 FZenFileSystemManifest::Generate()
 		.ExcludeExtension(TEXT("uregs"));
 	AddFilesFromDirectory(TEXT("/{engine}"), FPaths::Combine(CookDirectory, TEXT("Engine")), true, &CookedFilter);
 	AddFilesFromDirectory(TEXT("/{project}"), FPaths::Combine(CookDirectory, FApp::GetProjectName()), true, &CookedFilter);
-	
-	AddFilesFromDirectory(TEXT("/{project}"), ProjectDir, false);
+
+	FFileFilter CookedMetadataFilter = FFileFilter()
+		.ExcludeDirectory(TEXT("ShaderLibrarySource"))
+		.ExcludeExtension(TEXT("manifest"));
+	AddFilesFromDirectory(TEXT("/{project}/Metadata"), FPaths::Combine(CookDirectory, FApp::GetProjectName(), "Metadata"), false, &CookedMetadataFilter);
+
+	FFileFilter ProjectSourceFilter = FFileFilter()
+		.IncludeExtension(TEXT("uproject"));
+	AddFilesFromDirectory(TEXT("/{project}"), ProjectDir, false, &ProjectSourceFilter);
 	
 	FFileFilter ConfigFilter = FFileFilter()
 		.IncludeExtension(TEXT("ini"));
