UnrealPak: Implement DDC2 async API for iostore compression.
BatchGet with max 128 inflight requests (or ~1 GiB in total) in batches of 8 items (or ~16 MiB each).
BatchPut with max 128 inflight requests (or ~256 MiB in total) in batches of 8 items (or ~1 MiB each).
Skip ddc for chunks smaller than CompressionMinBytesSaved (1KiB by default).
Skip ddc for .umap to avoid cache churn since maps are known to cook non-deterministically.
Skip ddc for shaders that use a different code path in UnrealPak as well as in runtime.
Use a new DDC2 cache key (that includes the CompressionBufferSize) and cache bucket.
Use TArray64/FMemoryWriter64 for serializing the data to support chunks bigger than 2 GiB.
Postpone allocation of compression buffers until the ddc get request completes and the size is known.
Reduce memory buffer limits to 2 GiB and 3 GiB again (earlier temporary bumps to 3 GiB and 4 GiB are not needed after recent task/retraction changes).
Add logging of number of ddc hits and puts.

#tests identical binary output

diff --git a/Engine/Source/Developer/IoStoreUtilities/Private/IoStoreUtilities.cpp b/Engine/Source/Developer/IoStoreUtilities/Private/IoStoreUtilities.cpp
--- a/Engine/Source/Developer/IoStoreUtilities/Private/IoStoreUtilities.cpp
+++ b/Engine/Source/Developer/IoStoreUtilities/Private/IoStoreUtilities.cpp
@@ -7,6 +7,9 @@
 #include "Algo/TopologicalSort.h"
 #include "Async/AsyncWork.h"
 #include "CookMetadata.h"
+#include "DerivedDataCache.h"
+#include "DerivedDataCacheInterface.h"
+#include "GenericPlatform/GenericPlatformMisc.h"
 #include "HAL/FileManager.h"
 #include "HAL/PlatformFileManager.h"
 #include "HAL/PlatformMemory.h"
@@ -3328,7 +3331,7 @@ public:
 		InitiatorThread = Async(EAsyncExecution::Thread, [this]() { InitiatorThreadFunc(); });
 		RetirerThread = Async(EAsyncExecution::Thread, [this]() { RetirerThreadFunc(); });
 
-		MaxSourceBufferMemory = 4ull << 30;
+		MaxSourceBufferMemory = 3ull << 30;
 		FParse::Value(FCommandLine::Get(), TEXT("MaxSourceBufferMemory="), MaxSourceBufferMemory);
 
 		MaxConcurrentSourceReads = uint32(FMath::Clamp(FPlatformMisc::NumberOfCoresIncludingHyperthreads()/2, 4, 32));
@@ -5708,9 +5711,14 @@ int32 CreateTarget(const FIoStoreArguments& Arguments, const FIoStoreWriterSetti
 				{
 					ProgressStringBuilder.Appendf(TEXT(" [%llu compressed]"), Progress.CompressedChunksCount);
 				}
+				if (Progress.CompressionDDCHitCount || Progress.CompressionDDCPutCount)
+				{
+					ProgressStringBuilder.Appendf(TEXT(" [DDC: %llu hits, %llu puts]"),
+						Progress.CompressionDDCHitCount, Progress.CompressionDDCPutCount);
+				}
 				if (Progress.ScheduledCompressionTasksCount)
 				{
-					ProgressStringBuilder.Appendf(TEXT(" [%llu compression tasks scheduled]"), Progress.ScheduledCompressionTasksCount);
+					ProgressStringBuilder.Appendf(TEXT(" [%llu compression tasks]"), Progress.ScheduledCompressionTasksCount);
 				}
 			}
 			else
@@ -5726,6 +5734,24 @@ int32 CreateTarget(const FIoStoreArguments& Arguments, const FIoStoreWriterSetti
 		}
 	}
 
+	if (GeneralIoWriterSettings.bCompressionEnableDDC)
+	{
+		TRACE_CPUPROFILER_EVENT_SCOPE(WaitForDDC);
+		UE_LOG(LogIoStore, Display, TEXT(""));
+		UE_LOG(LogIoStore, Display, TEXT("Waiting for DDC..."));
+		GetDerivedDataCacheRef().WaitForQuiescence(true);
+
+		FIoStoreWriterContext::FProgress Progress = IoStoreWriterContext->GetProgress();
+		const uint64 TotalDDCAttempts = Progress.CompressionDDCHitCount + Progress.CompressionDDCMissCount;
+		const double DDCHitRate = double(Progress.CompressionDDCHitCount) / TotalDDCAttempts * 100.0;
+		const double DDCPutRate = double(Progress.CompressionDDCPutCount) / TotalDDCAttempts * 100.0;
+		UE_LOG(LogIoStore, Display, TEXT("Compression DDC hits: %llu/%llu (%.2lf%%)"),
+			Progress.CompressionDDCHitCount, TotalDDCAttempts, DDCHitRate);
+		UE_LOG(LogIoStore, Display, TEXT("Compression DDC puts: %llu/%llu (%.2lf%%) [%llu failed]"),
+			Progress.CompressionDDCPutCount, TotalDDCAttempts, DDCPutRate, Progress.CompressionDDCPutErrorCount);
+		UE_LOG(LogIoStore, Display, TEXT(""));
+	}
+
 	{
 		FIoStoreWriterContext::FProgress Progress = IoStoreWriterContext->GetProgress();
 		if (Progress.HashDbChunksCount)
@@ -5765,6 +5791,40 @@ int32 CreateTarget(const FIoStoreArguments& Arguments, const FIoStoreWriterSetti
 				}
 			}
 		}
+		if (Progress.CompressionDDCHitCount)
+		{
+			UE_LOG(LogIoStore, Display, TEXT("%s / %s chunks for %s bytes were loaded from DDC, by type:"),
+				*NumberString(Progress.CompressionDDCHitCount),
+				*NumberString(Progress.TotalChunksCount),
+				*NumberString(Progress.CompressionDDCGetBytes));
+
+			for (uint8 i = 0; i < (uint8)EIoChunkType::MAX; i++)
+			{
+				if (Progress.CompressionDDCHitsByType[i])
+				{
+					UE_LOG(LogIoStore, Display, TEXT("    %-26s %s / %s"), *LexToString((EIoChunkType)i),
+						*NumberString(Progress.CompressionDDCHitsByType[i]),
+						*NumberString(Progress.BeginCompressChunksByType[i]));
+				}
+			}
+		}
+		if (Progress.CompressionDDCPutCount)
+		{
+			UE_LOG(LogIoStore, Display, TEXT("%s / %s chunks for %s bytes were stored in DDC, by type:"),
+				*NumberString(Progress.CompressionDDCPutCount),
+				*NumberString(Progress.TotalChunksCount),
+				*NumberString(Progress.CompressionDDCPutBytes));
+
+			for (uint8 i = 0; i < (uint8)EIoChunkType::MAX; i++)
+			{
+				if (Progress.CompressionDDCPutsByType[i])
+				{
+					UE_LOG(LogIoStore, Display, TEXT("    %-26s %s / %s"), *LexToString((EIoChunkType)i),
+						*NumberString(Progress.CompressionDDCPutsByType[i]),
+						*NumberString(Progress.BeginCompressChunksByType[i]));
+				}
+			}
+		}
 		UE_LOG(LogIoStore, Display, TEXT("Source bytes read:"));
 		uint64 ZenTotalBytes = 0;
 		for (uint64 b : WriteRequestManager.ZenSourceBytes)
@@ -5810,14 +5870,6 @@ int32 CreateTarget(const FIoStoreArguments& Arguments, const FIoStoreWriterSetti
 		}
 	}
 
-	if (GeneralIoWriterSettings.bCompressionEnableDDC)
-	{
-		FIoStoreWriterContext::FProgress Progress = IoStoreWriterContext->GetProgress();
-		uint64 TotalDDCAttempts = Progress.CompressionDDCHitCount + Progress.CompressionDDCMissCount;
-		double DDCHitRate = double(Progress.CompressionDDCHitCount) / TotalDDCAttempts * 100.0;
-		UE_LOG(LogIoStore, Display, TEXT("Compression DDC hits: %llu/%llu (%.2f%%)"), Progress.CompressionDDCHitCount, TotalDDCAttempts, DDCHitRate);
-	}
-
 	if (Arguments.WriteBackMetadataToAssetRegistry != EAssetRegistryWritebackMethod::Disabled)
 	{
 		DoAssetRegistryWritebackDuringStage(
diff --git a/Engine/Source/Developer/IoStoreUtilities/Private/IoStoreWriter.cpp b/Engine/Source/Developer/IoStoreUtilities/Private/IoStoreWriter.cpp
--- a/Engine/Source/Developer/IoStoreUtilities/Private/IoStoreWriter.cpp
+++ b/Engine/Source/Developer/IoStoreUtilities/Private/IoStoreWriter.cpp
@@ -7,11 +7,17 @@
 #include "Async/AsyncFileHandle.h"
 #include "Async/Future.h"
 #include "Async/ParallelFor.h"
+#include "Compression/CompressedBuffer.h"
+#include "Compression/OodleDataCompression.h"
 #include "Containers/Map.h"
+#include "DerivedDataCache.h"
 #include "DerivedDataCacheInterface.h"
+#include "DerivedDataCacheRecord.h"
+#include "DerivedDataRequestOwner.h"
 #include "HAL/FileManager.h"
 #include "HAL/PlatformFileManager.h"
 #include "IO/IoDirectoryIndex.h"
+#include "Memory/SharedBuffer.h"
 #include "Misc/Compression.h"
 #include "Misc/Paths.h"
 #include "Misc/StringBuilder.h"
@@ -29,6 +35,15 @@ TRACE_DECLARE_ATOMIC_INT_COUNTER(IoStoreRefDbDone, TEXT("IoStoreWriter/RefDbDone
 TRACE_DECLARE_INT_COUNTER(IoStoreBeginCompressionCount, TEXT("IoStoreWriter/BeginCompression"));
 TRACE_DECLARE_INT_COUNTER(IoStoreBeginEncryptionAndSigningCount, TEXT("IoStoreWriter/BeginEncryptionAndSigning"));
 TRACE_DECLARE_INT_COUNTER(IoStoreBeginWriteCount, TEXT("IoStoreWriter/BeginWrite"));
+TRACE_DECLARE_ATOMIC_INT_COUNTER(IoStoreDDCGetInflightCount, TEXT("IoStoreWriter/DDCGetInflightCount"));
+TRACE_DECLARE_ATOMIC_INT_COUNTER(IoStoreDDCPutInflightCount, TEXT("IoStoreWriter/DDCPutInflightCount"));
+TRACE_DECLARE_ATOMIC_INT_COUNTER(IoStoreDDCHitCount, TEXT("IoStoreWriter/DDCHitCount"));
+TRACE_DECLARE_ATOMIC_INT_COUNTER(IoStoreDDCMissCount, TEXT("IoStoreWriter/DDCMissCount"));
+TRACE_DECLARE_ATOMIC_INT_COUNTER(IoStoreDDCPutCount, TEXT("IoStoreWriter/DDCPutCount"));
+
+static UE::DerivedData::FCacheBucket IoStoreDDCBucket = UE::DerivedData::FCacheBucket(ANSITEXTVIEW("IoStoreCompression"));
+static UE::DerivedData::ECachePolicy IoStoreDDCPolicy = UE::DerivedData::ECachePolicy::Default;
+static FStringView IoStoreDDCVersion = TEXTVIEW("985D3FAD-71D0-4758-A777-A910B49CC4BF");
 
 struct FChunkBlock
 {
@@ -63,6 +78,7 @@ struct FIoStoreWriteQueueEntry
 	// We make this optional because at the latest it might not be valid until FinishCompressionBarrior
 	// completes and we'd like to have a check() on that.
 	TOptional<uint64> UncompressedSize;
+	uint64 CompressedSize = 0; 
 
 	// this is not filled out until after encryption completes and *includes the alignment padding for encryption*!
 	uint64 DiskSize = 0; 
@@ -79,9 +95,11 @@ struct FIoStoreWriteQueueEntry
 	TAtomic<int32> CompressedBlocksCount{ 0 };
 	int32 PartitionIndex = -1;
 	int32 NumChunkBlocks = 0;
-	FString DDCCacheKey;
+	UE::DerivedData::FCacheKey DDCKey;
 	bool bAdded = false;
 	bool bModified = false;
+	bool bUseDDCForCompression = false;
+	bool bFoundInDDC = false;
 	bool bStoreCompressedDataInDDC = false;
 	
 	bool bCouldBeFromReferenceDb = false; // Whether the chunk is a valid candidate for the reference db.
@@ -167,9 +185,247 @@ private:
 	TAtomic<bool> bIsDoneAdding { false };
 };
 
+struct FIoStoreDDCRequestDispatcherParams
+{
+	// Maximum time for filling up a batch, after this time limit is reached,
+	// any queued requests are dispatched even if the batch is not full
+	double QueueTimeLimitMs = 20.f;
+	// Maximum number of (estimated) bytes in a batch,
+	// when this is reached a batch with the currently queued requests will be dispatched immediately
+	uint64 MaxBatchBytes = 16ull << 20;
+	// Maximum number of (estimated) bytes for all inflight requests, 
+	// if this limit is reached, then wait for requests to complete before dispatching a new batch
+	uint64 MaxInflightBytes = 1ull << 30;
+	// The number of queued requests to collect before dispatching a batch
+	int32 MaxBatchItems = 8;
+	// Maximum number of inflight requests,
+	// if this limit is reached, then wait for requests to complete before dispatching a new batch
+	int32 MaxInflightCount = 128;
+	// Do a blocking wait after dispatching each batch (for debugging)
+	bool bBlockingWait = false;
+};
+
+template<class T>
+struct FIoStoreDDCRequestDispatcherQueue
+{
+	FIoStoreDDCRequestDispatcherQueue(const FIoStoreDDCRequestDispatcherParams& InParams)
+		: Params(InParams)
+		, RequestOwner(UE::DerivedData::EPriority::Highest)
+	{ }
+
+	FIoStoreDDCRequestDispatcherParams Params;
+	UE::DerivedData::FRequestOwner RequestOwner;
+	TArray<T> Requests;
+	FEventRef RequestCompletedEvent;
+	TAtomic<uint64> InflightCount{ 0 };
+	TAtomic<uint64> InflightBytes{ 0 };
+	uint64 QueuedBytes = 0;
+	uint64 LastRequestCycle = 0;
+
+	T& EnqueueRequest(uint64 Size)
+	{
+		if (Requests.Num() == 0)
+		{
+			LastRequestCycle = FPlatformTime::Cycles64();
+		}
+		QueuedBytes += Size;
+		InflightBytes.AddExchange(Size);
+		return Requests.AddDefaulted_GetRef();
+	}
+
+	bool ReadyOrWaitForDispatch(bool bForceDispatch);
+
+	void OnDispatch()
+	{
+		QueuedBytes = 0;
+		LastRequestCycle = FPlatformTime::Cycles64();
+		InflightCount.AddExchange(Requests.Num());
+		Requests.Reset();
+		if (Params.bBlockingWait)
+		{
+			TRACE_CPUPROFILER_EVENT_SCOPE(WaitForDDC);
+			RequestOwner.Wait();
+		}
+	}
+
+	void OnRequestComplete(uint64 Size)
+	{
+		InflightCount.DecrementExchange();
+		InflightBytes.SubExchange(Size);
+		RequestCompletedEvent->Trigger();
+	}
+};
+
+template<class T>
+bool FIoStoreDDCRequestDispatcherQueue<T>::ReadyOrWaitForDispatch(bool bForceDispatch)
+{
+	int32 NumRequests = Requests.Num();
+	if (NumRequests == 0)
+	{
+		return false;
+	}
+
+	bForceDispatch |= (NumRequests >= Params.MaxBatchItems) || (QueuedBytes >= Params.MaxBatchBytes);
+
+	const bool bLazyDispatch = !bForceDispatch &&
+		FPlatformTime::ToMilliseconds64(FPlatformTime::Cycles64() - LastRequestCycle) >= Params.QueueTimeLimitMs;
+
+	if (!bForceDispatch && !bLazyDispatch)
+	{
+		return false;
+	}
+
+	int32 LocalInflightCount = InflightCount.Load();
+	if (bForceDispatch)
+	{
+		while (LocalInflightCount > 0 && LocalInflightCount + NumRequests > Params.MaxInflightCount)
+		{
+			TRACE_CPUPROFILER_EVENT_SCOPE(WaitForDDCBatch);
+			RequestCompletedEvent->Wait();
+			LocalInflightCount = InflightCount.Load();
+		}
+		while (LocalInflightCount > 0 && InflightBytes.Load() + QueuedBytes > Params.MaxInflightBytes)
+		{
+			TRACE_CPUPROFILER_EVENT_SCOPE(WaitForDDCMemory);
+			RequestCompletedEvent->Wait();
+			LocalInflightCount = InflightCount.Load();
+		}
+	}
+	else if (LocalInflightCount + NumRequests > Params.MaxInflightCount)
+	{
+		return false;
+	}
+	else if (InflightBytes.Load() + QueuedBytes > Params.MaxInflightBytes)
+	{
+		return false;
+	}
+
+	return true;
+}
+
+class FIoStoreDDCGetRequestDispatcher
+{
+public:
+	FIoStoreDDCGetRequestDispatcher(const FIoStoreDDCRequestDispatcherParams& InParams) : RequestQueue(InParams) {};
+	void EnqueueGetRequest(FIoStoreWriteQueueEntry* Entry);
+	void DispatchGetRequests(
+		TFunction<void (FIoStoreWriteQueueEntry* Entry, FSharedBuffer Result)> Callback,
+		bool bForceDispatch = false);
+	void FlushGetRequests(TFunction<void (FIoStoreWriteQueueEntry* Entry, FSharedBuffer Result)> Callback)
+	{
+		DispatchGetRequests(Callback, true);
+		RequestQueue.RequestOwner.Wait();
+	};
+
+private:
+	FIoStoreDDCRequestDispatcherQueue<UE::DerivedData::FCacheGetValueRequest> RequestQueue;
+};
+
+class FIoStoreDDCPutRequestDispatcher
+{
+public:
+	FIoStoreDDCPutRequestDispatcher(const FIoStoreDDCRequestDispatcherParams& InParams) : RequestQueue(InParams) {};
+	void EnqueuePutRequest(FIoStoreWriteQueueEntry* Entry, FSharedBuffer SharedBuffer);
+	void DispatchPutRequests(
+		TFunction<void (FIoStoreWriteQueueEntry* Entry, bool bSuccess)> Callback,
+		bool bForceDispatch = false);
+	void FlushPutRequests(TFunction<void (FIoStoreWriteQueueEntry* Entry, bool bSuccess)> Callback)
+	{
+		DispatchPutRequests(Callback, true);
+		RequestQueue.RequestOwner.Wait();
+	};
+
+private:
+	FIoStoreDDCRequestDispatcherQueue<UE::DerivedData::FCachePutValueRequest> RequestQueue;
+};
+
+void FIoStoreDDCGetRequestDispatcher::EnqueueGetRequest(FIoStoreWriteQueueEntry* Entry)
+{
+	UE::DerivedData::FCacheGetValueRequest& Request = RequestQueue.EnqueueRequest(Entry->Request->GetSourceBufferSizeEstimate());
+	Request.Name = *Entry->Options.FileName;
+	Request.Key = Entry->DDCKey;
+	Request.Policy = IoStoreDDCPolicy;
+	Request.UserData = reinterpret_cast<uint64>(Entry);
+}
+
+void FIoStoreDDCGetRequestDispatcher::DispatchGetRequests(
+	TFunction<void (FIoStoreWriteQueueEntry* Entry, FSharedBuffer Result)> Callback,
+	bool bForceDispatch)
+{
+	using namespace UE::DerivedData;
+
+	if (!RequestQueue.ReadyOrWaitForDispatch(bForceDispatch))
+	{
+		return;
+	}
+
+	TRACE_CPUPROFILER_EVENT_SCOPE(DispatchDDCGetRequests);
+	TRACE_COUNTER_ADD(IoStoreDDCGetInflightCount, RequestQueue.Requests.Num());
+
+	{
+		FRequestBarrier RequestBarrier(RequestQueue.RequestOwner);
+		GetCache().GetValue(RequestQueue.Requests, RequestQueue.RequestOwner, [this, Callback](FCacheGetValueResponse&& Response)
+		{
+			TRACE_CPUPROFILER_EVENT_SCOPE(ReadFromDDC_Decompress);
+			FIoStoreWriteQueueEntry* Entry = reinterpret_cast<FIoStoreWriteQueueEntry*>(Response.UserData);
+			FSharedBuffer Result;
+			if (Response.Status == EStatus::Ok)
+			{
+				Result = Response.Value.GetData().Decompress();
+			}
+			Callback(Entry, Result);
+			RequestQueue.OnRequestComplete(Entry->Request->GetSourceBufferSizeEstimate());
+			TRACE_COUNTER_DECREMENT(IoStoreDDCGetInflightCount);
+		});
+	}
+	RequestQueue.OnDispatch();
+}
+
+void FIoStoreDDCPutRequestDispatcher::EnqueuePutRequest(FIoStoreWriteQueueEntry* Entry, FSharedBuffer SharedBuffer)
+{
+	FCompressedBuffer CompressedBuffer = FCompressedBuffer::Compress(
+		SharedBuffer,
+		ECompressedBufferCompressor::NotSet,
+		ECompressedBufferCompressionLevel::None);
+
+	UE::DerivedData::FCachePutValueRequest& Request = RequestQueue.EnqueueRequest(Entry->CompressedSize);
+	Request.Name = *Entry->Options.FileName;
+	Request.Key = Entry->DDCKey;
+	Request.Policy = IoStoreDDCPolicy;
+	Request.Value = UE::DerivedData::FValue(MoveTemp(CompressedBuffer));
+	Request.UserData = reinterpret_cast<uint64>(Entry);
+}
+
+void FIoStoreDDCPutRequestDispatcher::DispatchPutRequests(
+	TFunction<void (FIoStoreWriteQueueEntry* Entry, bool bSuccess)> Callback,
+	bool bForceDispatch)
+{
+	using namespace UE::DerivedData;
+	
+	if (!RequestQueue.ReadyOrWaitForDispatch(bForceDispatch))
+	{
+		return;
+	}
+
+	TRACE_CPUPROFILER_EVENT_SCOPE(DispatchDDCPutRequests);
+	TRACE_COUNTER_ADD(IoStoreDDCPutInflightCount, RequestQueue.Requests.Num());
+
+	{
+		FRequestBarrier RequestBarrier(RequestQueue.RequestOwner);
+		GetCache().PutValue(RequestQueue.Requests, RequestQueue.RequestOwner, [this, Callback](FCachePutValueResponse&& Response)
+		{
+			FIoStoreWriteQueueEntry* Entry = reinterpret_cast<FIoStoreWriteQueueEntry*>(Response.UserData);
+			bool bSuccess = Response.Status == EStatus::Ok;
+			Callback(Entry, bSuccess);
+			RequestQueue.OnRequestComplete(Entry->CompressedSize);
+			TRACE_COUNTER_DECREMENT(IoStoreDDCPutInflightCount);
+		});
+	}
+	RequestQueue.OnDispatch();
+}
+
 class FIoStoreWriterContextImpl
 {
-	static constexpr uint64 CompressionMemorySoftLimit = 3ull << 30ull;
 public:
 	FIoStoreWriterContextImpl()
 	{
@@ -196,7 +452,10 @@ public:
 
 		if (WriterSettings.bCompressionEnableDDC)
 		{
-			DDC = GetDerivedDataCache();
+			TRACE_CPUPROFILER_EVENT_SCOPE(InitializeDDC);
+			UE_LOG(LogIoStore, Display, TEXT("InitializeDDC"));
+			GetDerivedDataCacheRef();
+			UE::DerivedData::GetCache();
 		}
 
 		if (WriterSettings.CompressionMethod != NAME_None)
@@ -206,7 +465,10 @@ public:
 		CompressionBufferSize = FMath::Max(CompressionBufferSize, static_cast<int32>(WriterSettings.CompressionBlockSize));
 		CompressionBufferSize = Align(CompressionBufferSize, FAES::AESBlockSize);
 
-		const int32 InitialCompressionBufferCount = int32(CompressionMemorySoftLimit / CompressionBufferSize);
+		MaxCompressionBufferMemory = 2ull << 30;
+		FParse::Value(FCommandLine::Get(), TEXT("MaxCompressionBufferMemory="), MaxCompressionBufferMemory);
+
+		const int32 InitialCompressionBufferCount = int32(MaxCompressionBufferMemory / CompressionBufferSize);
 		AvailableCompressionBuffers.Reserve(InitialCompressionBufferCount);
 		for (int32 BufferIndex = 0; BufferIndex < InitialCompressionBufferCount; ++BufferIndex)
 		{
@@ -227,28 +489,25 @@ public:
 		for (uint8 i=0; i<(uint8)EIoChunkType::MAX; i++)
 		{
 			Progress.HashDbChunksByType[i] = HashDbChunksByType[i].Load();
+			Progress.CompressedChunksByType[i] = CompressedChunksByType[i].Load();
+			Progress.BeginCompressChunksByType[i] = BeginCompressChunksByType[i].Load();
+			Progress.RefDbChunksByType[i] = RefDbChunksByType[i].Load();
+			Progress.CompressionDDCHitsByType[i] = CompressionDDCHitsByType[i].Load();
+			Progress.CompressionDDCPutsByType[i] = CompressionDDCPutsByType[i].Load();
+			Progress.CompressionDDCHitCount += Progress.CompressionDDCHitsByType[i];
+			Progress.CompressionDDCPutCount += Progress.CompressionDDCPutsByType[i];
 		}
 
 		Progress.TotalChunksCount = TotalChunksCount.Load();
 		Progress.HashedChunksCount = HashedChunksCount.Load();
 		Progress.CompressedChunksCount = CompressedChunksCount.Load();
-		for (uint8 i = 0; i < (uint8)EIoChunkType::MAX; i++)
-		{
-			Progress.CompressedChunksByType[i] = CompressedChunksByType[i].Load();
-		}
-		for (uint8 i = 0; i < (uint8)EIoChunkType::MAX; i++)
-		{
-			Progress.BeginCompressChunksByType[i] = BeginCompressChunksByType[i].Load();
-		}
 		Progress.SerializedChunksCount = SerializedChunksCount.Load();
 		Progress.ScheduledCompressionTasksCount = ScheduledCompressionTasksCount.Load();
-		Progress.CompressionDDCHitCount = CompressionDDCHitCount.Load();
+		Progress.CompressionDDCGetBytes = CompressionDDCGetBytes.Load();
+		Progress.CompressionDDCPutBytes = CompressionDDCPutBytes.Load();
 		Progress.CompressionDDCMissCount = CompressionDDCMissCount.Load();
+		Progress.CompressionDDCPutErrorCount = CompressionDDCPutErrorCount.Load();
 		Progress.RefDbChunksCount = RefDbChunksCount.Load();
-		for (uint8 i = 0; i < (uint8)EIoChunkType::MAX; i++)
-		{
-			Progress.RefDbChunksByType[i] = RefDbChunksByType[i].Load();
-		}
 
 		return Progress;
 	}
@@ -284,13 +543,14 @@ public:
 	}
 
 private:
-	void ScheduleEntry(FIoStoreWriteQueueEntry* Entry);
+	UE::DerivedData::FCacheKey MakeDDCKey(FIoStoreWriteQueueEntry* Entry) const;
+
+	void ScheduleAllEntries(TArrayView<FIoStoreWriteQueueEntry*> AllEntries);
 	void BeginCompressionThreadFunc();
 	void BeginEncryptionAndSigningThreadFunc();
 	void WriterThreadFunc();
 
 	FIoStoreWriterSettings WriterSettings;
-	FDerivedDataCacheInterface* DDC = nullptr;
 	FEventRef CompressionMemoryReleasedEvent;
 	TFuture<void> BeginCompressionThread;
 	TFuture<void> BeginEncryptionAndSigningThread;
@@ -307,15 +567,20 @@ private:
 	TAtomic<uint64> CompressedChunksCount{ 0 };
 	TAtomic<uint64> CompressedChunksByType[(int8)EIoChunkType::MAX] = { 0 };
 	TAtomic<uint64> BeginCompressChunksByType[(int8)EIoChunkType::MAX] = { 0 };
+	TAtomic<uint64> CompressionDDCHitsByType[(int8)EIoChunkType::MAX]{ 0 };
+	TAtomic<uint64> CompressionDDCPutsByType[(int8)EIoChunkType::MAX]{ 0 };
 	TAtomic<uint64> SerializedChunksCount{ 0 };
 	TAtomic<uint64> WriteCycleCount{ 0 };
 	TAtomic<uint64> WriteByteCount{ 0 };
 	TAtomic<uint64> ScheduledCompressionTasksCount{ 0 };
-	TAtomic<uint64> CompressionDDCHitCount{ 0 };
+	TAtomic<uint64> CompressionDDCGetBytes{ 0 };
+	TAtomic<uint64> CompressionDDCPutBytes{ 0 };
 	TAtomic<uint64> CompressionDDCMissCount{ 0 };
+	TAtomic<uint64> CompressionDDCPutErrorCount{ 0 };
 	TAtomic<uint64> ScheduledCompressionMemory{ 0 };
 	FCriticalSection AvailableCompressionBuffersCritical;
 	TArray<FIoBuffer*> AvailableCompressionBuffers;
+	uint64 MaxCompressionBufferMemory = 0;
 	int32 CompressionBufferSize = -1;
 	TArray<TSharedPtr<FIoStoreWriter>> IoStoreWriters;
 
@@ -580,6 +845,8 @@ public:
 		check(!bHasFlushed);
 		checkf(ChunkId.IsValid(), TEXT("ChunkId is not valid!"));
 
+		const FIoStoreWriterSettings& WriterSettings = WriterContext->GetSettings();
+
 		WriterContext->TotalChunksCount.IncrementExchange();
 		FIoStoreWriteQueueEntry* Entry = new FIoStoreWriteQueueEntry();
 		Entries.Add(Entry);
@@ -589,6 +856,12 @@ public:
 		Entry->Options = WriteOptions;
 		Entry->CompressionMethod = CompressionMethodForEntry(WriteOptions);
 		Entry->CompressionMemoryEstimate = CalculateCompressionBufferMemory(Request->GetSourceBufferSizeEstimate());
+		Entry->bUseDDCForCompression =
+			WriterSettings.bCompressionEnableDDC &&
+			Entry->CompressionMethod != NAME_None &&
+			Request->GetSourceBufferSizeEstimate() > WriterSettings.CompressionMinBytesSaved &&
+			Request->GetSourceBufferSizeEstimate() > WriterSettings.CompressionMinSizeToConsiderDDC &&
+			!Entry->Options.FileName.EndsWith(TEXT(".umap")); // avoid cache churn while maps are known to cook non-deterministically
 		Entry->Request = Request;		
 
 		// If we can get the hash without reading the whole thing and hashing it, do so to avoid the IO.
@@ -596,7 +869,7 @@ public:
 		{
 			check(!ChunkHash->IsZero());
 			Entry->ChunkHash = *ChunkHash;
-			if (WriterContext->WriterSettings.bValidateChunkHashes == false)
+			if (WriterSettings.bValidateChunkHashes == false)
 			{
 				// If we aren't validating then we just use it and bail.
 				WriterContext->HashDbChunksCount.IncrementExchange();
@@ -608,6 +881,7 @@ public:
 					Entry->bLoadingFromReferenceDb = ReferenceChunkDatabase->ChunkExists(ContainerSettings.ContainerId, Entry->ChunkHash, Entry->ChunkId, Entry->NumChunkBlocks);
 					Entry->bCouldBeFromReferenceDb = true;
 				}
+				Entry->bUseDDCForCompression &= !Entry->bLoadingFromReferenceDb;
 				return;
 			}
 			// If we are validating run the normal path to verify it.
@@ -636,6 +910,7 @@ public:
 				Entry->bLoadingFromReferenceDb = ReferenceChunkDatabase->ChunkExists(ContainerSettings.ContainerId, Entry->ChunkHash, Entry->ChunkId, Entry->NumChunkBlocks);
 				Entry->bCouldBeFromReferenceDb = true;
 			}
+			Entry->bUseDDCForCompression &= !Entry->bLoadingFromReferenceDb;
 
 			// Release the source data buffer, it will be reloaded later when we start compressing the chunk
 			Entry->Request->FreeSourceBuffer();
@@ -1256,20 +1531,26 @@ private:
 		}
 	}
 
-	bool SerializeCompressedDDCData(FIoStoreWriteQueueEntry* Entry, FArchive& Ar)
+	bool SerializeCompressedDDCData(FIoStoreWriteQueueEntry* Entry, FArchive& Ar, uint64* OutCompressedSize = nullptr)
 	{
-		uint32 NumBlocks = Entry->ChunkBlocks.Num();
-		Ar << NumBlocks;
-		if (NumBlocks != Entry->ChunkBlocks.Num())
+		uint64 UncompressedSize = Entry->UncompressedSize.Get(0);
+		uint32 NumChunkBlocks = Entry->ChunkBlocks.Num();
+		Ar << UncompressedSize;
+		Ar << NumChunkBlocks;
+		if (Ar.IsLoading())
 		{
-			return false;
+			Entry->NumChunkBlocks = NumChunkBlocks;
+			Entry->UncompressedSize.Emplace(UncompressedSize);
+			AllocateCompressionBuffers(Entry);
 		}
+		bool bError = false;
 		for (FChunkBlock& Block : Entry->ChunkBlocks)
 		{
 			Ar << Block.CompressedSize;
 			if (Block.CompressedSize > Block.UncompressedSize)
 			{
-				return false;
+				bError = true;
+				break;
 			}
 			if (Ar.IsLoading() && Block.CompressedSize == Block.UncompressedSize)
 			{
@@ -1277,11 +1558,21 @@ private:
 			}
 			if (Block.IoBuffer->DataSize() < Block.CompressedSize)
 			{
-				return false;
+				bError = true;
+				break;
 			}
 			Ar.Serialize(Block.IoBuffer->Data(), Block.CompressedSize);
+			if (OutCompressedSize)
+			{
+				*OutCompressedSize += Block.CompressedSize;
+			}
+		}
+		bError |= Ar.IsError();
+		if (Ar.IsLoading() && bError)
+		{
+			FreeCompressionBuffers(Entry);
 		}
-		return !Ar.IsError();
+		return !bError;
 	}
 
 	FName CompressionMethodForEntry(const FIoWriteOptions& Options) const
@@ -1308,6 +1599,40 @@ private:
 		return WriterContext->CompressionBufferSize * NumBlocks;
 	}
 
+	void AllocateCompressionBuffers(FIoStoreWriteQueueEntry* Entry, const uint8* UncompressedData = nullptr)
+	{
+		check(Entry->ChunkBlocks.Num() == 0);
+		const FIoStoreWriterSettings& WriterSettings = WriterContext->WriterSettings;
+		check(WriterSettings.CompressionBlockSize > 0);
+		
+		Entry->ChunkBlocks.SetNum(Entry->NumChunkBlocks);
+		{
+			uint64 BytesToProcess = Entry->UncompressedSize.GetValue();
+			for (int32 BlockIndex = 0; BlockIndex < Entry->NumChunkBlocks; ++BlockIndex)
+			{
+				FChunkBlock& Block = Entry->ChunkBlocks[BlockIndex];
+				Block.IoBuffer = WriterContext->AllocCompressionBuffer();
+				Block.CompressionMethod = Entry->CompressionMethod;
+				Block.UncompressedSize = FMath::Min(BytesToProcess, WriterSettings.CompressionBlockSize);
+				BytesToProcess -= Block.UncompressedSize;
+				if (UncompressedData)
+				{
+					Block.UncompressedData = UncompressedData;
+					UncompressedData += Block.UncompressedSize;
+				}
+			}
+		}
+	}
+
+	void FreeCompressionBuffers(FIoStoreWriteQueueEntry* Entry)
+	{
+		for (FChunkBlock& ChunkBlock : Entry->ChunkBlocks)
+		{
+			WriterContext->FreeCompressionBuffer(ChunkBlock.IoBuffer);
+		}
+		Entry->ChunkBlocks.Empty();
+	}
+
 	void LoadFromReferenceDb(FIoStoreWriteQueueEntry* Entry)
 	{
 		if (Entry->NumChunkBlocks == 0)
@@ -1393,14 +1718,13 @@ private:
 	{
 		WriterContext->BeginCompressChunksByType[(int8)Entry->ChunkId.GetChunkType()].IncrementExchange();
 		
-		if (Entry->bLoadingFromReferenceDb)
+		if (Entry->bLoadingFromReferenceDb || Entry->bFoundInDDC)
 		{
+			check(Entry->UncompressedSize.IsSet());
 			Entry->FinishCompressionBarrier.Trigger();
 			return;
 		}
 
-		const FIoStoreWriterSettings& WriterSettings = WriterContext->WriterSettings;
-
 		const FIoBuffer* SourceBuffer = Entry->Request->GetSourceBuffer();
 		Entry->UncompressedSize.Emplace(SourceBuffer->DataSize());
 		Entry->NumChunkBlocks = CalculateNumChunkBlocks(Entry->UncompressedSize.GetValue());
@@ -1411,22 +1735,7 @@ private:
 			return;
 		}
 
-		Entry->ChunkBlocks.SetNum(Entry->NumChunkBlocks);
-		{
-			// We must allocate resources for our tasks up front to prevent resource deadlock.
-			uint64 BytesToProcess = Entry->UncompressedSize.GetValue();
-			const uint8* UncompressedData = SourceBuffer->Data();
-			for (int32 BlockIndex = 0; BlockIndex < Entry->NumChunkBlocks; ++BlockIndex)
-			{
-				FChunkBlock& Block = Entry->ChunkBlocks[BlockIndex];
-				Block.IoBuffer = WriterContext->AllocCompressionBuffer();
-				Block.CompressionMethod = Entry->CompressionMethod;
-				Block.UncompressedSize = FMath::Min(BytesToProcess, WriterSettings.CompressionBlockSize);
-				Block.UncompressedData = UncompressedData;
-				BytesToProcess -= Block.UncompressedSize;
-				UncompressedData += Block.UncompressedSize;
-			}
-		}
+		AllocateCompressionBuffers(Entry, SourceBuffer->Data());
 
 		if (Entry->CompressionMethod == NAME_None)
 		{
@@ -1440,45 +1749,6 @@ private:
 			return;
 		}
 
-		bool bUseDDCCompression =
-			WriterContext->DDC &&
-			WriterSettings.bCompressionEnableDDC &&
-			Entry->UncompressedSize.GetValue() > WriterSettings.CompressionMinSizeToConsiderDDC;
-		if (bUseDDCCompression)
-		{
-			TStringBuilder<256> CacheKeySuffix;
-			CacheKeySuffix.Append(LexToString(Entry->ChunkHash));
-			WriterSettings.CompressionMethod.AppendString(CacheKeySuffix);
-			CacheKeySuffix.Append(FCompression::GetCompressorDDCSuffix(WriterSettings.CompressionMethod));
-			CacheKeySuffix.Appendf(TEXT("%d_%d"),
-				WriterSettings.CompressionMinBytesSaved,
-				WriterSettings.CompressionMinPercentSaved);
-			Entry->DDCCacheKey = FDerivedDataCacheInterface::BuildCacheKey(
-				TEXT("IoStoreCompression"),
-				TEXT("985D3FAD-71D0-4758-A777-A910B49CC4BF"),
-				CacheKeySuffix.ToString());
-
-			TArray<uint8> DDCData;
-			TRACE_CPUPROFILER_EVENT_SCOPE(ReadFromDDC);
-			bool bFoundInDDC = false;
-			if (WriterContext->DDC->GetSynchronous(*Entry->DDCCacheKey, DDCData, Entry->Options.FileName))
-			{
-				FLargeMemoryReader DDCDataReader(DDCData.GetData(), DDCData.Num());
-				bFoundInDDC = SerializeCompressedDDCData(Entry, DDCDataReader);
-			}
-			if (bFoundInDDC)
-			{
-				WriterContext->CompressionDDCHitCount.IncrementExchange();
-				Entry->FinishCompressionBarrier.Trigger();
-				return;
-			}
-			else
-			{
-				WriterContext->CompressionDDCMissCount.IncrementExchange();
-				Entry->bStoreCompressedDataInDDC = true;
-			}
-		}
-
 		ScheduleCompressionTasks(Entry);
 	}
 
@@ -1515,6 +1785,14 @@ private:
 
 	void BeginEncryptAndSign(FIoStoreWriteQueueEntry* Entry)
 	{
+		Entry->Request->FreeSourceBuffer();
+
+		Entry->CompressedSize = 0;
+		for (const FChunkBlock& ChunkBlock : Entry->ChunkBlocks)
+		{
+			Entry->CompressedSize += ChunkBlock.CompressedSize;
+		}
+
 		if (ContainerSettings.IsEncrypted() || ContainerSettings.IsSigned())
 		{
 			UE::Tasks::FTask EncryptAndSignTask = UE::Tasks::Launch(TEXT("EncryptAndSign"), [this, Entry]()
@@ -1573,11 +1851,7 @@ private:
 		ON_SCOPE_EXIT
 		{
 			TRACE_CPUPROFILER_EVENT_SCOPE(FreeBlocks);
-			for (FChunkBlock& ChunkBlock : Entry->ChunkBlocks)
-			{
-				WriterContext->FreeCompressionBuffer(ChunkBlock.IoBuffer);
-			}
-			Entry->ChunkBlocks.Empty();
+			FreeCompressionBuffers(Entry);
 			delete Entry->Request;
 			Entry->Request = nullptr;
 
@@ -1827,7 +2101,7 @@ void FIoStoreWriterContextImpl::Flush()
 		AllEntries.Append(IoStoreWriter->Entries);
 	}
 
-	// Start scheduler threads, queue all entries, and wait for them to finish
+	// Start scheduler threads, enqueue all entries, and wait for them to finish
 	{
 		double WritesStart = FPlatformTime::Seconds();
 
@@ -1835,12 +2109,7 @@ void FIoStoreWriterContextImpl::Flush()
 		BeginEncryptionAndSigningThread = Async(EAsyncExecution::Thread, [this]() { BeginEncryptionAndSigningThreadFunc(); });
 		WriterThread = Async(EAsyncExecution::Thread, [this]() { WriterThreadFunc(); });
 
-		for (FIoStoreWriteQueueEntry* Entry : AllEntries)
-		{
-			BeginCompressionQueue.Enqueue(Entry);
-			ScheduleEntry(Entry);
-		}
-		BeginCompressionQueue.CompleteAdding();
+		ScheduleAllEntries(AllEntries);
 
 		{
 			TRACE_CPUPROFILER_EVENT_SCOPE(WaitForWritesToComplete);
@@ -1855,7 +2124,6 @@ void FIoStoreWriterContextImpl::Flush()
 			*FText::AsNumber(WriteByteCount.Load()).ToString(),
 			*FText::AsNumber((int64)((double)WriteByteCount.Load() / FMath::Max(.0001f, WritesSeconds))).ToString()
 			);
-		AllEntries.Empty();
 	}
 
 	// Classically there were so few writers that this didn't need to be multi threaded, but it
@@ -1883,29 +2151,102 @@ void FIoStoreWriterContextImpl::Flush()
 		);
 }
 
-void FIoStoreWriterContextImpl::ScheduleEntry(FIoStoreWriteQueueEntry* Entry)
+
+UE::DerivedData::FCacheKey FIoStoreWriterContextImpl::MakeDDCKey(FIoStoreWriteQueueEntry* Entry) const
 {
-	uint64 LocalScheduledCompressionMemory = ScheduledCompressionMemory.Load();
-	
-	while (LocalScheduledCompressionMemory > 0 &&
-		LocalScheduledCompressionMemory + Entry->CompressionMemoryEstimate > CompressionMemorySoftLimit)
-	{
-		TRACE_CPUPROFILER_EVENT_SCOPE(WaitForCompressionMemory);
-		CompressionMemoryReleasedEvent->Wait();
-		LocalScheduledCompressionMemory = ScheduledCompressionMemory.Load();
-	}
+	TStringBuilder<256> CacheKeySuffix;
+	CacheKeySuffix << IoStoreDDCVersion;
+	CacheKeySuffix << Entry->ChunkHash;
+	CacheKeySuffix.Append(FCompression::GetCompressorDDCSuffix(Entry->CompressionMethod));
+	CacheKeySuffix.Appendf(TEXT("%llu_%d_%d_%d"),
+		WriterSettings.CompressionBlockSize,
+		CompressionBufferSize,
+		WriterSettings.CompressionMinBytesSaved,
+		WriterSettings.CompressionMinPercentSaved);
+
+	return { IoStoreDDCBucket, FIoHash::HashBuffer(MakeMemoryView(FTCHARToUTF8(CacheKeySuffix.ToString()))) };
+}
 
-	ScheduledCompressionMemory.AddExchange(Entry->CompressionMemoryEstimate);
-	TRACE_COUNTER_SET(IoStoreCompressionMemoryScheduled, ScheduledCompressionMemory.Load());
+void FIoStoreWriterContextImpl::ScheduleAllEntries(TArrayView<FIoStoreWriteQueueEntry*> AllEntries)
+{
+	TRACE_CPUPROFILER_EVENT_SCOPE(ScheduleAllEntries);
 
-	if (Entry->bLoadingFromReferenceDb)
+	const auto HandleDDCGetResult = [this](FIoStoreWriteQueueEntry* Entry, FSharedBuffer Result)
 	{
-		Entry->Writer->LoadFromReferenceDb(Entry);
-	}
-	else
+		bool bFoundInDDC = false;
+		uint64 CompressedSize = 0;
+		if (!Result.IsNull())
+		{
+			FLargeMemoryReader DDCDataReader((uint8*)Result.GetData(), Result.GetSize());
+			bFoundInDDC = Entry->Writer->SerializeCompressedDDCData(Entry, DDCDataReader, &CompressedSize);
+
+			UE_CLOG(!bFoundInDDC, LogIoStore, Warning,
+				TEXT("Ignoring invalid DDC data for ChunkId=%s, DDCKey=%s, UncompressedSize=%llu, NumChunkBlocks=%d"),
+				*LexToString(Entry->ChunkId),
+				*WriteToString<96>(Entry->DDCKey),
+				Entry->UncompressedSize.Get(0),
+				Entry->NumChunkBlocks);
+		}
+		if (bFoundInDDC)
+		{
+			Entry->bFoundInDDC = true;
+			CompressionDDCHitsByType[(int8)Entry->ChunkId.GetChunkType()].IncrementExchange();
+			CompressionDDCGetBytes.AddExchange(CompressedSize);
+			TRACE_COUNTER_INCREMENT(IoStoreDDCHitCount);
+			Entry->BeginCompressionBarrier.Trigger();
+		}
+		else
+		{
+			Entry->bStoreCompressedDataInDDC = true;
+			CompressionDDCMissCount.IncrementExchange();
+			TRACE_COUNTER_INCREMENT(IoStoreDDCMissCount);
+			// kick off source buffer read, and proceed to begin compression
+			Entry->Request->PrepareSourceBufferAsync(Entry->BeginCompressionBarrier);
+		}
+	};
+
+	FIoStoreDDCGetRequestDispatcher DDCGetRequestDispatcher(FIoStoreDDCRequestDispatcherParams{});
+
+	for (FIoStoreWriteQueueEntry* Entry : AllEntries)
 	{
-		Entry->Request->PrepareSourceBufferAsync(Entry->BeginCompressionBarrier);
+		uint64 LocalScheduledCompressionMemory = ScheduledCompressionMemory.Load();
+		
+		while (LocalScheduledCompressionMemory > 0 &&
+			LocalScheduledCompressionMemory + Entry->CompressionMemoryEstimate > MaxCompressionBufferMemory)
+		{
+			TRACE_CPUPROFILER_EVENT_SCOPE(WaitForCompressionMemory);
+			if (!CompressionMemoryReleasedEvent->Wait(100.f))
+			{
+				// if the event timed out,
+				// make sure we are not waiting for unsubmitted ddc requests
+				DDCGetRequestDispatcher.DispatchGetRequests(HandleDDCGetResult);
+			}
+			LocalScheduledCompressionMemory = ScheduledCompressionMemory.Load();
+		}
+
+		ScheduledCompressionMemory.AddExchange(Entry->CompressionMemoryEstimate);
+		TRACE_COUNTER_SET(IoStoreCompressionMemoryScheduled, ScheduledCompressionMemory.Load());
+
+		if (Entry->bLoadingFromReferenceDb)
+		{
+			Entry->Writer->LoadFromReferenceDb(Entry);
+		}
+		else if (Entry->bUseDDCForCompression)
+		{
+			Entry->DDCKey = MakeDDCKey(Entry);
+			DDCGetRequestDispatcher.EnqueueGetRequest(Entry);
+		}
+		else
+		{
+			Entry->Request->PrepareSourceBufferAsync(Entry->BeginCompressionBarrier);
+		}
+
+		DDCGetRequestDispatcher.DispatchGetRequests(HandleDDCGetResult);
+		BeginCompressionQueue.Enqueue(Entry);
 	}
+
+	DDCGetRequestDispatcher.FlushGetRequests(HandleDDCGetResult);
+	BeginCompressionQueue.CompleteAdding();
 }
 
 void FIoStoreWriterContextImpl::BeginCompressionThreadFunc()
@@ -1934,6 +2275,27 @@ void FIoStoreWriterContextImpl::BeginCompressionThreadFunc()
 void FIoStoreWriterContextImpl::BeginEncryptionAndSigningThreadFunc()
 {
 	TRACE_CPUPROFILER_EVENT_SCOPE(BeginEncryptionAndSigningThread);
+
+	const auto HandleDDCPutResult = [this](FIoStoreWriteQueueEntry* Entry, bool bSuccess)
+	{
+		if (bSuccess)
+		{
+			TRACE_COUNTER_INCREMENT(IoStoreDDCPutCount);
+			CompressionDDCPutsByType[(int8)Entry->ChunkId.GetChunkType()].IncrementExchange();
+			CompressionDDCPutBytes.AddExchange(Entry->CompressedSize);
+		}
+		else
+		{
+			CompressionDDCPutErrorCount.IncrementExchange();
+		}
+	};
+
+	FIoStoreDDCRequestDispatcherParams PutRequestDispatcherParams;
+	PutRequestDispatcherParams.QueueTimeLimitMs = 1000.f;
+	PutRequestDispatcherParams.MaxInflightBytes = 256ull << 20;
+	PutRequestDispatcherParams.MaxBatchBytes = 1ull << 20;
+	FIoStoreDDCPutRequestDispatcher DDCPutRequestDispatcher(PutRequestDispatcherParams);
+
 	for (;;)
 	{
 		FIoStoreWriteQueueEntry* Entry = BeginEncryptionAndSigningQueue.DequeueOrWait();
@@ -1945,23 +2307,29 @@ void FIoStoreWriterContextImpl::BeginEncryptionAndSigningThreadFunc()
 		{
 			FIoStoreWriteQueueEntry* Next = Entry->Next;
 			Entry->FinishCompressionBarrier.Wait();
-			Entry->Request->FreeSourceBuffer();
+			TRACE_COUNTER_INCREMENT(IoStoreBeginEncryptionAndSigningCount);
+			Entry->Writer->BeginEncryptAndSign(Entry);
 			if (Entry->bStoreCompressedDataInDDC)
 			{
-				TRACE_CPUPROFILER_EVENT_SCOPE(StoreInDDC);
-				TArray<uint8> DDCData;
-				FMemoryWriter DDCDataWriter(DDCData, true);
+				TRACE_CPUPROFILER_EVENT_SCOPE(AddDDCPutRequest);
+				TArray64<uint8> DDCData;
+				FMemoryWriter64 DDCDataWriter(DDCData, true);
+				DDCData.Reserve(16 + 8 * Entry->ChunkBlocks.Num() + Entry->CompressedSize);
 				if (Entry->Writer->SerializeCompressedDDCData(Entry, DDCDataWriter))
 				{
-					DDC->Put(*Entry->DDCCacheKey, DDCData, Entry->Options.FileName);
+					DDCPutRequestDispatcher.EnqueuePutRequest(Entry, MakeSharedBufferFromArray(MoveTemp(DDCData)));
+				}
+				else
+				{
+					CompressionDDCPutErrorCount.IncrementExchange();
 				}
 			}
-			TRACE_COUNTER_INCREMENT(IoStoreBeginEncryptionAndSigningCount);
-			Entry->Writer->BeginEncryptAndSign(Entry);
+			DDCPutRequestDispatcher.DispatchPutRequests(HandleDDCPutResult);
 			WriterQueue.Enqueue(Entry);
 			Entry = Next;
 		}
 	}
+	DDCPutRequestDispatcher.FlushPutRequests(HandleDDCPutResult);
 	WriterQueue.CompleteAdding();
 }
 
diff --git a/Engine/Source/Developer/IoStoreUtilities/Private/IoStoreWriter.h b/Engine/Source/Developer/IoStoreUtilities/Private/IoStoreWriter.h
--- a/Engine/Source/Developer/IoStoreUtilities/Private/IoStoreWriter.h
+++ b/Engine/Source/Developer/IoStoreUtilities/Private/IoStoreWriter.h
@@ -70,8 +70,14 @@ public:
 		uint64 CompressedChunksByType[(int8)EIoChunkType::MAX] = { 0 };
 		uint64 SerializedChunksCount = 0;
 		uint64 ScheduledCompressionTasksCount = 0;
+		uint64 CompressionDDCHitsByType[(int8)EIoChunkType::MAX] = { 0 };
+		uint64 CompressionDDCPutsByType[(int8)EIoChunkType::MAX] = { 0 };
 		uint64 CompressionDDCHitCount = 0;
 		uint64 CompressionDDCMissCount = 0;
+		uint64 CompressionDDCPutCount = 0;
+		uint64 CompressionDDCPutErrorCount = 0;
+		uint64 CompressionDDCGetBytes = 0;
+		uint64 CompressionDDCPutBytes = 0;
 
 		// The number of chunk retrieved from the reference cache database, and their types.
 		uint64 RefDbChunksCount{ 0 };
