Incremental step towards being able to stage to both a pak build as well as a nopak (streamng) build from a snapshot entirely stored in zenserver (no loose files on the filesystem except a ue.projectstore file).
Key changes:
- Removed the use of cookedfiles.manifest
- Changed IoStore mode of UnrealPak to  be capable of getting zenserver launch data from either a package store manifest (cbobject, metadata) OR a project store marker file (json)
- Ensured that the UAT and IoStore mode of UnrealPak can launch zenserver reliably by passing along the SponsorProcessId when launching zenserver
- Ensured that shader archives (and their accompanying json metadata files) can be read either as loose files on disk or directly from zenserver if it the data for them is internal to zenserver (ie: they have a valid  chunk id)
Remaining work:
- Pak mode of UnrealPak must be able to launch zenserver and pull data from it.

diff --git a/Engine/Source/Developer/IoStoreUtilities/Internal/ZenStoreHttpClient.h b/Engine/Source/Developer/IoStoreUtilities/Internal/ZenStoreHttpClient.h
--- a/Engine/Source/Developer/IoStoreUtilities/Internal/ZenStoreHttpClient.h
+++ b/Engine/Source/Developer/IoStoreUtilities/Internal/ZenStoreHttpClient.h
@@ -56,6 +56,7 @@ public:
 	const TCHAR* GetHostName() const { return ZenService.GetInstance().GetHostName(); }
 	uint16 GetPort() const { return ZenService.GetInstance().GetPort(); }
 	const UE::Zen::FZenServiceInstance& GetZenServiceInstance() const { return ZenService.GetInstance(); }
+	UE::Zen::FZenServiceInstance& GetZenServiceInstance() { return ZenService.GetInstance(); }
 #else // Default to localhost:8558 for platforms where Zen wouldn't be supported yet
 	const TCHAR* GetHostName() const { return TEXT("localhost"); }
 	uint16 GetPort() const { return 8558; }
diff --git a/Engine/Source/Developer/IoStoreUtilities/Private/IoStoreUtilities.cpp b/Engine/Source/Developer/IoStoreUtilities/Private/IoStoreUtilities.cpp
--- a/Engine/Source/Developer/IoStoreUtilities/Private/IoStoreUtilities.cpp
+++ b/Engine/Source/Developer/IoStoreUtilities/Private/IoStoreUtilities.cpp
@@ -1044,14 +1044,14 @@ public:
 		TArray<FFileRegion> FileRegions;
 	};
 
-	FCookedPackageStore(const FString& InCookedDir)
+	FCookedPackageStore(FStringView InCookedDir)
 		: CookedDir(InCookedDir)
 	{
 	}
 
-	FIoStatus Load(const TCHAR* ManifestFilename)
+	FIoStatus LoadManifest(const TCHAR* ManifestFilename)
 	{
-		IOSTORE_CPU_SCOPE(LoadCookedPackageStore);
+		IOSTORE_CPU_SCOPE(LoadCookedPackageStoreManifest);
 		double StartTime = FPlatformTime::Seconds();
 
 		TUniquePtr<FArchive> Ar(IFileManager::Get().CreateFileReader(ManifestFilename));
@@ -1095,6 +1095,57 @@ public:
 		return FIoStatus::Ok;
 	}
 
+	FIoStatus LoadProjectStore(const TCHAR* ProjectStoreFilename)
+	{
+		IOSTORE_CPU_SCOPE(LoadCookedProjectStore);
+		double StartTime = FPlatformTime::Seconds();
+
+		TUniquePtr<FArchive> Ar(IFileManager::Get().CreateFileReader(ProjectStoreFilename));
+		if (!Ar)
+		{
+			return FIoStatus(EIoErrorCode::NotFound);
+		}
+
+		TSharedPtr<FJsonObject> ProjectStoreObject;
+		TSharedRef<TJsonReader<UTF8CHAR>> Reader = TJsonReaderFactory<UTF8CHAR>::Create(Ar.Get());
+		if (FJsonSerializer::Deserialize(Reader, ProjectStoreObject) && ProjectStoreObject.IsValid())
+		{
+			const TSharedPtr<FJsonObject>* ZenServerObjectPtr = nullptr;
+			if (ProjectStoreObject->TryGetObjectField(TEXT("zenserver"), ZenServerObjectPtr) && (ZenServerObjectPtr != nullptr))
+			{
+				FString ProjectId;
+				FString OplogId;
+
+				const TSharedPtr<FJsonObject>& ZenServerObject = *ZenServerObjectPtr;
+				if (ZenServerObject->TryGetStringField(TEXT("projectid"), ProjectId) && !ProjectId.IsEmpty() &&
+					ZenServerObject->TryGetStringField(TEXT("oplogid"), OplogId) && !OplogId.IsEmpty())
+				{
+					ZenStoreClient = MakeUnique<UE::FZenStoreHttpClient>();
+					ZenStoreClient->InitializeReadOnly(ProjectId, OplogId);
+
+					IOSTORE_CPU_SCOPE(FetchOplog);
+					TIoStatusOr<FCbObject> OplogStatus = ZenStoreClient->GetOplog().Get();
+					if (!OplogStatus.IsOk())
+					{
+						return OplogStatus.Status();
+					}
+					FCbObject OplogObject = OplogStatus.ConsumeValueOrDie();
+
+					UE_LOG(LogIoStore, Display, TEXT("Fetched %d oplog items from Zen in %.2lf seconds"),
+						OplogObject["entries"].AsArrayView().Num(),
+						FPlatformTime::Seconds() - StartTime);
+
+					ParseOplog(OplogObject);
+					LoadChunkHashes();
+
+					return FIoStatus::Ok;
+				}
+			}
+		}
+
+		return FIoStatus(EIoErrorCode::NotFound);
+	}
+
 	void ParseOplog(FCbObject& OplogObject)
 	{
 		IOSTORE_CPU_SCOPE(ParseOplog);
@@ -1117,32 +1168,71 @@ public:
 					FCbObjectView ChunkObj = ChunkEntry.AsObjectView();
 					FIoChunkId ChunkId;
 					ChunkId.Set(ChunkObj["id"].AsObjectId().GetView());
-					FChunkInfo& ChunkInfo = ChunkInfoMap.Add(ChunkId);
-					ChunkInfo.ChunkId = ChunkId;
-					ChunkInfo.PackageName = PackageStoreEntry.PackageName;
-					if (ChunkObj["filename"])
-					{
-						TStringBuilder<1024> RelativeFilename;
-						RelativeFilename.Append(ChunkObj["filename"].AsString());
-						ChunkInfo.RelativeFileName = RelativeFilename;
-						TStringBuilder<1024> PathBuilder;
-						FPathViews::AppendPath(PathBuilder, CookedDir);
-						FPathViews::AppendPath(PathBuilder, RelativeFilename);
-						FPathViews::NormalizeFilename(PathBuilder);
-						FilenameToChunkIdMap.Add(*PathBuilder, ChunkId);
-					}
-					const FCbArrayView RegionsArray = ChunkObj["fileregions"].AsArrayView();
-					ChunkInfo.FileRegions.Reserve(RegionsArray.Num());
-					for (FCbFieldView RegionObj : RegionsArray)
+					if (ChunkId.IsValid())
 					{
-						FFileRegion& Region = ChunkInfo.FileRegions.AddDefaulted_GetRef();
-						FFileRegion::LoadFromCompactBinary(RegionObj, Region);
+						FChunkInfo& ChunkInfo = ChunkInfoMap.Add(ChunkId);
+						ChunkInfo.ChunkId = ChunkId;
+						ChunkInfo.PackageName = PackageStoreEntry.PackageName;
+						if (ChunkObj["filename"])
+						{
+							TStringBuilder<1024> RelativeFilename;
+							RelativeFilename.Append(ChunkObj["filename"].AsString());
+							ChunkInfo.RelativeFileName = RelativeFilename;
+							TStringBuilder<1024> PathBuilder;
+							FPathViews::AppendPath(PathBuilder, CookedDir);
+							FPathViews::AppendPath(PathBuilder, RelativeFilename);
+							FPathViews::NormalizeFilename(PathBuilder);
+							FilenameToChunkIdMap.Add(*PathBuilder, ChunkId);
+						}
+						else if (ChunkObj["clientpath"])
+						{
+							TStringBuilder<1024> RelativeFilename;
+							FUtf8StringView ClientPathView = ChunkObj["clientpath"].AsString();
+							if (ClientPathView.StartsWith("/{project}/"))
+							{
+								if (!FPaths::IsProjectFilePathSet())
+								{
+									RelativeFilename.Append(ClientPathView);
+									UE_LOG(LogIoStore, Warning, TEXT("Project relative path could not be remapped because project file path is unset (possibly due to not specifying uproject path as first argument). %s"), RelativeFilename.ToString());
+								}
+								else
+								{
+									RelativeFilename.Append(FPathViews::GetBaseFilename(FPaths::GetProjectFilePath()));
+									RelativeFilename.AppendChar('/');
+									RelativeFilename.Append(ClientPathView.RightChop(11));
+								}
+							}
+							else if (ClientPathView.StartsWith("/{engine}/"))
+							{
+								RelativeFilename.Append(TEXT("Engine/"));
+								RelativeFilename.Append(ClientPathView.RightChop(9));
+							}
+							else
+							{
+								RelativeFilename.Append(ClientPathView);
+							}
+							ChunkInfo.RelativeFileName = RelativeFilename;
+							TStringBuilder<1024> PathBuilder;
+							FPathViews::AppendPath(PathBuilder, CookedDir);
+							FPathViews::AppendPath(PathBuilder, RelativeFilename);
+							FPathViews::NormalizeFilename(PathBuilder);
+							FilenameToChunkIdMap.Add(*PathBuilder, ChunkId);
+						}
+
+						const FCbArrayView RegionsArray = ChunkObj["fileregions"].AsArrayView();
+						ChunkInfo.FileRegions.Reserve(RegionsArray.Num());
+						for (FCbFieldView RegionObj : RegionsArray)
+						{
+							FFileRegion& Region = ChunkInfo.FileRegions.AddDefaulted_GetRef();
+							FFileRegion::LoadFromCompactBinary(RegionObj, Region);
+						}
 					}
 				}
 			};
 
 			AddChunksFromOplog("packagedata");
 			AddChunksFromOplog("bulkdata");
+			AddChunksFromOplog("files");
 
 			PackageIdToEntry.Add(PackageStoreEntry.GetPackageId(), MoveTemp(PackageStoreEntry));
 		}
@@ -2247,12 +2337,29 @@ TArray<TUniquePtr<FIoStoreReader>> CreatePatchSourceReaders(const TArray<FString
 	return Readers;
 }
 
-bool LoadShaderAssetInfo(const FString& Filename, TMap<FSHAHash, TSet<FName>>& OutShaderCodeToAssets)
+static bool LoadShaderAssetInfo(const FString& Filename, FCookedPackageStore* PackageStore, TMap<FSHAHash, TSet<FName>>& OutShaderCodeToAssets)
 {
 	FString JsonText; 
-	if (!FFileHelper::LoadFileToString(JsonText, *Filename))
+	if (PackageStore)
 	{
-		return false;
+		FIoChunkId ChunkId = PackageStore->GetChunkIdFromFileName(Filename);
+		if (!ChunkId.IsValid())
+		{
+			return false;
+		}
+		TIoStatusOr<FIoBuffer> Buffer = PackageStore->ReadChunk(ChunkId);
+		if (!Buffer.IsOk())
+		{
+			return false;
+		}
+		FFileHelper::BufferToString(JsonText, Buffer.ValueOrDie().GetData(), (int32)Buffer.ValueOrDie().GetSize());
+	}
+	else
+	{
+		if (!FFileHelper::LoadFileToString(JsonText, *Filename))
+		{
+			return false;
+		}
 	}
 
 	TSharedPtr<FJsonObject> JsonObject;
@@ -2302,6 +2409,8 @@ bool LoadShaderAssetInfo(const FString& Filename, TMap<FSHAHash, TSet<FName>>& O
 
 static bool ConvertToIoStoreShaderLibrary(
 	const TCHAR* FileName,
+	FIoChunkId InChunkId,
+	FCookedPackageStore* PackageStore,
 	FOodleDataCompression::ECompressor InShaderOodleCompressor,
 	FOodleDataCompression::ECompressionLevel InShaderOodleLevel,
 	TTuple<FIoChunkId, FIoBuffer>& OutLibraryIoChunk,
@@ -2327,23 +2436,35 @@ static bool ConvertToIoStoreShaderLibrary(
 	FString LibraryName(BaseFileNameView.Mid(LibraryNameStartIndex + 1, FormatStartIndex - LibraryNameStartIndex - 1));
 	FName FormatName(BaseFileNameView.RightChop(FormatStartIndex + 1));
 	
-	TUniquePtr<FArchive> LibraryAr(IFileManager::Get().CreateFileReader(FileName));
-	if (!LibraryAr)
+	FIoBuffer LibraryBuffer;
+	if (PackageStore && InChunkId.IsValid())
 	{
-		UE_LOG(LogIoStore, Error, TEXT("Missing shader code library file '%s'."), FileName);
-		return false;
+		LibraryBuffer = PackageStore->ReadChunk(InChunkId).ConsumeValueOrDie();
+	}
+	else
+	{
+		TUniquePtr<FArchive> LibraryAr(IFileManager::Get().CreateFileReader(FileName));
+		if (!LibraryAr)
+		{
+			UE_LOG(LogIoStore, Error, TEXT("Missing shader code library file '%s'."), FileName);
+			return false;
+		}
+
+		LibraryBuffer = FIoBuffer(LibraryAr->TotalSize());
+		LibraryAr->Serialize(LibraryBuffer.GetData(), LibraryBuffer.GetSize());
 	}
 
+	FLargeMemoryReader LibraryAr(LibraryBuffer.GetData(), LibraryBuffer.GetSize());
 	uint32 Version;
-	(*LibraryAr) << Version;
+	LibraryAr << Version;
 
 	FSerializedShaderArchive SerializedShaders;
-	(*LibraryAr) << SerializedShaders;
-	int64 OffsetToShaderCode = LibraryAr->Tell();
+	LibraryAr << SerializedShaders;
+	int64 OffsetToShaderCode = LibraryAr.Tell();
 
 	FString AssetInfoFileName = FPaths::GetPath(FileName) / FString::Printf(TEXT("ShaderAssetInfo-%s-%s.assetinfo.json"), *LibraryName, *FormatName.ToString());
 	TMap<FSHAHash, FShaderMapAssetPaths> ShaderCodeToAssets;
-	if (!LoadShaderAssetInfo(AssetInfoFileName, ShaderCodeToAssets))
+	if (!LoadShaderAssetInfo(AssetInfoFileName, PackageStore && InChunkId.IsValid() ? PackageStore : nullptr, ShaderCodeToAssets))
 	{
 		UE_LOG(LogIoStore, Error, TEXT("Failed loading asset asset info file '%s'"), *AssetInfoFileName);
 		return false;
@@ -2362,7 +2483,7 @@ static bool ConvertToIoStoreShaderLibrary(
 	OutCodeIoChunks.SetNum(GroupCount);
 	FCriticalSection DiskArchiveAccess;
 	ParallelFor(OutCodeIoChunks.Num(),
-		[&OutCodeIoChunks, &IoStoreLibraryHeader, &DiskArchiveAccess, &LibraryAr, &SerializedShaders, &OffsetToShaderCode, InShaderOodleCompressor, InShaderOodleLevel](int32 GroupIndex)
+		[&OutCodeIoChunks, &IoStoreLibraryHeader, &DiskArchiveAccess, &LibraryBuffer, &SerializedShaders, &OffsetToShaderCode, InShaderOodleCompressor, InShaderOodleLevel](int32 GroupIndex)
 		{
 			FIoStoreShaderGroupEntry& Group = IoStoreLibraryHeader.ShaderGroupEntries[GroupIndex];
 			uint8* UncompressedGroupMemory = reinterpret_cast<uint8*>(FMemory::Malloc(Group.UncompressedSize));
@@ -2388,26 +2509,12 @@ static bool ConvertToIoStoreShaderLibrary(
 					// small shaders might be stored without compression, handle them here
 					if (IndividuallyCompressedShader.Size == IndividuallyCompressedShader.UncompressedSize)
 					{
-						// disk access has to be serialized between the for loops
-						FScopeLock Lock(&DiskArchiveAccess);
-						LibraryAr->Seek(OffsetToShaderCode + IndividuallyCompressedShader.Offset);
-						LibraryAr->Serialize(ShaderStart, IndividuallyCompressedShader.Size);
+						FMemory::Memcpy(ShaderStart, LibraryBuffer.GetData() + OffsetToShaderCode + IndividuallyCompressedShader.Offset, IndividuallyCompressedShader.Size);
 					}
 					else
 					{
-						uint8* CompressedShaderMemory = reinterpret_cast<uint8*>(FMemory::Malloc(IndividuallyCompressedShader.Size));
-
-						// disk access has to be serialized between the for loops
-						{
-							FScopeLock Lock(&DiskArchiveAccess);
-							LibraryAr->Seek(OffsetToShaderCode + IndividuallyCompressedShader.Offset);
-							LibraryAr->Serialize(CompressedShaderMemory, IndividuallyCompressedShader.Size);
-						}
-
 						// This function will crash if decompression fails.
-						ShaderCodeArchive::DecompressShaderWithOodle(ShaderStart, IndividuallyCompressedShader.UncompressedSize, CompressedShaderMemory, IndividuallyCompressedShader.Size);
-
-						FMemory::Free(CompressedShaderMemory);
+						ShaderCodeArchive::DecompressShaderWithOodle(ShaderStart, IndividuallyCompressedShader.UncompressedSize, LibraryBuffer.GetData() + OffsetToShaderCode + IndividuallyCompressedShader.Offset, IndividuallyCompressedShader.Size);
 					}
 				}
 
@@ -2592,7 +2699,7 @@ static void ProcessShaderLibraries(const FIoStoreArguments& Arguments, TArray<FC
 					TTuple<FIoChunkId, FIoBuffer> LibraryChunk;
 					TArray<TTuple<FIoChunkId, FIoBuffer, uint32>> CodeChunks;
 					TArray<TTuple<FSHAHash, TSet<FName>>> ShaderMapAssetAssociations;
-				    if (!ConvertToIoStoreShaderLibrary(*TargetFile.NormalizedSourcePath, Arguments.ShaderOodleCompressor, Arguments.ShaderOodleLevel, LibraryChunk, CodeChunks, ShaderMaps, ShaderMapAssetAssociations))
+				    if (!ConvertToIoStoreShaderLibrary(*TargetFile.NormalizedSourcePath, TargetFile.ChunkId, Arguments.PackageStore.Get(), Arguments.ShaderOodleCompressor, Arguments.ShaderOodleLevel, LibraryChunk, CodeChunks, ShaderMaps, ShaderMapAssetAssociations))
 					{
 						UE_LOG(LogIoStore, Warning, TEXT("Failed converting shader library '%s'"), *TargetFile.NormalizedSourcePath);
 						continue;
@@ -2970,15 +3077,22 @@ void InitializeContainerTargetsAndPackages(
 		FStringView Extension = GetFullExtension(SourceFile.NormalizedPath);
 		if (Extension == TEXT(".ushaderbytecode"))
 		{
-			const FCookedFileStatData* CookedFileStatData = Arguments.CookedFileStatMap.Find(SourceFile.NormalizedPath);
-			if (!CookedFileStatData)
+			if (const FCookedFileStatData* CookedFileStatData = Arguments.CookedFileStatMap.Find(SourceFile.NormalizedPath))
 			{
-				UE_LOG(LogIoStore, Warning, TEXT("File not found: '%s'"), *SourceFile.NormalizedPath);
-				return false;
+				OutTargetFile.ChunkType = EContainerChunkType::ShaderCodeLibrary;
+				OutTargetFile.SourceSize = uint64(CookedFileStatData->FileSize);
+				OutTargetFile.ChunkId = FIoChunkId::InvalidChunkId;
+				return true;
 			}
-			OutTargetFile.ChunkType = EContainerChunkType::ShaderCodeLibrary;
-			OutTargetFile.SourceSize = uint64(CookedFileStatData->FileSize);
-			return true;
+			else if (const FCookedPackageStore::FChunkInfo* ChunkInfo = PackageStore.GetChunkInfoFromFileName(SourceFile.NormalizedPath))
+			{
+				OutTargetFile.ChunkType = EContainerChunkType::ShaderCodeLibrary;
+				OutTargetFile.SourceSize = ChunkInfo->ChunkSize;
+				OutTargetFile.ChunkId = ChunkInfo->ChunkId;
+				return true;
+			}
+			UE_LOG(LogIoStore, Warning, TEXT("File not found: '%s'"), *SourceFile.NormalizedPath);
+			return false;
 		}
 
 		const FCookedPackageStore::FChunkInfo* ChunkInfo = PackageStore.GetChunkInfoFromFileName(SourceFile.NormalizedPath);
@@ -9830,36 +9944,72 @@ int32 CreateIoStoreContainerFiles(const TCHAR* CmdLine)
 	}
 	else if (FParse::Param(FCommandLine::Get(), TEXT("StartZenServerForStage")))
 	{
-		FString ManifestFilename;
-		if (!FParse::Value(FCommandLine::Get(), TEXT("PackageStoreManifest="), ManifestFilename))
+		TArray<uint32> SponsorProcessIds;
+		uint32 SponsorProcessId = 0;
+		if (FParse::Value(FCommandLine::Get(), TEXT("SponsorProcessID="), SponsorProcessId) && (SponsorProcessId > 0))
 		{
-			UE_LOG(LogIoStore, Error, TEXT("Expected -PackageStoreManifest=<path to package store manifest>"));
-			return -1;
+			SponsorProcessIds.Add(SponsorProcessId);
 		}
 
-		TUniquePtr<FArchive> Ar(IFileManager::Get().CreateFileReader(*ManifestFilename));
-		if (!Ar)
+		FString ManifestFilename;
+		if (FParse::Value(FCommandLine::Get(), TEXT("PackageStoreManifest="), ManifestFilename))
 		{
-			UE_LOG(LogIoStore, Error, TEXT("Failed reading package store manifest"));
-			return -1;
+			TUniquePtr<FArchive> Ar(IFileManager::Get().CreateFileReader(*ManifestFilename));
+			if (Ar)
+			{
+				FCbObject ManifestObject = LoadCompactBinary(*Ar).AsObject();
+				FCbObject OplogObject;
+				FCbField ZenServerField = ManifestObject["zenserver"];
+				if (ZenServerField)
+				{
+					UE::Zen::FServiceSettings ZenServiceSettings;
+					ZenServiceSettings.ReadFromCompactBinary(ZenServerField["settings"]);
+					FString ProjectId = FString(ZenServerField["projectid"].AsString());
+					FString OplogId = FString(ZenServerField["oplogid"].AsString());
+
+					// We just want the auto launch functionality
+					UE::FZenStoreHttpClient ZenStoreClient(MoveTemp(ZenServiceSettings));
+					ZenStoreClient.InitializeReadOnly(ProjectId, OplogId);
+
+					if (!SponsorProcessIds.IsEmpty() && ZenServiceSettings.IsAutoLaunch() &&
+						ZenServiceSettings.SettingsVariant.Get<UE::Zen::FServiceAutoLaunchSettings>().bLimitProcessLifetime)
+					{
+						ZenStoreClient.GetZenServiceInstance().AddSponsorProcessIDs(SponsorProcessIds);
+					}
+
+					return 0;
+				}
+			}
+			else
+			{
+				UE_LOG(LogIoStore, Error, TEXT("Failed reading package store manifest"));
+			}
 		}
 
-		FCbObject ManifestObject = LoadCompactBinary(*Ar).AsObject();
-		FCbObject OplogObject;
-		FCbField ZenServerField = ManifestObject["zenserver"];
-		if (ZenServerField)
+		FString MarkerFilename;
+		if (FParse::Value(FCommandLine::Get(), TEXT("ProjectStore="), MarkerFilename))
 		{
-			UE::Zen::FServiceSettings ZenServiceSettings;
-			ZenServiceSettings.ReadFromCompactBinary(ZenServerField["settings"]);
-			FString ProjectId = FString(ZenServerField["projectid"].AsString());
-			FString OplogId = FString(ZenServerField["oplogid"].AsString());
+			TUniquePtr<FArchive> Ar(IFileManager::Get().CreateFileReader(*MarkerFilename));
+			if (Ar)
+			{
+				UE::Zen::FZenServiceInstance& ZenServiceInstance = UE::Zen::GetDefaultServiceInstance();
+				const UE::Zen::FServiceSettings& ZenServiceSettings = ZenServiceInstance.GetServiceSettings();
 
-			// We just want the auto launch functionality
-			UE::FZenStoreHttpClient ZenStoreClient(MoveTemp(ZenServiceSettings));
-			ZenStoreClient.InitializeReadOnly(ProjectId, OplogId);
+				if (!SponsorProcessIds.IsEmpty() && ZenServiceSettings.IsAutoLaunch() &&
+					ZenServiceInstance.GetServiceSettings().SettingsVariant.Get<UE::Zen::FServiceAutoLaunchSettings>().bLimitProcessLifetime)
+				{
+					ZenServiceInstance.AddSponsorProcessIDs(SponsorProcessIds);
+				}
+				return 0;
+			}
+			else
+			{
+				UE_LOG(LogIoStore, Error, TEXT("Failed reading project store marker"));
+			}
 		}
 
-		return 0;
+		UE_LOG(LogIoStore, Error, TEXT("Expected -PackageStoreManifest=<path to package store manifest> or -ProjectStore=<path to project store marker>"));
+		return -1;
 	}
 	else if (FParse::Value(FCommandLine::Get(), TEXT("CreateDLCContainer="), Arguments.DLCPluginPath))
 	{
@@ -9986,7 +10136,7 @@ int32 CreateIoStoreContainerFiles(const TCHAR* CmdLine)
 	if (FParse::Value(FCommandLine::Get(), TEXT("PackageStoreManifest="), PackageStoreManifestFilename))
 	{
 		TUniquePtr<FCookedPackageStore> PackageStore = MakeUnique<FCookedPackageStore>(Arguments.CookedDir);
-		FIoStatus Status = PackageStore->Load(*PackageStoreManifestFilename);
+		FIoStatus Status = PackageStore->LoadManifest(*PackageStoreManifestFilename);
 		if (Status.IsOk())
 		{
 			Arguments.PackageStore = MoveTemp(PackageStore);
@@ -9998,8 +10148,25 @@ int32 CreateIoStoreContainerFiles(const TCHAR* CmdLine)
 	}
 	else
 	{
-		UE_LOG(LogIoStore, Error, TEXT("Expected -PackageStoreManifest=<path to package store manifest>"));
-		return -1;
+		FString ProjectStoreFilename;
+		if (FParse::Value(FCommandLine::Get(), TEXT("ProjectStore="), ProjectStoreFilename))
+		{
+			TUniquePtr<FCookedPackageStore> PackageStore = MakeUnique<FCookedPackageStore>(Arguments.CookedDir);
+			FIoStatus Status = PackageStore->LoadProjectStore(*ProjectStoreFilename);
+			if (Status.IsOk())
+			{
+				Arguments.PackageStore = MoveTemp(PackageStore);
+			}
+			else
+			{
+				UE_LOG(LogIoStore, Fatal, TEXT("Failed loading project store '%s'"), *ProjectStoreFilename);
+			}
+		}
+		else
+		{
+			UE_LOG(LogIoStore, Error, TEXT("Expected -PackageStoreManifest=<path to package store manifest> or -ProjectStore=<path to project store marker>"));
+			return -1;
+		}
 	}
 
 	if (!ParseOrderFileArguments(Arguments))
diff --git a/Engine/Source/Developer/IoStoreUtilities/Private/ZenStoreWriter.cpp b/Engine/Source/Developer/IoStoreUtilities/Private/ZenStoreWriter.cpp
--- a/Engine/Source/Developer/IoStoreUtilities/Private/ZenStoreWriter.cpp
+++ b/Engine/Source/Developer/IoStoreUtilities/Private/ZenStoreWriter.cpp
@@ -649,69 +649,6 @@ void FZenStoreWriter::EndCook(const FCookInfo& Info)
 
 		TIoStatusOr<uint64> Status = HttpClient->EndBuildPass(Pkg);
 		UE_CLOG(!Status.IsOk(), LogZenStoreWriter, Fatal, TEXT("Failed to append OpLog and end the build pass"));
-
-		FCbWriter ManifestWriter;
-		ManifestWriter.BeginObject();
-		ManifestWriter.BeginObject("zenserver");
-#if UE_WITH_ZEN
-		ManifestWriter.BeginObject("settings");
-		const UE::Zen::FZenServiceInstance& ZenServiceInstance = HttpClient->GetZenServiceInstance();
-		ZenServiceInstance.GetServiceSettings().WriteToCompactBinary(ManifestWriter);
-		ManifestWriter.EndObject();
-#endif
-		ManifestWriter << "projectid" << ProjectId;
-		ManifestWriter << "oplogid" << OplogId;
-		ManifestWriter.EndObject();
-		ManifestWriter.EndObject();
-
-		FString PackageStoreManifestFilePath = FPaths::Combine(MetadataDirectoryPath, TEXT("packagestore.manifest"));
-		TUniquePtr<FArchive> Ar(IFileManager::Get().CreateFileWriter(*PackageStoreManifestFilePath));
-		if (Ar)
-		{
-			SaveCompactBinary(*Ar, ManifestWriter.Save());
-		}
-		else
-		{
-			UE_LOG(LogSavePackage, Error, TEXT("Failed saving package store manifest file '%s'"), *PackageStoreManifestFilePath);
-		}
-
-		{
-			// Temporary solution until we can reliably read the oplog from UAT
-			TArray<FString> CookedFiles;
-			TIoStatusOr<FCbObject> OplogStatus = HttpClient->GetOplog().Get();
-			if (OplogStatus.IsOk())
-			{
-				FCbObject Oplog = OplogStatus.ConsumeValueOrDie();
-				for (FCbField& OplogEntry : Oplog["entries"].AsArray())
-				{
-					FCbObject OplogObj = OplogEntry.AsObject();
-					for (FCbField& ChunkEntry : OplogEntry["packagedata"].AsArray())
-					{
-						FCbObject ChunkObj = ChunkEntry.AsObject();
-						if (ChunkObj["filename"])
-						{
-							CookedFiles.Add(FString(ChunkObj["filename"].AsString()));
-						}
-					}
-					for (FCbField& ChunkEntry : OplogEntry["bulkdata"].AsArray())
-					{
-						FCbObject ChunkObj = ChunkEntry.AsObject();
-						if (ChunkObj["filename"])
-						{
-							CookedFiles.Add(FString(ChunkObj["filename"].AsString()));
-						}
-					}
-				}
-				if (!FFileHelper::SaveStringArrayToFile(CookedFiles, *FPaths::Combine(MetadataDirectoryPath, TEXT("cookedfiles.manifest"))))
-				{
-					UE_LOG(LogSavePackage, Error, TEXT("Failed writing UAT file manifest"));
-				}
-			}
-			else
-			{
-				UE_LOG(LogSavePackage, Error, TEXT("Failed reading oplog"));
-			}
-		}
 	}
 
 	UE_LOG(LogZenStoreWriter, Display, TEXT("Output:\t%d Public runtime script objects"), PackageStoreOptimizer->GetTotalScriptObjectCount());
diff --git a/Engine/Source/Programs/AutomationTool/AutomationUtils/DeploymentContext.cs b/Engine/Source/Programs/AutomationTool/AutomationUtils/DeploymentContext.cs
--- a/Engine/Source/Programs/AutomationTool/AutomationUtils/DeploymentContext.cs
+++ b/Engine/Source/Programs/AutomationTool/AutomationUtils/DeploymentContext.cs
@@ -94,9 +94,10 @@ public class FilesToStage
 	}
 }
 
-public class PackageStoreManifest
+public class PackageStoreData
 {
-	public string FullPath { get; set; }
+	public string ManifestFullPath { get; set; }
+	public string MarkerFullPath { get; set; }
 	public IList<string> ZenCookedFiles { get; set; }
 }
 
@@ -259,9 +260,9 @@ public class DeploymentContext //: ProjectParams
 	public string PakFileInternalRoot = "../../../";
 
 	/// <summary>
-	/// Cooked package store manifest if available
+	/// Cooked package store data if available
 	/// </summary>
-	public PackageStoreManifest PackageStoreManifest;
+	public PackageStoreData PackageStoreData;
 
 	/// <summary>
 	/// List of files to be staged
diff --git a/Engine/Source/Programs/AutomationTool/Scripts/CopyBuildToStagingDirectory.Automation.cs b/Engine/Source/Programs/AutomationTool/Scripts/CopyBuildToStagingDirectory.Automation.cs
--- a/Engine/Source/Programs/AutomationTool/Scripts/CopyBuildToStagingDirectory.Automation.cs
+++ b/Engine/Source/Programs/AutomationTool/Scripts/CopyBuildToStagingDirectory.Automation.cs
@@ -19,9 +19,11 @@ using UnrealBuildBase;
 using System.Text.Json;
 using System.Net.Http;
 using EpicGames.Serialization;
+using EpicGames.ProjectStore;
 using Microsoft.Extensions.Logging;
 
 using static AutomationTool.CommandUtils;
+using System.Text.Json.Serialization;
 
 namespace AutomationScripts
 {
@@ -829,39 +831,132 @@ namespace AutomationScripts
 			return Plugins.ToList();
 		}
 
-		private static void LoadPackageStoreManifest(ProjectParams Params, DeploymentContext SC)
+		private static List<string> ReadZenCookedFilesFromZenServer(string Host, int Port, string ProjectId, string OplogId, ProjectParams Params, DeploymentContext SC, bool bAutoLaunch, string PackageStoreFileArgName, string PackageStoreFileArgValue)
+		{
+			List<string> ZenCookedFiles = new List<string>();
+
+			bool bAttemptAutoLaunchOnFailure = bAutoLaunch;
+			HttpResponseMessage HttpGetResult = null;
+			for (; ; )
+			{
+				HttpClient HttpClient = new HttpClient();
+				using var Request = new HttpRequestMessage(HttpMethod.Get, string.Format("http://{0}:{1}/prj/{2}/oplog/{3}/entries", Host, Port, ProjectId, OplogId));
+				Request.Headers.Add("Accept", "application/x-ue-cb");
+				try
+				{
+					HttpGetResult = HttpClient.Send(Request);
+					break;
+				}
+				catch
+				{
+					if (bAttemptAutoLaunchOnFailure)
+					{
+						string Arguments = String.Format("IoStore -StartZenServerForStage -SponsorProcessID={0} -{1}={2}", Environment.ProcessId, PackageStoreFileArgName, MakePathSafeToUseWithCommandLine(PackageStoreFileArgValue));
+						if (Params.Unattended)
+						{
+							Arguments += " -unattended";
+						}
+
+						RunUnrealPak(Params.RawProjectPath, Arguments, "ZenAutoLaunch");
+						bAttemptAutoLaunchOnFailure = false;
+					}
+					else
+					{
+						throw new AutomationException(String.Format("Failed sending oplog request to Zen at {0}:{1}. Ensure that the server is running.", Host, Port));
+					}
+				}
+			}
+			if (!HttpGetResult.IsSuccessStatusCode)
+			{
+				throw new AutomationException(String.Format("Failed reading oplog {0}.{1} from Zen. Ensure that cooking was successful.", ProjectId, OplogId));
+			}
+			Task<byte[]> ReadOplogTask = HttpGetResult.Content.ReadAsByteArrayAsync();
+			ReadOplogTask.Wait();
+			byte[] OplogData = ReadOplogTask.Result;
+			CbObject OplogObject = new CbField(OplogData).AsObject();
+			foreach (CbField EntryField in OplogObject["entries"].AsArray())
+			{
+				foreach (CbField PackageDataField in EntryField["packagedata"].AsArray())
+				{
+					string RelativeFilename = PackageDataField["filename"].AsString();
+					ZenCookedFiles.Add(FileReference.Combine(SC.PlatformCookDir, RelativeFilename).FullName);
+				}
+				foreach (CbField PackageDataField in EntryField["bulkdata"].AsArray())
+				{
+					string RelativeFilename = PackageDataField["filename"].AsString();
+					ZenCookedFiles.Add(FileReference.Combine(SC.PlatformCookDir, RelativeFilename).FullName);
+				}
+				if (EntryField["key"].AsString() == "EndCook")
+				{
+					foreach (CbField PackageDataField in EntryField["files"].AsArray())
+					{
+						string RelativeFilename = PackageDataField["clientpath"].AsString();
+						if (RelativeFilename.EndsWith(".bin") || RelativeFilename.EndsWith(".ushaderbytecode"))
+						{
+							RelativeFilename = RelativeFilename.Replace("/{engine}/", "Engine/");
+							RelativeFilename = RelativeFilename.Replace("/{project}/", SC.ShortProjectName+"/");
+							ZenCookedFiles.Add(FileReference.Combine(SC.PlatformCookDir, RelativeFilename).FullName);
+						}
+					}
+				}
+				else
+				{
+					foreach (CbField PackageDataField in EntryField["files"].AsArray())
+					{
+						string RelativeFilename = PackageDataField["clientpath"].AsString();
+						RelativeFilename = RelativeFilename.Replace("/{engine}/", "Engine/");
+						RelativeFilename = RelativeFilename.Replace("/{project}/", SC.ShortProjectName+"/");
+						ZenCookedFiles.Add(FileReference.Combine(SC.PlatformCookDir, RelativeFilename).FullName);
+					}
+				}
+			}
+
+			return ZenCookedFiles;
+		}
+
+		private static void LoadPackageStoreData(ProjectParams Params, DeploymentContext SC)
 		{
 			if (!ShouldCreateIoStoreContainerFiles(Params, SC))
 			{
 				return;
 			}
+
+			FileReference ProjectStoreFile = FileReference.Combine(SC.PlatformCookDir, "ue.projectstore");
+			System.IO.FileInfo ProjectStoreFileInfo = ProjectStoreFile.ToFileInfo();
+			if (ProjectStoreFileInfo.Exists)
+			{
+				SC.PackageStoreData = new PackageStoreData();
+				SC.PackageStoreData.MarkerFullPath = ProjectStoreFile.FullName;
+				byte[] ProjectStoreBytes = FileReference.ReadAllBytes(ProjectStoreFile);
+
+				JsonSerializerOptions Options = new JsonSerializerOptions();
+				Options.AllowTrailingCommas = true;
+				Options.ReadCommentHandling = JsonCommentHandling.Skip;
+				Options.PropertyNameCaseInsensitive = true;
+				Options.Converters.Add(new JsonStringEnumConverter());
+				ProjectStoreData ParsedProjectStore = JsonSerializer.Deserialize<ProjectStoreData>(ProjectStoreBytes, Options)!;
+
+				if ((ParsedProjectStore != null) && (ParsedProjectStore.ZenServer != null))
+				{
+					Logger.LogInformation("Reading oplog from Zen...");
+					ZenServerStoreData ZenServer = ParsedProjectStore.ZenServer;
+					SC.PackageStoreData.ZenCookedFiles = ReadZenCookedFilesFromZenServer(ZenServer.HostName, ZenServer.HostPort, ZenServer.ProjectId, ZenServer.OplogId, Params, SC, true, "ProjectStore", ProjectStoreFile.FullName);
+					return;
+				}
+			}
+
 			FileReference PackageStoreManifestFile = FileReference.Combine(SC.MetadataDir, "packagestore.manifest");
 			System.IO.FileInfo PackageStoreManifestFileInfo = PackageStoreManifestFile.ToFileInfo();
 			if (PackageStoreManifestFileInfo.Exists)
 			{
-				SC.PackageStoreManifest = new PackageStoreManifest();
-				SC.PackageStoreManifest.FullPath = PackageStoreManifestFile.FullName;
+				SC.PackageStoreData = new PackageStoreData();
+				SC.PackageStoreData.ManifestFullPath = PackageStoreManifestFile.FullName;
 				if (PackageStoreManifestFileInfo.Length > (64 << 10))
 				{
 					// We assume that this is a loose file deploy and we only need to read zen manifests here
 					return;
 				}
 
-				{
-					// Temporary solution until we can reliably read the oplog from UAT
-					FileReference CookedFilesManifestFile = FileReference.Combine(SC.MetadataDir, "cookedfiles.manifest");
-					if (FileReference.Exists(CookedFilesManifestFile))
-					{
-						SC.PackageStoreManifest.ZenCookedFiles = new List<string>();
-						string[] CookedFiles = FileReference.ReadAllLines(CookedFilesManifestFile);
-						foreach (string RelativeFilename in CookedFiles)
-						{
-							SC.PackageStoreManifest.ZenCookedFiles.Add(FileReference.Combine(SC.PlatformCookDir, RelativeFilename).FullName);
-						}
-						return;
-					}
-				}
-
 				byte[] ManifestData = File.ReadAllBytes(PackageStoreManifestFile.FullName);
 				CbObject ManifestObject = new CbField(ManifestData).AsObject();
 				CbObject ZenServerObject = ManifestObject["zenserver"].AsObject();
@@ -890,66 +985,12 @@ namespace AutomationScripts
 						}
 					}
 
-					SC.PackageStoreManifest.ZenCookedFiles = new List<string>();
-
-					bool bAttemptAutoLaunchOnFailure = bAutoLaunch;
-					HttpResponseMessage HttpGetResult = null;
-					for (; ; )
-					{
-						HttpClient HttpClient = new HttpClient();
-						using var Request = new HttpRequestMessage(HttpMethod.Get, string.Format("http://{0}:{1}/prj/{2}/oplog/{3}/entries", Host, Port, ProjectId, OplogId));
-						Request.Headers.Add("Accept", "application/x-ue-cb");
-						try
-						{
-							HttpGetResult = HttpClient.Send(Request);
-							break;
-						}
-						catch
-						{
-							if (bAttemptAutoLaunchOnFailure)
-							{
-								string Arguments = String.Format("IoStore -StartZenServerForStage -PackageStoreManifest={0}", MakePathSafeToUseWithCommandLine(SC.PackageStoreManifest.FullPath));
-								if (Params.Unattended)
-								{
-									Arguments += " -unattended";
-								}
-
-								RunUnrealPak(Params.RawProjectPath, Arguments, "ZenAutoLaunch");
-								bAttemptAutoLaunchOnFailure = false;
-							}
-							else
-							{
-								throw new AutomationException(String.Format("Failed sending oplog request to Zen at {0}:{1}. Ensure that the server is running.", Host, Port));
-							}
-						}
-					}
-					if (!HttpGetResult.IsSuccessStatusCode)
-					{
-						throw new AutomationException(String.Format("Failed reading oplog {0}.{1} from Zen. Ensure that cooking was successful.", ProjectId, OplogId));
-					}
-					Task<byte[]> ReadOplogTask = HttpGetResult.Content.ReadAsByteArrayAsync();
-					ReadOplogTask.Wait();
-					byte[] OplogData = ReadOplogTask.Result;
-					CbObject OplogObject = new CbField(OplogData).AsObject();
-					foreach (CbField EntryField in OplogObject["entries"].AsArray())
-					{
-						foreach (CbField PackageDataField in EntryField["packagedata"].AsArray())
-						{
-							string RelativeFilename = PackageDataField["filename"].AsString();
-							SC.PackageStoreManifest.ZenCookedFiles.Add(FileReference.Combine(SC.PlatformCookDir, RelativeFilename).FullName);
-						}
-						foreach (CbField PackageDataField in EntryField["bulkdata"].AsArray())
-						{
-							string RelativeFilename = PackageDataField["filename"].AsString();
-							SC.PackageStoreManifest.ZenCookedFiles.Add(FileReference.Combine(SC.PlatformCookDir, RelativeFilename).FullName);
-						}
-					}
+					SC.PackageStoreData.ZenCookedFiles = ReadZenCookedFilesFromZenServer(Host, Port, ProjectId, OplogId, Params, SC, bAutoLaunch, "PackageStoreManifest", PackageStoreManifestFile.FullName);
+					return;
 				}
 			}
-			else
-			{
-				throw new AutomationException(String.Format("A package store manifest is required when staging to IoStore. Expected to find {0}. Ensure that cooking was successful.", PackageStoreManifestFile.FullName));
-			}
+
+			throw new AutomationException(String.Format("A package store manifest or project store is required when staging to IoStore. Expected to find {0} or {1}. Ensure that cooking was successful.", PackageStoreManifestFile.FullName, ProjectStoreFile.FullName));
 		}
 
 		private static bool CanCookedFileBeStaged(FileReference File)
@@ -1144,18 +1185,25 @@ namespace AutomationScripts
 				// Stage all the cooked data, this is the same rule as normal stage except we may skip Engine
 				List<FileReference> CookedFiles = new List<FileReference>();
 
-				if (DirectoryReference.Exists(SC.PlatformCookDir))
+				// When cooking to Zen get the list of cooked package files
+				LoadPackageStoreData(Params, SC);
+				if (SC.PackageStoreData != null && SC.PackageStoreData.ZenCookedFiles != null)
 				{
-					CookedFiles.AddRange(DirectoryReference.EnumerateFiles(SC.PlatformCookDir, "*", SearchOption.AllDirectories).ToList());
+					if (DirectoryReference.Exists(SC.PlatformCookDir))
+					{
+						CookedFiles.AddRange(DirectoryReference.EnumerateFiles(SC.PlatformCookDir, "*.projectstore", SearchOption.TopDirectoryOnly).ToList());
+					}
+					foreach (string FilePath in SC.PackageStoreData.ZenCookedFiles)
+					{
+						CookedFiles.Add(new FileReference(FilePath));
+					}
 				}
-
-				// When cooking to Zen get the list of cooked package files from the manifest
-				LoadPackageStoreManifest(Params, SC);
-				if (SC.PackageStoreManifest != null && SC.PackageStoreManifest.ZenCookedFiles != null)
+				else
 				{
-					foreach (string FilePath in SC.PackageStoreManifest.ZenCookedFiles)
+
+					if (DirectoryReference.Exists(SC.PlatformCookDir))
 					{
-						CookedFiles.Add(new FileReference(FilePath));
+						CookedFiles.AddRange(DirectoryReference.EnumerateFiles(SC.PlatformCookDir, "*", SearchOption.AllDirectories).ToList());
 					}
 				}
 
@@ -1500,11 +1548,11 @@ namespace AutomationScripts
 					{
 						List<FileReference> CookedFiles = DirectoryReference.EnumerateFiles(SC.PlatformCookDir, "*", SearchOption.AllDirectories).ToList();
 
-						// When cooking to Zen get the list of cooked package files from the manifest
-						LoadPackageStoreManifest(Params, SC);
-						if (SC.PackageStoreManifest != null && SC.PackageStoreManifest.ZenCookedFiles != null)
+						// When cooking to Zen get the list of cooked package files
+						LoadPackageStoreData(Params, SC);
+						if (SC.PackageStoreData != null && SC.PackageStoreData.ZenCookedFiles != null)
 						{
-							foreach (string FilePath in SC.PackageStoreManifest.ZenCookedFiles)
+							foreach (string FilePath in SC.PackageStoreData.ZenCookedFiles)
 							{
 								CookedFiles.Add(new FileReference(FilePath));
 							}
@@ -4160,11 +4208,23 @@ namespace AutomationScripts
 				CommandletParams += String.Format("-CreateGlobalContainer={0}", MakePathSafeToUseWithCommandLine(GlobalContainerOutputLocation.FullName));
 			}
 
-			CommandletParams += String.Format(" -CookedDirectory={0} -PackageStoreManifest={1} -Commands={2}",
+			if (!string.IsNullOrEmpty(SC.PackageStoreData.ManifestFullPath))
+			{
+				CommandletParams += String.Format(" -PackageStoreManifest={0}",
+					MakePathSafeToUseWithCommandLine(SC.PackageStoreData.ManifestFullPath));
+			}
+
+			if (!string.IsNullOrEmpty(SC.PackageStoreData.MarkerFullPath))
+			{
+				CommandletParams += String.Format(" -ProjectStore={0}",
+					MakePathSafeToUseWithCommandLine(SC.PackageStoreData.MarkerFullPath));
+			}
+
+			CommandletParams += String.Format(" -CookedDirectory={0} -Commands={1}",
 				MakePathSafeToUseWithCommandLine(SC.PlatformCookDir.ToString()),
-				MakePathSafeToUseWithCommandLine(SC.PackageStoreManifest.FullPath),
 				MakePathSafeToUseWithCommandLine(CommandsFileName));
-			if (SC.PackageStoreManifest.ZenCookedFiles == null)
+
+			if (SC.PackageStoreData.ZenCookedFiles == null)
 			{
 				FileReference ScriptObjectsFile = FileReference.Combine(SC.MetadataDir, "scriptobjects.bin");
 				if (!FileReference.Exists(ScriptObjectsFile))
