Added option -CookSoftPackageReferences to save soft package references when cooking

#rnx

diff --git a/Engine/Source/Developer/IoStoreUtilities/Internal/PackageStoreOptimizer.h b/Engine/Source/Developer/IoStoreUtilities/Internal/PackageStoreOptimizer.h
--- a/Engine/Source/Developer/IoStoreUtilities/Internal/PackageStoreOptimizer.h
+++ b/Engine/Source/Developer/IoStoreUtilities/Internal/PackageStoreOptimizer.h
@@ -194,6 +194,7 @@ private:
 	TArray<FImportedPackageRef> ImportedPackages;
 	TArray<uint64> ImportedPublicExportHashes;
 	TArray<FBulkDataMapEntry> BulkDataEntries;
+	TArray<FName> SoftPackageReferences;
 
 	FIoBuffer HeaderBuffer;
 
@@ -244,6 +245,7 @@ private:
 		TArray<FObjectExport> ObjectExports;
 		TArray<FPackageIndex> PreloadDependencies;
 		TArray<FObjectDataResource> DataResources;
+		TArray<FName> SoftPackageReferences;
 	};
 
 	using FExportGraphEdges = TMultiMap<FPackageStorePackage::FExportGraphNode*, FPackageStorePackage::FExportGraphNode*>;
diff --git a/Engine/Source/Developer/IoStoreUtilities/Private/IoStoreUtilities.cpp b/Engine/Source/Developer/IoStoreUtilities/Private/IoStoreUtilities.cpp
--- a/Engine/Source/Developer/IoStoreUtilities/Private/IoStoreUtilities.cpp
+++ b/Engine/Source/Developer/IoStoreUtilities/Private/IoStoreUtilities.cpp
@@ -4,6 +4,7 @@
 #include "IoStoreWriter.h"
 
 #include "IoStoreLooseFiles.h"
+#include "Algo/BinarySearch.h"
 #include "Algo/TopologicalSort.h"
 #include "Async/AsyncWork.h"
 #include "CookMetadata.h"
@@ -5281,6 +5282,58 @@ void CreateContainerHeader(FContainerTargetSpec& ContainerTarget, bool bIsOption
 		Writer.StoreTocArchive << OffsetToDataFromThis;
 	};
 
+	struct FSoftPackageReferenceWriter
+	{
+		explicit FSoftPackageReferenceWriter(TSet<FPackageId>&& SoftReferences, int32 NumPackageEntries)
+			: SoftPackageReferences(SoftReferences.Array())
+			, TotalEntrySize(NumPackageEntries * sizeof(FFilePackageStoreEntrySoftReferences))
+		{
+			SoftPackageReferences.Sort();
+		}
+
+		void Append(TConstArrayView<FPackageId> SoftRefs)
+		{
+			const int64 RemainingEntrySize = TotalEntrySize - EntryAr.Tell();
+			const int64 OffsetFromThis = RemainingEntrySize + DataAr.Tell();
+			uint32		ArrayNum = static_cast<uint32>(SoftRefs.Num());
+			uint32		OffsetToDataFromThis = ArrayNum > 0 ? OffsetFromThis : 0;
+
+			EntryAr << ArrayNum;
+			EntryAr << OffsetToDataFromThis;
+
+			for (FPackageId SoftRef : SoftRefs)
+			{
+				int32 Index = Algo::BinarySearch(SoftPackageReferences, SoftRef);
+				check(Index != INDEX_NONE);
+				DataAr << Index;
+			}
+		}
+
+		void Flush(FIoContainerHeaderSoftPackageReferences& OutContainerSoftReferences)
+		{
+			const int64 TotalSize = EntryAr.TotalSize() + DataAr.TotalSize();
+			if (TotalSize == 0)
+			{
+				check(SoftPackageReferences.IsEmpty());
+				OutContainerSoftReferences.Empty();
+				return;
+			}
+
+			OutContainerSoftReferences.bContainsSoftPackageReferences = true;
+			OutContainerSoftReferences.PackageIds = MoveTemp(SoftPackageReferences);
+			OutContainerSoftReferences.PackageIndices.AddUninitialized(TotalSize);
+			FBufferWriter Ar(OutContainerSoftReferences.PackageIndices.GetData(), TotalSize);
+			Ar.Serialize(EntryAr.GetData(), EntryAr.TotalSize());
+			Ar.Serialize(DataAr.GetData(), DataAr.TotalSize());
+		}
+
+	private:
+		TArray<FPackageId>	SoftPackageReferences;
+		const int64			TotalEntrySize;
+		FLargeMemoryWriter	EntryAr;
+		FLargeMemoryWriter	DataAr;
+	};
+
 	TArray<const FCookedPackage*> SortedPackages(ContainerTarget.Packages);
 	Algo::Sort(SortedPackages, [](const FCookedPackage* A, const FCookedPackage* B)
 		{
@@ -5327,6 +5380,17 @@ void CreateContainerHeader(FContainerTargetSpec& ContainerTarget, bool bIsOption
 	}
 	else
 	{
+		TSet<FPackageId> AllSoftPackageReferences;
+		for (const FCookedPackage* Package : SortedPackages)
+		{
+			for (FPackageId SoftRef : Package->PackageStoreEntry.SoftPackageReferences)
+			{
+				AllSoftPackageReferences.Add(SoftRef);
+			}
+		}
+
+		FSoftPackageReferenceWriter SoftRefsWriter(MoveTemp(AllSoftPackageReferences), NonOptionalSegmentStoreEntriesCount);
+
 		for (const FCookedPackage* Package : SortedPackages)
 		{
 			const FPackageStoreEntryResource& Entry = Package->PackageStoreEntry;
@@ -5365,7 +5429,12 @@ void CreateContainerHeader(FContainerTargetSpec& ContainerTarget, bool bIsOption
 			{
 				StoreEntriesWriter.StoreDataArchive << const_cast<FSHAHash&>(ShaderMapHash);
 			}
+
+			// SoftPackageReferences
+			SoftRefsWriter.Append(Entry.SoftPackageReferences);
 		}
+
+		SoftRefsWriter.Flush(Header.SoftPackageReferences);
 	}
 	Header.RedirectsNameMap = RedirectsNameMapBuilder.GetNameMap();
 
diff --git a/Engine/Source/Developer/IoStoreUtilities/Private/PackageStoreOptimizer.cpp b/Engine/Source/Developer/IoStoreUtilities/Private/PackageStoreOptimizer.cpp
--- a/Engine/Source/Developer/IoStoreUtilities/Private/PackageStoreOptimizer.cpp
+++ b/Engine/Source/Developer/IoStoreUtilities/Private/PackageStoreOptimizer.cpp
@@ -12,6 +12,7 @@
 #include "UObject/Object.h"
 #include "Serialization/MemoryReader.h"
 #include "UObject/Package.h"
+#include "Misc/PackageName.h"
 #include "Misc/SecureHash.h"
 #include "Misc/StringBuilder.h"
 #include "Serialization/LargeMemoryReader.h"
@@ -112,6 +113,8 @@ FPackageStorePackage* FPackageStoreOptimizer::CreatePackageFromCookedHeader(cons
 		Package->NameMapBuilder.AddName(CookedHeaderData.SummaryNames[I]);
 	}
 
+	Package->SoftPackageReferences = MoveTemp(CookedHeaderData.SoftPackageReferences);
+
 	TArray<FPackageStorePackage::FUnresolvedImport> Imports;
 	ProcessImports(CookedHeaderData, Package, Imports);
 	ProcessExports(CookedHeaderData, Package, Imports.GetData());
@@ -222,7 +225,18 @@ FPackageStoreOptimizer::FCookedHeaderData FPackageStoreOptimizer::LoadCookedHead
 			ProxyAr << ObjectExport;
 		}
 	}
-	
+
+	if (Summary.SoftPackageReferencesCount > 0)
+	{
+		ProxyAr.Seek(Summary.SoftPackageReferencesOffset);
+		CookedHeaderData.SoftPackageReferences.Reserve(Summary.SoftPackageReferencesCount);
+		for (int32 I = 0; I < Summary.SoftPackageReferencesCount; ++I)
+		{
+			FName& SoftReference = CookedHeaderData.SoftPackageReferences.Emplace_GetRef();
+			ProxyAr << SoftReference;
+		}
+	}
+
 	if (Summary.DataResourceOffset > 0)
 	{
 		ProxyAr.Seek(Summary.DataResourceOffset);
@@ -1083,5 +1097,26 @@ FPackageStoreEntryResource FPackageStoreOptimizer::CreatePackageStoreEntry(const
 			Result.OptionalSegmentImportedPackageIds.Add(ImportedPackage.Id);
 		}
 	}
+
+	Result.SoftPackageReferences.Reserve(Package->SoftPackageReferences.Num());
+	for (const FName& SoftRefName : Package->SoftPackageReferences)
+	{
+		TCHAR NameStr[FName::StringBufferSize];
+		const uint32 NameLen = SoftRefName.ToString(NameStr);
+		const FStringView SoftRef(NameStr, NameLen);
+
+		if (!FPackageName::IsScriptPackage(SoftRef))
+		{
+			if (!FPackageName::IsValidLongPackageName(SoftRef))
+			{
+				UE_LOG(LogPackageStoreOptimizer, Warning, TEXT("Invalid soft package reference name '%s'"), *SoftRefName.ToString());
+				continue;
+			}
+
+			Result.SoftPackageReferences.Add(FPackageId::FromName(SoftRefName));
+		}
+	}
+	Result.SoftPackageReferences.Sort();
+
 	return Result;
 }
diff --git a/Engine/Source/Editor/UnrealEd/Private/CookOnTheFlyServer.cpp b/Engine/Source/Editor/UnrealEd/Private/CookOnTheFlyServer.cpp
--- a/Engine/Source/Editor/UnrealEd/Private/CookOnTheFlyServer.cpp
+++ b/Engine/Source/Editor/UnrealEd/Private/CookOnTheFlyServer.cpp
@@ -6269,6 +6269,11 @@ void FSaveCookedPackageContext::SetupPackage()
 		| (COTFS.IsCookFlagSet(ECookInitializationFlags::Unversioned) ? SAVE_Unversioned : 0);
 	SaveFlags |= COTFS.IsCookFlagSet(ECookInitializationFlags::CookEditorOptional) ? SAVE_Optional : SAVE_None;
 
+	if (COTFS.CookByTheBookOptions->bCookSoftPackageReferences)
+	{
+		SaveFlags |= SAVE_CookSoftPackageReferences;
+	}
+
 	// removing editor only packages only works when cooking in commandlet and non iterative cooking
 	bool bCanSkipEditorOnlyPackages = COTFS.IsCookByTheBookMode() && !COTFS.IsCookingInEditor();
 	bCanSkipEditorOnlyPackages &= !COTFS.IsCookFlagSet(ECookInitializationFlags::Iterative);
@@ -11634,6 +11639,7 @@ FBeginCookContext UCookOnTheFlyServer::CreateBeginCookByTheBookContext(const FCo
 	CookByTheBookOptions->CreateReleaseVersion = StartupOptions.CreateReleaseVersion;
 	CookByTheBookOptions->bSkipHardReferences = !!(CookOptions & ECookByTheBookOptions::SkipHardReferences);
 	CookByTheBookOptions->bSkipSoftReferences = !!(CookOptions & ECookByTheBookOptions::SkipSoftReferences);
+	CookByTheBookOptions->bCookSoftPackageReferences = FParse::Param(FCommandLine::Get(), TEXT("CookSoftPackageReferences"));
 	CookByTheBookOptions->bCookAgainstFixedBase = !!(CookOptions & ECookByTheBookOptions::CookAgainstFixedBase);
 	CookByTheBookOptions->bDlcLoadMainAssetRegistry = !!(CookOptions & ECookByTheBookOptions::DlcLoadMainAssetRegistry);
 	CookByTheBookOptions->bErrorOnEngineContentUse = StartupOptions.bErrorOnEngineContentUse;
diff --git a/Engine/Source/Editor/UnrealEd/Private/Cooker/CookTypes.cpp b/Engine/Source/Editor/UnrealEd/Private/Cooker/CookTypes.cpp
--- a/Engine/Source/Editor/UnrealEd/Private/Cooker/CookTypes.cpp
+++ b/Engine/Source/Editor/UnrealEd/Private/Cooker/CookTypes.cpp
@@ -878,6 +878,7 @@ FCbWriter& operator<<(FCbWriter& Writer, const UE::Cook::FCookByTheBookOptions&
 	Writer << "SkipSoftReferences" << Value.bSkipSoftReferences;
 	Writer << "CookAgainstFixedBase" << Value.bCookAgainstFixedBase;
 	Writer << "DlcLoadMainAssetRegistry" << Value.bDlcLoadMainAssetRegistry;
+	Writer << "CookSoftPackageReferences" << Value.bCookSoftPackageReferences;
 	Writer.EndObject();
 	return Writer;
 }
@@ -907,6 +908,7 @@ bool LoadFromCompactBinary(FCbFieldView Field, UE::Cook::FCookByTheBookOptions&
 	bOk = LoadFromCompactBinary(Field["SkipSoftReferences"], OutValue.bSkipSoftReferences) & bOk;
 	bOk = LoadFromCompactBinary(Field["CookAgainstFixedBase"], OutValue.bCookAgainstFixedBase) & bOk;
 	bOk = LoadFromCompactBinary(Field["DlcLoadMainAssetRegistry"], OutValue.bDlcLoadMainAssetRegistry) & bOk;
+	bOk = LoadFromCompactBinary(Field["CookSoftPackageReferences"], OutValue.bCookSoftPackageReferences) & bOk;
 
 	return bOk;
 }
diff --git a/Engine/Source/Editor/UnrealEd/Private/Cooker/CookTypes.h b/Engine/Source/Editor/UnrealEd/Private/Cooker/CookTypes.h
--- a/Engine/Source/Editor/UnrealEd/Private/Cooker/CookTypes.h
+++ b/Engine/Source/Editor/UnrealEd/Private/Cooker/CookTypes.h
@@ -457,6 +457,7 @@ public:
 	bool							bAllowUncookedAssetReferences = false;
 	bool							bSkipHardReferences = false;
 	bool							bSkipSoftReferences = false;
+	bool							bCookSoftPackageReferences = false;
 	bool							bCookAgainstFixedBase = false;
 	bool							bDlcLoadMainAssetRegistry = false;
 
diff --git a/Engine/Source/Runtime/Core/Private/IO/IoContainerHeader.cpp b/Engine/Source/Runtime/Core/Private/IO/IoContainerHeader.cpp
--- a/Engine/Source/Runtime/Core/Private/IO/IoContainerHeader.cpp
+++ b/Engine/Source/Runtime/Core/Private/IO/IoContainerHeader.cpp
@@ -20,6 +20,24 @@ FArchive& operator<<(FArchive& Ar, FIoContainerHeaderLocalizedPackage& Localized
 	return Ar;
 }
 
+void FIoContainerHeaderSoftPackageReferences::Empty()
+{
+	PackageIds.Empty();
+	PackageIndices.Empty();
+	bContainsSoftPackageReferences = false;
+}
+
+FArchive& operator<<(FArchive& Ar, FIoContainerHeaderSoftPackageReferences& SoftPackageReferences)
+{
+	Ar << SoftPackageReferences.bContainsSoftPackageReferences;
+	if (SoftPackageReferences.bContainsSoftPackageReferences)
+	{
+		Ar << SoftPackageReferences.PackageIds;
+		Ar << SoftPackageReferences.PackageIndices;
+	}
+	return Ar;
+}
+
 FArchive& operator<<(FArchive& Ar, FIoContainerHeader& ContainerHeader)
 {
 	uint32 Signature = FIoContainerHeader::Signature;
@@ -56,5 +74,10 @@ FArchive& operator<<(FArchive& Ar, FIoContainerHeader& ContainerHeader)
 	Ar << ContainerHeader.LocalizedPackages;
 	Ar << ContainerHeader.PackageRedirects;
 
+	if (Ar.IsSaving() || Version >= EIoContainerHeaderVersion::SoftPackageReferences)
+	{
+		Ar << ContainerHeader.SoftPackageReferences;
+	}
+
 	return Ar;
 }
diff --git a/Engine/Source/Runtime/Core/Private/IO/PackageStore.cpp b/Engine/Source/Runtime/Core/Private/IO/PackageStore.cpp
--- a/Engine/Source/Runtime/Core/Private/IO/PackageStore.cpp
+++ b/Engine/Source/Runtime/Core/Private/IO/PackageStore.cpp
@@ -22,6 +22,7 @@ FArchive& operator<<(FArchive& Ar, FPackageStoreEntryResource& PackageStoreEntry
 	Ar << PackageStoreEntry.PackageName;
 	Ar << PackageStoreEntry.ImportedPackageIds;
 	Ar << PackageStoreEntry.OptionalSegmentImportedPackageIds;
+	Ar << PackageStoreEntry.SoftPackageReferences;
 
 	if (Ar.IsLoading())
 	{
@@ -69,6 +70,16 @@ FCbWriter& operator<<(FCbWriter& Writer, const FPackageStoreEntryResource& Packa
 		Writer.EndArray();
 	}
 
+	if (PackageStoreEntry.SoftPackageReferences.Num())
+	{
+		Writer.BeginArray("softpackagereferences");
+		for (const FPackageId& SoftRef : PackageStoreEntry.SoftPackageReferences)
+		{
+			Writer << SoftRef.Value();
+		}
+		Writer.EndArray();
+	}
+
 	Writer.EndObject();
 
 	return Writer;
@@ -107,6 +118,14 @@ FPackageStoreEntryResource FPackageStoreEntryResource::FromCbObject(FCbObjectVie
 		}
 	}
 
+	if (Obj["softpackagereferences"])
+	{
+		for (FCbFieldView ArrayField : Obj["softpackagereferences"])
+		{
+			Entry.SoftPackageReferences.Add(FPackageId::FromValue(ArrayField.AsUInt64()));
+		}
+	}
+
 	return Entry;
 }
 
diff --git a/Engine/Source/Runtime/Core/Public/IO/IoContainerHeader.h b/Engine/Source/Runtime/Core/Public/IO/IoContainerHeader.h
--- a/Engine/Source/Runtime/Core/Public/IO/IoContainerHeader.h
+++ b/Engine/Source/Runtime/Core/Public/IO/IoContainerHeader.h
@@ -63,12 +63,32 @@ struct FIoContainerHeaderLocalizedPackage
 	CORE_API friend FArchive& operator<<(FArchive& Ar, FIoContainerHeaderLocalizedPackage& LocalizedPackage);
 };
 
+struct FFilePackageStoreEntrySoftReferences
+{
+	TFilePackageStoreEntryCArrayView<uint32> Indices;
+};
+
+struct FIoContainerHeaderSoftPackageReferences
+{
+	CORE_API void Empty();
+
+	// Deduplicated list of soft referenced package IDs for all packages in the container.
+	TArray<FPackageId> PackageIds;
+	// Indices into the package ID list for all packages in the container serialized as array views.
+	TArray<uint8> PackageIndices; 
+	// Flag indicating whether any soft package references exists.
+	bool bContainsSoftPackageReferences = false;
+
+	CORE_API friend FArchive& operator<<(FArchive& Ar, FIoContainerHeaderSoftPackageReferences& SoftPackageReferences);
+};
+
 enum class EIoContainerHeaderVersion : uint32
 {
 	Initial = 0,
 	LocalizedPackages = 1,
 	OptionalSegmentPackages = 2,
 	NoExportInfo = 3,
+	SoftPackageReferences = 4,
 
 	LatestPlusOne,
 	Latest = LatestPlusOne - 1
@@ -89,6 +109,7 @@ struct FIoContainerHeader
 	TArray<FDisplayNameEntryId> RedirectsNameMap;
 	TArray<FIoContainerHeaderLocalizedPackage> LocalizedPackages;
 	TArray<FIoContainerHeaderPackageRedirect> PackageRedirects;
+	FIoContainerHeaderSoftPackageReferences SoftPackageReferences;
 
 	CORE_API friend FArchive& operator<<(FArchive& Ar, FIoContainerHeader& ContainerHeader);
 };
diff --git a/Engine/Source/Runtime/Core/Public/IO/PackageStore.h b/Engine/Source/Runtime/Core/Public/IO/PackageStore.h
--- a/Engine/Source/Runtime/Core/Public/IO/PackageStore.h
+++ b/Engine/Source/Runtime/Core/Public/IO/PackageStore.h
@@ -74,6 +74,8 @@ struct FPackageStoreEntryResource
 	TArray<FSHAHash> ShaderMapHashes;
 	/** Editor data imported package IDs. */
 	TArray<FPackageId> OptionalSegmentImportedPackageIds;
+	/** Soft package references. */
+	TArray<FPackageId> SoftPackageReferences;
 
 	/** Returns the package ID. */
 	FPackageId GetPackageId() const
diff --git a/Engine/Source/Runtime/CoreUObject/Private/UObject/SavePackage2.cpp b/Engine/Source/Runtime/CoreUObject/Private/UObject/SavePackage2.cpp
--- a/Engine/Source/Runtime/CoreUObject/Private/UObject/SavePackage2.cpp
+++ b/Engine/Source/Runtime/CoreUObject/Private/UObject/SavePackage2.cpp
@@ -1906,7 +1906,8 @@ ESavePackageResult WritePackageHeader(FStructuredArchive::FRecord& StructuredArc
 	}
 
 	// Write Soft Package references & Searchable Names
-	if (!SaveContext.IsFilterEditorOnly())
+	const bool bCookSoftPackagesReferences = SaveContext.IsCooking() && SaveContext.GetSaveArgs().SaveFlags & SAVE_CookSoftPackageReferences;
+	if (!SaveContext.IsFilterEditorOnly() || bCookSoftPackagesReferences)
 	{
 		SCOPED_SAVETIMER(UPackage_Save_SaveSoftPackagesAndSearchableNames);
 
diff --git a/Engine/Source/Runtime/CoreUObject/Public/UObject/ObjectMacros.h b/Engine/Source/Runtime/CoreUObject/Public/UObject/ObjectMacros.h
--- a/Engine/Source/Runtime/CoreUObject/Public/UObject/ObjectMacros.h
+++ b/Engine/Source/Runtime/CoreUObject/Public/UObject/ObjectMacros.h
@@ -104,6 +104,7 @@ enum ESaveFlags
 	SAVE_Optional					= 0x00008000,	///< Indicate that we to save optional exports. This flag is only valid while cooking. Optional exports are filtered if not specified during cooking. 
 	SAVE_AllowTimeout				= 0x00010000,   ///< If present, if SavePackage is blocked by an asynchronous operation, it will quickly fail with ESavePackageResult::Timeout, otherwise it will wait for a while and then fail with ESavePackageResult::Error.
 	SAVE_RehydratePayloads			= 0x00020000,   ///< Any virtualized payloads in the package should be pulled and stored locally in the package file during the save
+	SAVE_CookSoftPackageReferences	= 0x00040000,	///< Indicate that we want to save soft package references for cooked packages
 	SAVE_Unversioned = SAVE_Unversioned_Native | SAVE_Unversioned_Properties,
 };
 
diff --git a/Engine/Source/Runtime/PakFile/Private/FilePackageStore.cpp b/Engine/Source/Runtime/PakFile/Private/FilePackageStore.cpp
--- a/Engine/Source/Runtime/PakFile/Private/FilePackageStore.cpp
+++ b/Engine/Source/Runtime/PakFile/Private/FilePackageStore.cpp
@@ -651,6 +651,7 @@ void FFilePackageStoreBackend::Update()
 
 				MountedContainer.NumMountedPackages = NumNewPackages;
 				ContainerHeader->PackageIds.Empty();
+				ContainerHeader->SoftPackageReferences.Empty();
 			}
 
 			
