Fixed bug when serializing empty soft package references

#rnx

diff --git a/Engine/Source/Developer/IoStoreUtilities/Private/IoStoreUtilities.cpp b/Engine/Source/Developer/IoStoreUtilities/Private/IoStoreUtilities.cpp
--- a/Engine/Source/Developer/IoStoreUtilities/Private/IoStoreUtilities.cpp
+++ b/Engine/Source/Developer/IoStoreUtilities/Private/IoStoreUtilities.cpp
@@ -5293,6 +5293,13 @@ void CreateContainerHeader(FContainerTargetSpec& ContainerTarget, bool bIsOption
 
 		void Append(TConstArrayView<FPackageId> SoftRefs)
 		{
+			if (SoftPackageReferences.IsEmpty())
+			{
+				// Skip serialization when there's no soft references for any package
+				check(SoftRefs.IsEmpty());
+				return;
+			}
+
 			const int64 RemainingEntrySize = TotalEntrySize - EntryAr.Tell();
 			const int64 OffsetFromThis = RemainingEntrySize + DataAr.Tell();
 			uint32		ArrayNum = static_cast<uint32>(SoftRefs.Num());
