Ensure DLC cook and pak via zenstore don't interfere with the base project cooked output oplogs by segregating their cooked data into its own project id.  Tested using EngineTestDLC.  EditorDomain oplog is allowed to be shared between the DLC and the base project.

diff --git a/Engine/Source/Developer/IoStoreUtilities/Internal/ZenStoreHttpClient.h b/Engine/Source/Developer/IoStoreUtilities/Internal/ZenStoreHttpClient.h
--- a/Engine/Source/Developer/IoStoreUtilities/Internal/ZenStoreHttpClient.h
+++ b/Engine/Source/Developer/IoStoreUtilities/Internal/ZenStoreHttpClient.h
@@ -34,7 +34,7 @@ public:
 	FZenStoreHttpClient(UE::Zen::FServiceSettings&& InSettings);
 	~FZenStoreHttpClient();
 
-	bool TryCreateProject(FStringView InProjectId, FStringView InOplogId, FStringView ServerRoot, 
+	bool TryCreateProject(FStringView InProjectId, FStringView InParentProjectId, FStringView InOplogId, FStringView ServerRoot, 
 					FStringView EngineRoot, FStringView ProjectRoot,
 					FStringView ProjectFilePath);
 	bool TryCreateOplog(FStringView InProjectId, FStringView InOplogId, FStringView InOplogLifetimeMarkerPath, bool bFullBuild);
diff --git a/Engine/Source/Developer/IoStoreUtilities/Private/ZenStoreHttpClient.cpp b/Engine/Source/Developer/IoStoreUtilities/Private/ZenStoreHttpClient.cpp
--- a/Engine/Source/Developer/IoStoreUtilities/Private/ZenStoreHttpClient.cpp
+++ b/Engine/Source/Developer/IoStoreUtilities/Private/ZenStoreHttpClient.cpp
@@ -61,8 +61,9 @@ FZenStoreHttpClient::~FZenStoreHttpClient()
 }
 
 bool
-FZenStoreHttpClient::TryCreateProject(FStringView InProjectId, 
-	FStringView InOplogId, 
+FZenStoreHttpClient::TryCreateProject(FStringView InProjectId,
+	FStringView InParentProjectId,
+	FStringView InOplogId,
 	FStringView ServerRoot,
 	FStringView EngineRoot,
 	FStringView ProjectRoot,
@@ -97,6 +98,10 @@ FZenStoreHttpClient::TryCreateProject(FStringView InProjectId,
 			FCbWriter ProjInfo;
 			ProjInfo.BeginObject();
 			ProjInfo << "id" << InProjectId;
+			if (!InParentProjectId.IsEmpty())
+			{
+				ProjInfo << "parentid" << InParentProjectId;
+			}
 			ProjInfo << "root" << ServerRoot;
 			ProjInfo << "engine" << EngineRoot;
 			ProjInfo << "project" << ProjectRoot;
diff --git a/Engine/Source/Developer/IoStoreUtilities/Private/ZenStoreWriter.cpp b/Engine/Source/Developer/IoStoreUtilities/Private/ZenStoreWriter.cpp
--- a/Engine/Source/Developer/IoStoreUtilities/Private/ZenStoreWriter.cpp
+++ b/Engine/Source/Developer/IoStoreUtilities/Private/ZenStoreWriter.cpp
@@ -175,7 +175,15 @@ FZenStoreWriter::FZenStoreWriter(
 {
 	StaticInit();
 
-	ProjectId = FApp::GetZenStoreProjectId();
+	FString DLCName;
+	FParse::Value(FCommandLine::Get(), TEXT("DLCNAME="), DLCName);
+	DLCName.ToLowerInline();
+	ProjectId = FApp::GetZenStoreProjectId(DLCName);
+	FString ParentProjectId;
+	if (!DLCName.IsEmpty())
+	{
+		ParentProjectId = FApp::GetZenStoreProjectId();
+	}
 	
 	if (FParse::Value(FCommandLine::Get(), TEXT("-ZenStorePlatform="), OplogId) == false)
 	{
@@ -202,7 +210,7 @@ FZenStoreWriter::FZenStoreWriter(
 	FString AbsProjectDir = PlatformFile.ConvertToAbsolutePathForExternalAppForRead(*ProjectDir);
 	FString ProjectFilePath = PlatformFile.ConvertToAbsolutePathForExternalAppForRead(*ProjectPath);
 
-	HttpClient->TryCreateProject(ProjectId, OplogId, AbsServerRoot, AbsEngineDir, AbsProjectDir, IsLocalConnection ? ProjectFilePath : FStringView());
+	HttpClient->TryCreateProject(ProjectId, ParentProjectId, OplogId, AbsServerRoot, AbsEngineDir, AbsProjectDir, IsLocalConnection ? ProjectFilePath : FStringView());
 
 	PackageStoreOptimizer->Initialize();
 
diff --git a/Engine/Source/Editor/UnrealEd/Private/TargetDomain/TargetDomainUtils.cpp b/Engine/Source/Editor/UnrealEd/Private/TargetDomain/TargetDomainUtils.cpp
--- a/Engine/Source/Editor/UnrealEd/Private/TargetDomain/TargetDomainUtils.cpp
+++ b/Engine/Source/Editor/UnrealEd/Private/TargetDomain/TargetDomainUtils.cpp
@@ -1019,7 +1019,7 @@ FEditorDomainOplog::FEditorDomainOplog()
 	if (UE::Zen::IsDefaultServicePresent())
 	{
 		bool IsLocalConnection = HttpClient.GetZenServiceInstance().IsServiceRunningLocally();
-		HttpClient.TryCreateProject(ProjectId, OplogId, AbsServerRoot, AbsEngineDir, AbsProjectDir, IsLocalConnection ? ProjectFilePath : FStringView());
+		HttpClient.TryCreateProject(ProjectId, FStringView(), OplogId, AbsServerRoot, AbsEngineDir, AbsProjectDir, IsLocalConnection ? ProjectFilePath : FStringView());
 		HttpClient.TryCreateOplog(ProjectId, OplogId, TEXT("") /*InOplogMarkerFile*/, false /* bFullBuild */);
 	}
 #endif
diff --git a/Engine/Source/Runtime/Core/Private/Misc/App.cpp b/Engine/Source/Runtime/Core/Private/Misc/App.cpp
--- a/Engine/Source/Runtime/Core/Private/Misc/App.cpp
+++ b/Engine/Source/Runtime/Core/Private/Misc/App.cpp
@@ -442,7 +442,7 @@ void FApp::PrintStartupLogMessages()
 	FDevVersionRegistration::DumpVersionsToLog();
 }
 
-FString FApp::GetZenStoreProjectId()
+FString FApp::GetZenStoreProjectId(FStringView SubProject)
 {
 	FString ProjectId;
 	if (FParse::Value(FCommandLine::Get(), TEXT("-ZenStoreProject="), ProjectId))
@@ -459,7 +459,15 @@ FString FApp::GetZenStoreProjectId()
 		FTCHARToUTF8 AbsProjectFilePathUTF8(*AbsProjectFilePath);
 
 		FString HashString = FMD5::HashBytes((unsigned char*)AbsProjectFilePathUTF8.Get(), AbsProjectFilePathUTF8.Length()).Left(8);
-		return FString::Printf(TEXT("%s.%.8s"), FApp::GetProjectName(), *HashString);
+
+		if (SubProject.IsEmpty())
+		{
+			return FString::Printf(TEXT("%s.%.8s"), FApp::GetProjectName(), *HashString);
+		}
+		else
+		{
+			return FString::Printf(TEXT("%s.%.*s.%.8s"), FApp::GetProjectName(), SubProject.Len(), SubProject.GetData(), *HashString);
+		}
 	}
 	UE_LOG(LogInit, Fatal, TEXT("GetZenStoreProjectId() called before having a valid project file path"));
 #else
diff --git a/Engine/Source/Runtime/Core/Public/Misc/App.h b/Engine/Source/Runtime/Core/Public/Misc/App.h
--- a/Engine/Source/Runtime/Core/Public/Misc/App.h
+++ b/Engine/Source/Runtime/Core/Public/Misc/App.h
@@ -240,7 +240,7 @@ public:
 	 *
 	 * @return Zen store project id.
 	 */
-	static CORE_API FString GetZenStoreProjectId();
+	static CORE_API FString GetZenStoreProjectId(FStringView SubProject = FStringView());
 
 	/**
 	 * Gets the globally unique identifier of this application instance.
