UnrealPak: Fix division by zero in logging code

#rnx

diff --git a/Engine/Source/Developer/IoStoreUtilities/Private/IoStoreUtilities.cpp b/Engine/Source/Developer/IoStoreUtilities/Private/IoStoreUtilities.cpp
--- a/Engine/Source/Developer/IoStoreUtilities/Private/IoStoreUtilities.cpp
+++ b/Engine/Source/Developer/IoStoreUtilities/Private/IoStoreUtilities.cpp
@@ -5933,8 +5933,8 @@ int32 CreateTarget(const FIoStoreArguments& Arguments, const FIoStoreWriterSetti
 
 		FIoStoreWriterContext::FProgress Progress = IoStoreWriterContext->GetProgress();
 		const uint64 TotalDDCAttempts = Progress.CompressionDDCHitCount + Progress.CompressionDDCMissCount;
-		const double DDCHitRate = double(Progress.CompressionDDCHitCount) / TotalDDCAttempts * 100.0;
-		const double DDCPutRate = double(Progress.CompressionDDCPutCount) / TotalDDCAttempts * 100.0;
+		const double DDCHitRate = TotalDDCAttempts > 0 ? double(Progress.CompressionDDCHitCount) / TotalDDCAttempts * 100.0 : 0.0;
+		const double DDCPutRate = TotalDDCAttempts > 0 ? double(Progress.CompressionDDCPutCount) / TotalDDCAttempts * 100.0 : 0.0;
 		UE_LOG(LogIoStore, Display, TEXT("Compression DDC hits: %llu/%llu (%.2lf%%)"),
 			Progress.CompressionDDCHitCount, TotalDDCAttempts, DDCHitRate);
 		UE_LOG(LogIoStore, Display, TEXT("Compression DDC puts: %llu/%llu (%.2lf%%) [%llu failed]"),
