UnrealPak: Remove the stricter ddc put limits and use the same as for gets.

#rnx

diff --git a/Engine/Source/Developer/IoStoreUtilities/Private/IoStoreWriter.cpp b/Engine/Source/Developer/IoStoreUtilities/Private/IoStoreWriter.cpp
--- a/Engine/Source/Developer/IoStoreUtilities/Private/IoStoreWriter.cpp
+++ b/Engine/Source/Developer/IoStoreUtilities/Private/IoStoreWriter.cpp
@@ -2292,8 +2292,6 @@ void FIoStoreWriterContextImpl::BeginEncryptionAndSigningThreadFunc()
 
 	FIoStoreDDCRequestDispatcherParams PutRequestDispatcherParams;
 	PutRequestDispatcherParams.QueueTimeLimitMs = 1000.f;
-	PutRequestDispatcherParams.MaxInflightBytes = 256ull << 20;
-	PutRequestDispatcherParams.MaxBatchBytes = 1ull << 20;
 	FIoStoreDDCPutRequestDispatcher DDCPutRequestDispatcher(PutRequestDispatcherParams);
 
 	for (;;)
