UnrealPak: Fix seen once crash in the DDC FOnCacheGetValueComplete callback provided by FIoStoreDDCGetRequestDispatcher::DispatchGetRequests.
The HandleDDCGetResult callback was releasing the FIoStoreWriteQueueEntry so that it could be completed and deleted before it was used to update the dispatcher queue size tracking.

#rnx

diff --git a/Engine/Source/Developer/IoStoreUtilities/Private/IoStoreWriter.cpp b/Engine/Source/Developer/IoStoreUtilities/Private/IoStoreWriter.cpp
--- a/Engine/Source/Developer/IoStoreUtilities/Private/IoStoreWriter.cpp
+++ b/Engine/Source/Developer/IoStoreUtilities/Private/IoStoreWriter.cpp
@@ -367,14 +367,19 @@ void FIoStoreDDCGetRequestDispatcher::DispatchGetRequests(
 		GetCache().GetValue(RequestQueue.Requests, RequestQueue.RequestOwner, [this, Callback](FCacheGetValueResponse&& Response)
 		{
 			TRACE_CPUPROFILER_EVENT_SCOPE(ReadFromDDC_Decompress);
-			FIoStoreWriteQueueEntry* Entry = reinterpret_cast<FIoStoreWriteQueueEntry*>(Response.UserData);
-			FSharedBuffer Result;
-			if (Response.Status == EStatus::Ok)
+			uint64 SourceBufferSizeEstimate = 0;
 			{
-				Result = Response.Value.GetData().Decompress();
+				FIoStoreWriteQueueEntry* Entry = reinterpret_cast<FIoStoreWriteQueueEntry*>(Response.UserData);
+				SourceBufferSizeEstimate = Entry->Request->GetSourceBufferSizeEstimate();
+
+				FSharedBuffer Result;
+				if (Response.Status == EStatus::Ok)
+				{
+					Result = Response.Value.GetData().Decompress();
+				}
+				Callback(Entry, Result); // Entry could be deleted after this call
 			}
-			Callback(Entry, Result);
-			RequestQueue.OnRequestComplete(Entry->Request->GetSourceBufferSizeEstimate());
+			RequestQueue.OnRequestComplete(SourceBufferSizeEstimate);
 			TRACE_COUNTER_DECREMENT(IoStoreDDCGetInflightCount);
 		});
 	}
@@ -414,10 +419,14 @@ void FIoStoreDDCPutRequestDispatcher::DispatchPutRequests(
 		FRequestBarrier RequestBarrier(RequestQueue.RequestOwner);
 		GetCache().PutValue(RequestQueue.Requests, RequestQueue.RequestOwner, [this, Callback](FCachePutValueResponse&& Response)
 		{
-			FIoStoreWriteQueueEntry* Entry = reinterpret_cast<FIoStoreWriteQueueEntry*>(Response.UserData);
-			bool bSuccess = Response.Status == EStatus::Ok;
-			Callback(Entry, bSuccess);
-			RequestQueue.OnRequestComplete(Entry->CompressedSize);
+			uint64 CompressedSize = 0;
+			{
+				FIoStoreWriteQueueEntry* Entry = reinterpret_cast<FIoStoreWriteQueueEntry*>(Response.UserData);
+				CompressedSize = Entry->CompressedSize;
+				bool bSuccess = Response.Status == EStatus::Ok;
+				Callback(Entry, bSuccess); // Entry could be deleted after this call
+			}
+			RequestQueue.OnRequestComplete(CompressedSize);
 			TRACE_COUNTER_DECREMENT(IoStoreDDCPutInflightCount);
 		});
 	}
