Extract CookedPackageStore and CookMetadataFiles functionality from IoStoreUtilities.cpp into Internal headers that allow use from IoStoreUtilities as well as PakFileUtilities.  Allows Pak commands to accept a projectstore file argument and make use of data stored in zenserver.

diff --git a/Engine/Source/Developer/IoStoreUtilities/Internal/CookedPackageStore.h b/Engine/Source/Developer/IoStoreUtilities/Internal/CookedPackageStore.h
new file mode 100644
--- /dev/null
+++ b/Engine/Source/Developer/IoStoreUtilities/Internal/CookedPackageStore.h
@@ -0,0 +1,55 @@
+// Copyright Epic Games, Inc. All Rights Reserved.
+
+#pragma once
+
+#include "Containers/Array.h"
+#include "Containers/Map.h"
+#include "Containers/UnrealString.h"
+#include "IO/IoBuffer.h"
+#include "IO/IoChunkId.h"
+#include "IO/IoHash.h"
+#include "IO/IoStatus.h"
+#include "IO/PackageStore.h"
+#include "Serialization/FileRegions.h"
+#include "Tasks/Task.h"
+#include "UObject/NameTypes.h"
+
+namespace UE { class FZenStoreHttpClient; }
+
+class FCookedPackageStore
+{
+public:
+	struct FChunkInfo
+	{
+		FIoChunkId ChunkId;
+		FIoHash ChunkHash;
+		uint64 ChunkSize = 0;
+		FName PackageName;
+		FString RelativeFileName;
+		TArray<FFileRegion> FileRegions;
+	};
+
+	IOSTOREUTILITIES_API FCookedPackageStore(FStringView InCookedDir);
+
+	FIoStatus LoadManifest(const TCHAR* ManifestFilename);
+	IOSTOREUTILITIES_API FIoStatus LoadProjectStore(const TCHAR* ProjectStoreFilename);
+	IOSTOREUTILITIES_API void ParseOplog(FCbObject& OplogObject);
+	IOSTOREUTILITIES_API FIoStatus LoadChunkHashes();
+	IOSTOREUTILITIES_API FIoChunkId GetChunkIdFromFileName(const FString& Filename) const;
+	IOSTOREUTILITIES_API const FChunkInfo* GetChunkInfoFromChunkId(const FIoChunkId& ChunkId) const;
+	IOSTOREUTILITIES_API const FChunkInfo* GetChunkInfoFromFileName(const FString& Filename) const;
+	IOSTOREUTILITIES_API FString GetRelativeFilenameFromChunkId(const FIoChunkId& ChunkId) const;
+	IOSTOREUTILITIES_API FName GetPackageNameFromChunkId(const FIoChunkId& ChunkId) const;
+	IOSTOREUTILITIES_API FName GetPackageNameFromFileName(const FString& Filename) const;
+	IOSTOREUTILITIES_API const FPackageStoreEntryResource* GetPackageStoreEntry(FPackageId PackageId) const;
+	IOSTOREUTILITIES_API bool HasZenStoreClient() const;
+	IOSTOREUTILITIES_API TIoStatusOr<FIoBuffer> ReadChunk(const FIoChunkId& ChunkId);
+	IOSTOREUTILITIES_API UE::Tasks::FTask ReadChunkAsync(const FIoChunkId& ChunkId, TFunction<void(TIoStatusOr<FIoBuffer>)>&& Callback);
+
+private:
+	TUniquePtr<UE::FZenStoreHttpClient> ZenStoreClient;
+	FString CookedDir;
+	TMap<FPackageId, FPackageStoreEntryResource> PackageIdToEntry;
+	TMap<FString, FIoChunkId> FilenameToChunkIdMap;
+	TMap<FIoChunkId, FChunkInfo> ChunkInfoMap;
+};
diff --git a/Engine/Source/Developer/IoStoreUtilities/Private/CookMetadataFiles.cpp b/Engine/Source/Developer/IoStoreUtilities/Private/CookMetadataFiles.cpp
new file mode 100644
--- /dev/null
+++ b/Engine/Source/Developer/IoStoreUtilities/Private/CookMetadataFiles.cpp
@@ -0,0 +1,148 @@
+// Copyright Epic Games, Inc. All Rights Reserved.
+
+#include "CookMetadataFiles.h"
+
+#include "AssetRegistry/AssetRegistryState.h"
+#include "Containers/Array.h"
+#include "Containers/UnrealString.h"
+#include "CookMetadata.h"
+#include "HAL/FileManager.h"
+#include "IO/IoStore.h"
+#include "Misc/Paths.h"
+#include "ProfilingDebugging/CountersTrace.h"
+#include "Serialization/LargeMemoryReader.h"
+
+// Returns the hash of the development asset registry or 0 on failure.
+static uint64 LoadAssetRegistry(const FString& InAssetRegistryFileName, FAssetRegistryState& OutAssetRegistry)
+{
+	FAssetRegistryVersion::Type Version;
+	FAssetRegistryLoadOptions Options(UE::AssetRegistry::ESerializationTarget::ForDevelopment);
+
+	TUniquePtr<FArchive> FileReader(IFileManager::Get().CreateFileReader(*InAssetRegistryFileName));
+	if (FileReader)
+	{
+		TArray64<uint8> Data;
+		Data.SetNumUninitialized(FileReader->TotalSize());
+		FileReader->Serialize(Data.GetData(), Data.Num());
+		check(!FileReader->IsError());
+
+		uint64 DevArHash = UE::Cook::FCookMetadataState::ComputeHashOfDevelopmentAssetRegistry(MakeMemoryView(Data));
+
+		FLargeMemoryReader MemoryReader(Data.GetData(), Data.Num());
+		if (OutAssetRegistry.Load(MemoryReader, Options, &Version))
+		{
+			return DevArHash;
+		}
+	}
+
+	return 0;;
+}
+
+ECookMetadataFiles FindAndLoadMetadataFiles(
+	const FString& InCookedDir, ECookMetadataFiles InRequiredFiles, 
+	FAssetRegistryState& OutAssetRegistry, FString* OutAssetRegistryFileName /*optional, set on success*/,
+	UE::Cook::FCookMetadataState* OutCookMetadata, FString* OutCookMetadataFileName /*optional, set on success or need*/)
+{
+	TRACE_CPUPROFILER_EVENT_SCOPE(LoadingAssetRegistry);
+
+	// Look for the development registry. Should be in \\GameName\\Metadata\\DevelopmentAssetRegistry.bin, but we don't know what "GameName" is.
+	TArray<FString> PossibleAssetRegistryFiles;
+	IFileManager::Get().FindFilesRecursive(PossibleAssetRegistryFiles, *InCookedDir, GetDevelopmentAssetRegistryFilename(), true, false);
+
+	if (PossibleAssetRegistryFiles.Num() > 1)
+	{
+		UE_LOG(LogIoStore, Warning, TEXT("Found multiple possible development asset registries:"));
+		for (FString& Filename : PossibleAssetRegistryFiles)
+		{
+			UE_LOG(LogIoStore, Warning, TEXT("    %s"), *Filename);
+		}
+	}
+
+	if (PossibleAssetRegistryFiles.Num() == 0)
+	{
+		if (EnumHasAnyFlags(InRequiredFiles, ECookMetadataFiles::AssetRegistry))
+		{
+			UE_LOG(LogIoStore, Error, TEXT("No development asset registry file found!"));
+		}
+		else
+		{
+			UE_LOG(LogIoStore, Display, TEXT("No development asset registry file found!"));
+		}
+		return ECookMetadataFiles::None;
+	}
+
+	UE_LOG(LogIoStore, Display, TEXT("Using input asset registry: %s"), *PossibleAssetRegistryFiles[0]);
+	uint64 LoadedDevArHash = LoadAssetRegistry(PossibleAssetRegistryFiles[0], OutAssetRegistry);
+
+	if (LoadedDevArHash == 0)
+	{
+		return ECookMetadataFiles::None; // already logged
+	}
+
+	// If we found the asset registry, try and find the cook metadata that should be next to it.
+	ECookMetadataFiles ResultFiles = ECookMetadataFiles::AssetRegistry;
+
+	if (OutCookMetadata)
+	{
+		// The cook metadata file should be adjacent to the development asset registry.
+		FString CookMetadataFileName = FPaths::GetPath(PossibleAssetRegistryFiles[0]) / UE::Cook::GetCookMetadataFilename();
+		if (IFileManager::Get().FileExists(*CookMetadataFileName))
+		{
+			if (OutCookMetadata->ReadFromFile(CookMetadataFileName) == false)
+			{
+				UE_LOG(LogIoStore, Error, TEXT("Failed to deserialize cook metadata file - invalid data. [%s]"), *CookMetadataFileName);
+				if (EnumHasAnyFlags(InRequiredFiles, ECookMetadataFiles::CookMetadata))
+				{
+					return ECookMetadataFiles::None;
+				}
+			}
+			else if (OutCookMetadata->GetAssociatedDevelopmentAssetRegistryHash() != LoadedDevArHash &&
+				OutCookMetadata->GetAssociatedDevelopmentAssetRegistryHashPostWriteback() != LoadedDevArHash) // during testing we can repeat stage after cook so we might have already edited it.
+			{
+				if (EnumHasAnyFlags(InRequiredFiles, ECookMetadataFiles::CookMetadata))
+				{
+					UE_LOG(LogIoStore, Error,
+						TEXT("Cook metadata file mismatch: Hash of associated development asset registry does not match. [%s] %llx vs %llx (%llx post writeback)"),
+						*CookMetadataFileName, LoadedDevArHash, OutCookMetadata->GetAssociatedDevelopmentAssetRegistryHash(), OutCookMetadata->GetAssociatedDevelopmentAssetRegistryHashPostWriteback());
+					return ECookMetadataFiles::None;
+				}
+				else
+				{
+					UE_LOG(LogIoStore, Display,
+						TEXT("Cook metadata file mismatch: Hash of associated development asset registry does not match. [%s] %llx vs %llx (%llx post writeback)"),
+						*CookMetadataFileName, LoadedDevArHash, OutCookMetadata->GetAssociatedDevelopmentAssetRegistryHash(), OutCookMetadata->GetAssociatedDevelopmentAssetRegistryHashPostWriteback());
+					OutCookMetadata->Reset();
+				}
+			}
+			else
+			{
+				EnumAddFlags(ResultFiles, ECookMetadataFiles::CookMetadata);
+				if (OutCookMetadataFileName)
+				{
+					*OutCookMetadataFileName = MoveTemp(CookMetadataFileName);
+				}
+			}
+		}
+		else
+		{
+			if (EnumHasAnyFlags(InRequiredFiles, ECookMetadataFiles::CookMetadata))
+			{
+				UE_LOG(LogIoStore, Error, TEXT("Failed to open and read cook metadata file %s"), *CookMetadataFileName);
+				return ECookMetadataFiles::None;
+			}
+
+			UE_LOG(LogIoStore, Display, TEXT("No cook metadata file found, checked %s"), *CookMetadataFileName);
+			if (OutCookMetadataFileName)
+			{
+				*OutCookMetadataFileName = FString("");
+			}
+		}
+	}
+
+
+	if (OutAssetRegistryFileName)
+	{
+		*OutAssetRegistryFileName = MoveTemp(PossibleAssetRegistryFiles[0]);
+	}
+	return ResultFiles;
+}
\ No newline at end of file
diff --git a/Engine/Source/Developer/IoStoreUtilities/Private/CookMetadataFiles.h b/Engine/Source/Developer/IoStoreUtilities/Private/CookMetadataFiles.h
new file mode 100644
--- /dev/null
+++ b/Engine/Source/Developer/IoStoreUtilities/Private/CookMetadataFiles.h
@@ -0,0 +1,23 @@
+// Copyright Epic Games, Inc. All Rights Reserved.
+
+#pragma once
+
+#include "Misc/EnumClassFlags.h"
+
+class FAssetRegistryState;
+class FString;
+namespace UE::Cook { class FCookMetadataState; }
+
+enum class ECookMetadataFiles
+{
+	None = 0,
+	AssetRegistry = 1,
+	CookMetadata = 2,
+	All = 4
+};
+ENUM_CLASS_FLAGS(ECookMetadataFiles);
+
+ECookMetadataFiles FindAndLoadMetadataFiles(
+	const FString& InCookedDir, ECookMetadataFiles InRequiredFiles, 
+	FAssetRegistryState& OutAssetRegistry, FString* OutAssetRegistryFileName /*optional, set on success*/,
+	UE::Cook::FCookMetadataState* OutCookMetadata, FString* OutCookMetadataFileName /*optional, set on success or need*/);
diff --git a/Engine/Source/Developer/IoStoreUtilities/Private/CookedPackageStore.cpp b/Engine/Source/Developer/IoStoreUtilities/Private/CookedPackageStore.cpp
new file mode 100644
--- /dev/null
+++ b/Engine/Source/Developer/IoStoreUtilities/Private/CookedPackageStore.cpp
@@ -0,0 +1,367 @@
+// Copyright Epic Games, Inc. All Rights Reserved.
+
+#include "CookedPackageStore.h"
+
+#include "AssetRegistry/AssetRegistryState.h"
+#include "CookMetadataFiles.h"
+#include "HAL/FileManager.h"
+#include "IO/IoStore.h"
+#include "Misc/Paths.h"
+#include "Misc/PathViews.h"
+#include "ProfilingDebugging/CountersTrace.h"
+#include "Serialization/CompactBinarySerialization.h"
+#include "Serialization/JsonReader.h"
+#include "Serialization/JsonSerializer.h"
+#include "ZenStoreHttpClient.h"
+
+#define COOKEDPACKAGESTORE_CPU_SCOPE(NAME) TRACE_CPUPROFILER_EVENT_SCOPE(CookedPackageStore##NAME);
+
+FCookedPackageStore::FCookedPackageStore(FStringView InCookedDir)
+	: CookedDir(InCookedDir)
+{
+}
+
+FIoStatus FCookedPackageStore::LoadManifest(const TCHAR* ManifestFilename)
+{
+	COOKEDPACKAGESTORE_CPU_SCOPE(LoadCookedPackageStoreManifest);
+	double StartTime = FPlatformTime::Seconds();
+
+	TUniquePtr<FArchive> Ar(IFileManager::Get().CreateFileReader(ManifestFilename));
+	if (!Ar)
+	{
+		return FIoStatus(EIoErrorCode::NotFound);
+	}
+	FCbObject ManifestObject = LoadCompactBinary(*Ar).AsObject();
+	FCbObject OplogObject;
+	if (FCbFieldView ZenServerField = ManifestObject["zenserver"])
+	{
+		UE::Zen::FServiceSettings ZenServiceSettings;
+		ZenServiceSettings.ReadFromCompactBinary(ZenServerField["settings"]);
+		FString ProjectId = FString(ZenServerField["projectid"].AsString());
+		FString OplogId = FString(ZenServerField["oplogid"].AsString());
+
+		ZenStoreClient = MakeUnique<UE::FZenStoreHttpClient>(MoveTemp(ZenServiceSettings));
+		ZenStoreClient->InitializeReadOnly(ProjectId, OplogId);
+
+		COOKEDPACKAGESTORE_CPU_SCOPE(FetchOplog);
+		TIoStatusOr<FCbObject> OplogStatus = ZenStoreClient->GetOplog().Get();
+		if (!OplogStatus.IsOk())
+		{
+			return OplogStatus.Status();
+		}
+		OplogObject = OplogStatus.ConsumeValueOrDie();
+		
+	}
+	else
+	{
+		OplogObject = ManifestObject["oplog"].AsObject();
+	}
+	UE_LOG(LogIoStore, Display, TEXT("Fetched %d oplog items from %s in %.2lf seconds"),
+		OplogObject["entries"].AsArrayView().Num(),
+		HasZenStoreClient() ? TEXT("Zen") : TEXT("Manifest"),
+		FPlatformTime::Seconds() - StartTime);
+
+	ParseOplog(OplogObject);
+	LoadChunkHashes();
+
+	return FIoStatus::Ok;
+}
+
+FIoStatus FCookedPackageStore::LoadProjectStore(const TCHAR* ProjectStoreFilename)
+{
+	COOKEDPACKAGESTORE_CPU_SCOPE(LoadCookedProjectStore);
+	double StartTime = FPlatformTime::Seconds();
+
+	TUniquePtr<FArchive> Ar(IFileManager::Get().CreateFileReader(ProjectStoreFilename));
+	if (!Ar)
+	{
+		return FIoStatus(EIoErrorCode::NotFound);
+	}
+
+	TSharedPtr<FJsonObject> ProjectStoreObject;
+	TSharedRef<TJsonReader<UTF8CHAR>> Reader = TJsonReaderFactory<UTF8CHAR>::Create(Ar.Get());
+	if (FJsonSerializer::Deserialize(Reader, ProjectStoreObject) && ProjectStoreObject.IsValid())
+	{
+		const TSharedPtr<FJsonObject>* ZenServerObjectPtr = nullptr;
+		if (ProjectStoreObject->TryGetObjectField(TEXT("zenserver"), ZenServerObjectPtr) && (ZenServerObjectPtr != nullptr))
+		{
+			FString ProjectId;
+			FString OplogId;
+
+			const TSharedPtr<FJsonObject>& ZenServerObject = *ZenServerObjectPtr;
+			if (ZenServerObject->TryGetStringField(TEXT("projectid"), ProjectId) && !ProjectId.IsEmpty() &&
+				ZenServerObject->TryGetStringField(TEXT("oplogid"), OplogId) && !OplogId.IsEmpty())
+			{
+				ZenStoreClient = MakeUnique<UE::FZenStoreHttpClient>();
+				ZenStoreClient->InitializeReadOnly(ProjectId, OplogId);
+
+				COOKEDPACKAGESTORE_CPU_SCOPE(FetchOplog);
+				TIoStatusOr<FCbObject> OplogStatus = ZenStoreClient->GetOplog().Get();
+				if (!OplogStatus.IsOk())
+				{
+					return OplogStatus.Status();
+				}
+				FCbObject OplogObject = OplogStatus.ConsumeValueOrDie();
+
+				UE_LOG(LogIoStore, Display, TEXT("Fetched %d oplog items from Zen in %.2lf seconds"),
+					OplogObject["entries"].AsArrayView().Num(),
+					FPlatformTime::Seconds() - StartTime);
+
+				ParseOplog(OplogObject);
+				LoadChunkHashes();
+
+				return FIoStatus::Ok;
+			}
+		}
+	}
+
+	return FIoStatus(EIoErrorCode::NotFound);
+}
+
+void FCookedPackageStore::ParseOplog(FCbObject& OplogObject)
+{
+	COOKEDPACKAGESTORE_CPU_SCOPE(ParseOplog);
+	double StartTime = FPlatformTime::Seconds();
+
+	const FCbArrayView EntriesArray = OplogObject["entries"].AsArrayView();
+	const int32 EstimatedChunksCount = 3 * EntriesArray.Num(); //PackageData+BulkData+OptionalBulkData
+	ChunkInfoMap.Reserve(EstimatedChunksCount);
+	FilenameToChunkIdMap.Reserve(EstimatedChunksCount);
+	PackageIdToEntry.Reserve(EntriesArray.Num());
+	for (FCbFieldView& OplogEntry : EntriesArray)
+	{
+		FCbObjectView OplogObj = OplogEntry.AsObjectView();
+		FPackageStoreEntryResource PackageStoreEntry = FPackageStoreEntryResource::FromCbObject(OplogObj["packagestoreentry"].AsObjectView());
+
+		auto AddChunksFromOplog = [this, &OplogEntry, &PackageStoreEntry](const char* Field)
+		{
+			for (FCbFieldView& ChunkEntry : OplogEntry[Field].AsArrayView())
+			{
+				FCbObjectView ChunkObj = ChunkEntry.AsObjectView();
+				FIoChunkId ChunkId;
+				ChunkId.Set(ChunkObj["id"].AsObjectId().GetView());
+				if (ChunkId.IsValid())
+				{
+					FChunkInfo& ChunkInfo = ChunkInfoMap.Add(ChunkId);
+					ChunkInfo.ChunkId = ChunkId;
+					ChunkInfo.PackageName = PackageStoreEntry.PackageName;
+					if (ChunkObj["filename"])
+					{
+						TStringBuilder<1024> RelativeFilename;
+						RelativeFilename.Append(ChunkObj["filename"].AsString());
+						ChunkInfo.RelativeFileName = RelativeFilename;
+						TStringBuilder<1024> PathBuilder;
+						FPathViews::AppendPath(PathBuilder, CookedDir);
+						FPathViews::AppendPath(PathBuilder, RelativeFilename);
+						FPathViews::NormalizeFilename(PathBuilder);
+						FilenameToChunkIdMap.Add(*PathBuilder, ChunkId);
+					}
+					else if (ChunkObj["clientpath"])
+					{
+						TStringBuilder<1024> RelativeFilename;
+						FUtf8StringView ClientPathView = ChunkObj["clientpath"].AsString();
+						if (ClientPathView.StartsWith("/{project}/"))
+						{
+							if (!FPaths::IsProjectFilePathSet())
+							{
+								RelativeFilename.Append(ClientPathView);
+								UE_LOG(LogIoStore, Warning, TEXT("Project relative path could not be remapped because project file path is unset (possibly due to not specifying uproject path as first argument). %s"), RelativeFilename.ToString());
+							}
+							else
+							{
+								RelativeFilename.Append(FPathViews::GetBaseFilename(FPaths::GetProjectFilePath()));
+								RelativeFilename.AppendChar('/');
+								RelativeFilename.Append(ClientPathView.RightChop(11));
+							}
+						}
+						else if (ClientPathView.StartsWith("/{engine}/"))
+						{
+							RelativeFilename.Append(TEXT("Engine/"));
+							RelativeFilename.Append(ClientPathView.RightChop(10));
+						}
+						else
+						{
+							RelativeFilename.Append(ClientPathView);
+						}
+						ChunkInfo.RelativeFileName = RelativeFilename;
+						TStringBuilder<1024> PathBuilder;
+						FPathViews::AppendPath(PathBuilder, CookedDir);
+						FPathViews::AppendPath(PathBuilder, RelativeFilename);
+						FPathViews::NormalizeFilename(PathBuilder);
+						FilenameToChunkIdMap.Add(*PathBuilder, ChunkId);
+					}
+
+					const FCbArrayView RegionsArray = ChunkObj["fileregions"].AsArrayView();
+					ChunkInfo.FileRegions.Reserve(RegionsArray.Num());
+					for (FCbFieldView RegionObj : RegionsArray)
+					{
+						FFileRegion& Region = ChunkInfo.FileRegions.AddDefaulted_GetRef();
+						FFileRegion::LoadFromCompactBinary(RegionObj, Region);
+					}
+				}
+			}
+		};
+
+		AddChunksFromOplog("packagedata");
+		AddChunksFromOplog("bulkdata");
+		AddChunksFromOplog("files");
+
+		PackageIdToEntry.Add(PackageStoreEntry.GetPackageId(), MoveTemp(PackageStoreEntry));
+	}
+
+	UE_LOG(LogIoStore, Display, TEXT("Parsed %d oplog items %.2lf seconds, %d chunks"),
+		EntriesArray.Num(),
+		FPlatformTime::Seconds() - StartTime,
+		ChunkInfoMap.Num());
+}
+
+FIoStatus FCookedPackageStore::LoadChunkHashes()
+{
+	COOKEDPACKAGESTORE_CPU_SCOPE(LoadChunkHashes);
+	double StartLoadTime = FPlatformTime::Seconds();
+	double StartUpdateTime = StartLoadTime;
+
+	uint32 LoadedHashCount = 0;
+	uint32 UpdatedHashCount = 0;
+
+	if (HasZenStoreClient())
+	{
+		FCbObject ChunksObj;
+		{
+			COOKEDPACKAGESTORE_CPU_SCOPE(GetChunkInfos);
+			TIoStatusOr<FCbObject> Chunks = ZenStoreClient->GetChunkInfos().Get();
+			if (!Chunks.IsOk())
+			{
+				return Chunks.Status();
+			}
+			ChunksObj = Chunks.ConsumeValueOrDie();
+		}
+
+		StartUpdateTime = FPlatformTime::Seconds();
+
+		COOKEDPACKAGESTORE_CPU_SCOPE(ParseChunkInfos);
+		for (FCbField& ChunkEntry : ChunksObj["chunkinfos"])
+		{
+			FCbObject ChunkObj = ChunkEntry.AsObject();
+			FIoChunkId ChunkId;
+			if (!LoadFromCompactBinary(ChunkObj["id"], ChunkId))
+			{
+				UE_LOG(LogIoStore, Warning, TEXT("Received invalid chunk id, skipping."));
+				continue;
+			}
+			if (FChunkInfo* FindChunkInfo = ChunkInfoMap.Find(ChunkId); FindChunkInfo != nullptr)
+			{
+				FindChunkInfo->ChunkHash = ChunkObj["rawhash"].AsHash();
+				FindChunkInfo->ChunkSize = ChunkObj["rawsize"].AsUInt64();
+				++UpdatedHashCount;
+			}
+			++LoadedHashCount;
+		}
+	}
+	else
+	{
+		FAssetRegistryState AssetRegistry;
+		if (FindAndLoadMetadataFiles(CookedDir, ECookMetadataFiles::None, AssetRegistry, nullptr, nullptr, nullptr) == ECookMetadataFiles::None)
+		{
+			return FIoStatus(EIoErrorCode::NotFound);
+		}
+
+		StartUpdateTime = FPlatformTime::Seconds();
+
+		const TMap<FName, const FAssetPackageData*>& Packages = AssetRegistry.GetAssetPackageDataMap();
+		for (auto PackageIter : Packages)
+		{
+			for (const TPair<FIoChunkId, FIoHash>& HashIter : PackageIter.Value->ChunkHashes)
+			{
+				// For the moment, only bulk data types are added to teh asset registry - gate here so that
+				// we remember to verify all the hashes match when they eventually get added during cook.
+				if (HashIter.Key.GetChunkType() == EIoChunkType::BulkData ||
+					HashIter.Key.GetChunkType() == EIoChunkType::OptionalBulkData)
+				{
+					if (FChunkInfo* FindChunkInfo = ChunkInfoMap.Find(HashIter.Key); FindChunkInfo != nullptr)
+					{
+						FindChunkInfo->ChunkHash = HashIter.Value;
+						++UpdatedHashCount;
+					}
+					++LoadedHashCount;
+				}
+			}
+		}
+	}
+
+	UE_LOG(LogIoStore, Display, TEXT("Loaded %u chunk hashes from %s in %.2lf seconds, %d hashes updated in %.2lf seconds"),
+		LoadedHashCount,
+		HasZenStoreClient() ? TEXT("Zen") : TEXT("AssetRegistry"),
+		StartUpdateTime - StartLoadTime,
+		UpdatedHashCount,
+		FPlatformTime::Seconds() - StartUpdateTime);
+
+	return FIoStatus::Ok;
+}
+
+FIoChunkId FCookedPackageStore::GetChunkIdFromFileName(const FString& Filename) const
+{
+	return FilenameToChunkIdMap.FindRef(Filename);
+}
+
+const FCookedPackageStore::FChunkInfo* FCookedPackageStore::GetChunkInfoFromChunkId(const FIoChunkId& ChunkId) const
+{
+	return ChunkInfoMap.Find(ChunkId);
+}
+
+const FCookedPackageStore::FChunkInfo* FCookedPackageStore::GetChunkInfoFromFileName(const FString& Filename) const
+{
+	FIoChunkId ChunkId = GetChunkIdFromFileName(Filename);
+	return ChunkInfoMap.Find(ChunkId);
+}
+
+FString FCookedPackageStore::GetRelativeFilenameFromChunkId(const FIoChunkId& ChunkId) const
+{
+	const FChunkInfo* FindChunkInfo = ChunkInfoMap.Find(ChunkId);
+	if (!FindChunkInfo)
+	{
+		return FString();
+	}
+	return FindChunkInfo->RelativeFileName;
+}
+
+FName FCookedPackageStore::GetPackageNameFromChunkId(const FIoChunkId& ChunkId) const
+{
+	const FChunkInfo* FindChunkInfo = ChunkInfoMap.Find(ChunkId);
+	if (!FindChunkInfo)
+	{
+		return NAME_None;
+	}
+	return FindChunkInfo->PackageName;
+}
+
+FName FCookedPackageStore::GetPackageNameFromFileName(const FString& Filename) const
+{
+	FIoChunkId ChunkId = GetChunkIdFromFileName(Filename);
+	return GetPackageNameFromChunkId(ChunkId);
+}
+
+const FPackageStoreEntryResource* FCookedPackageStore::GetPackageStoreEntry(FPackageId PackageId) const
+{
+	return PackageIdToEntry.Find(PackageId);
+}
+
+bool FCookedPackageStore::HasZenStoreClient() const
+{
+	return ZenStoreClient.IsValid();
+}
+
+TIoStatusOr<FIoBuffer> FCookedPackageStore::ReadChunk(const FIoChunkId& ChunkId)
+{
+	FIoReadOptions ReadOptions;
+	return ZenStoreClient->ReadChunk(ChunkId, ReadOptions.GetOffset(), ReadOptions.GetSize());
+}
+
+UE::Tasks::FTask FCookedPackageStore::ReadChunkAsync(const FIoChunkId& ChunkId, TFunction<void(TIoStatusOr<FIoBuffer>)>&& Callback)
+{
+	return UE::Tasks::Launch(TEXT("ReadChunkAsync"), [this, ChunkId, Callback = MoveTemp(Callback)]()
+	{
+		FIoReadOptions ReadOptions;
+		Callback(ZenStoreClient->ReadChunk(ChunkId, ReadOptions.GetOffset(), ReadOptions.GetSize()));
+	}, UE::Tasks::ETaskPriority::Normal);
+}
diff --git a/Engine/Source/Developer/IoStoreUtilities/Private/IoStoreUtilities.cpp b/Engine/Source/Developer/IoStoreUtilities/Private/IoStoreUtilities.cpp
--- a/Engine/Source/Developer/IoStoreUtilities/Private/IoStoreUtilities.cpp
+++ b/Engine/Source/Developer/IoStoreUtilities/Private/IoStoreUtilities.cpp
@@ -7,7 +7,9 @@
 #include "Algo/BinarySearch.h"
 #include "Algo/TopologicalSort.h"
 #include "Async/AsyncWork.h"
+#include "CookedPackageStore.h"
 #include "CookMetadata.h"
+#include "CookMetadataFiles.h"
 #include "DerivedDataCache.h"
 #include "DerivedDataCacheInterface.h"
 #include "GenericPlatform/GenericPlatformMisc.h"
@@ -117,20 +119,6 @@ static const uint64 DefaultMemoryMappingAlignment = 16 << 10;
 
 static TUniquePtr<FIoStoreReader> CreateIoStoreReader(const TCHAR* Path, const FKeyChain& KeyChain);
 
-enum class ECookMetadataFiles
-{
-	None = 0,
-	AssetRegistry = 1,
-	CookMetadata = 2,
-	All = 4
-};
-ENUM_CLASS_FLAGS(ECookMetadataFiles);
-
-static ECookMetadataFiles FindAndLoadMetadataFiles(
-	const FString& InCookedDir, ECookMetadataFiles InRequiredFiles, 
-	FAssetRegistryState& OutAssetRegistry, FString* OutAssetRegistryFileName /*optional, set on success*/,
-	UE::Cook::FCookMetadataState* OutCookMetadata, FString* OutCookMetadataFileName /*optional, set on success or need*/);
-
 /*
 * Provides access to previously compressed chunks to the iostore writer, allowing
 * a) avoiding recompressing things and b) tweaks to compressors dont cause massive patches.
@@ -1032,377 +1020,6 @@ struct FContainerTargetFile
 	TArray<FFileRegion> FileRegions;
 };
 
-class FCookedPackageStore
-{
-public:
-	struct FChunkInfo
-	{
-		FIoChunkId ChunkId;
-		FIoHash ChunkHash;
-		uint64 ChunkSize = 0;
-		FName PackageName;
-		FString RelativeFileName;
-		TArray<FFileRegion> FileRegions;
-	};
-
-	FCookedPackageStore(FStringView InCookedDir)
-		: CookedDir(InCookedDir)
-	{
-	}
-
-	FIoStatus LoadManifest(const TCHAR* ManifestFilename)
-	{
-		IOSTORE_CPU_SCOPE(LoadCookedPackageStoreManifest);
-		double StartTime = FPlatformTime::Seconds();
-
-		TUniquePtr<FArchive> Ar(IFileManager::Get().CreateFileReader(ManifestFilename));
-		if (!Ar)
-		{
-			return FIoStatus(EIoErrorCode::NotFound);
-		}
-		FCbObject ManifestObject = LoadCompactBinary(*Ar).AsObject();
-		FCbObject OplogObject;
-		if (FCbFieldView ZenServerField = ManifestObject["zenserver"])
-		{
-			UE::Zen::FServiceSettings ZenServiceSettings;
-			ZenServiceSettings.ReadFromCompactBinary(ZenServerField["settings"]);
-			FString ProjectId = FString(ZenServerField["projectid"].AsString());
-			FString OplogId = FString(ZenServerField["oplogid"].AsString());
-
-			ZenStoreClient = MakeUnique<UE::FZenStoreHttpClient>(MoveTemp(ZenServiceSettings));
-			ZenStoreClient->InitializeReadOnly(ProjectId, OplogId);
-
-			IOSTORE_CPU_SCOPE(FetchOplog);
-			TIoStatusOr<FCbObject> OplogStatus = ZenStoreClient->GetOplog().Get();
-			if (!OplogStatus.IsOk())
-			{
-				return OplogStatus.Status();
-			}
-			OplogObject = OplogStatus.ConsumeValueOrDie();
-			
-		}
-		else
-		{
-			OplogObject = ManifestObject["oplog"].AsObject();
-		}
-		UE_LOG(LogIoStore, Display, TEXT("Fetched %d oplog items from %s in %.2lf seconds"),
-			OplogObject["entries"].AsArrayView().Num(),
-			HasZenStoreClient() ? TEXT("Zen") : TEXT("Manifest"),
-			FPlatformTime::Seconds() - StartTime);
-
-		ParseOplog(OplogObject);
-		LoadChunkHashes();
-
-		return FIoStatus::Ok;
-	}
-
-	FIoStatus LoadProjectStore(const TCHAR* ProjectStoreFilename)
-	{
-		IOSTORE_CPU_SCOPE(LoadCookedProjectStore);
-		double StartTime = FPlatformTime::Seconds();
-
-		TUniquePtr<FArchive> Ar(IFileManager::Get().CreateFileReader(ProjectStoreFilename));
-		if (!Ar)
-		{
-			return FIoStatus(EIoErrorCode::NotFound);
-		}
-
-		TSharedPtr<FJsonObject> ProjectStoreObject;
-		TSharedRef<TJsonReader<UTF8CHAR>> Reader = TJsonReaderFactory<UTF8CHAR>::Create(Ar.Get());
-		if (FJsonSerializer::Deserialize(Reader, ProjectStoreObject) && ProjectStoreObject.IsValid())
-		{
-			const TSharedPtr<FJsonObject>* ZenServerObjectPtr = nullptr;
-			if (ProjectStoreObject->TryGetObjectField(TEXT("zenserver"), ZenServerObjectPtr) && (ZenServerObjectPtr != nullptr))
-			{
-				FString ProjectId;
-				FString OplogId;
-
-				const TSharedPtr<FJsonObject>& ZenServerObject = *ZenServerObjectPtr;
-				if (ZenServerObject->TryGetStringField(TEXT("projectid"), ProjectId) && !ProjectId.IsEmpty() &&
-					ZenServerObject->TryGetStringField(TEXT("oplogid"), OplogId) && !OplogId.IsEmpty())
-				{
-					ZenStoreClient = MakeUnique<UE::FZenStoreHttpClient>();
-					ZenStoreClient->InitializeReadOnly(ProjectId, OplogId);
-
-					IOSTORE_CPU_SCOPE(FetchOplog);
-					TIoStatusOr<FCbObject> OplogStatus = ZenStoreClient->GetOplog().Get();
-					if (!OplogStatus.IsOk())
-					{
-						return OplogStatus.Status();
-					}
-					FCbObject OplogObject = OplogStatus.ConsumeValueOrDie();
-
-					UE_LOG(LogIoStore, Display, TEXT("Fetched %d oplog items from Zen in %.2lf seconds"),
-						OplogObject["entries"].AsArrayView().Num(),
-						FPlatformTime::Seconds() - StartTime);
-
-					ParseOplog(OplogObject);
-					LoadChunkHashes();
-
-					return FIoStatus::Ok;
-				}
-			}
-		}
-
-		return FIoStatus(EIoErrorCode::NotFound);
-	}
-
-	void ParseOplog(FCbObject& OplogObject)
-	{
-		IOSTORE_CPU_SCOPE(ParseOplog);
-		double StartTime = FPlatformTime::Seconds();
-
-		const FCbArrayView EntriesArray = OplogObject["entries"].AsArrayView();
-		const int32 EstimatedChunksCount = 3 * EntriesArray.Num(); //PackageData+BulkData+OptionalBulkData
-		ChunkInfoMap.Reserve(EstimatedChunksCount);
-		FilenameToChunkIdMap.Reserve(EstimatedChunksCount);
-		PackageIdToEntry.Reserve(EntriesArray.Num());
-		for (FCbFieldView& OplogEntry : EntriesArray)
-		{
-			FCbObjectView OplogObj = OplogEntry.AsObjectView();
-			FPackageStoreEntryResource PackageStoreEntry = FPackageStoreEntryResource::FromCbObject(OplogObj["packagestoreentry"].AsObjectView());
-
-			auto AddChunksFromOplog = [this, &OplogEntry, &PackageStoreEntry](const char* Field)
-			{
-				for (FCbFieldView& ChunkEntry : OplogEntry[Field].AsArrayView())
-				{
-					FCbObjectView ChunkObj = ChunkEntry.AsObjectView();
-					FIoChunkId ChunkId;
-					ChunkId.Set(ChunkObj["id"].AsObjectId().GetView());
-					if (ChunkId.IsValid())
-					{
-						FChunkInfo& ChunkInfo = ChunkInfoMap.Add(ChunkId);
-						ChunkInfo.ChunkId = ChunkId;
-						ChunkInfo.PackageName = PackageStoreEntry.PackageName;
-						if (ChunkObj["filename"])
-						{
-							TStringBuilder<1024> RelativeFilename;
-							RelativeFilename.Append(ChunkObj["filename"].AsString());
-							ChunkInfo.RelativeFileName = RelativeFilename;
-							TStringBuilder<1024> PathBuilder;
-							FPathViews::AppendPath(PathBuilder, CookedDir);
-							FPathViews::AppendPath(PathBuilder, RelativeFilename);
-							FPathViews::NormalizeFilename(PathBuilder);
-							FilenameToChunkIdMap.Add(*PathBuilder, ChunkId);
-						}
-						else if (ChunkObj["clientpath"])
-						{
-							TStringBuilder<1024> RelativeFilename;
-							FUtf8StringView ClientPathView = ChunkObj["clientpath"].AsString();
-							if (ClientPathView.StartsWith("/{project}/"))
-							{
-								if (!FPaths::IsProjectFilePathSet())
-								{
-									RelativeFilename.Append(ClientPathView);
-									UE_LOG(LogIoStore, Warning, TEXT("Project relative path could not be remapped because project file path is unset (possibly due to not specifying uproject path as first argument). %s"), RelativeFilename.ToString());
-								}
-								else
-								{
-									RelativeFilename.Append(FPathViews::GetBaseFilename(FPaths::GetProjectFilePath()));
-									RelativeFilename.AppendChar('/');
-									RelativeFilename.Append(ClientPathView.RightChop(11));
-								}
-							}
-							else if (ClientPathView.StartsWith("/{engine}/"))
-							{
-								RelativeFilename.Append(TEXT("Engine/"));
-								RelativeFilename.Append(ClientPathView.RightChop(9));
-							}
-							else
-							{
-								RelativeFilename.Append(ClientPathView);
-							}
-							ChunkInfo.RelativeFileName = RelativeFilename;
-							TStringBuilder<1024> PathBuilder;
-							FPathViews::AppendPath(PathBuilder, CookedDir);
-							FPathViews::AppendPath(PathBuilder, RelativeFilename);
-							FPathViews::NormalizeFilename(PathBuilder);
-							FilenameToChunkIdMap.Add(*PathBuilder, ChunkId);
-						}
-
-						const FCbArrayView RegionsArray = ChunkObj["fileregions"].AsArrayView();
-						ChunkInfo.FileRegions.Reserve(RegionsArray.Num());
-						for (FCbFieldView RegionObj : RegionsArray)
-						{
-							FFileRegion& Region = ChunkInfo.FileRegions.AddDefaulted_GetRef();
-							FFileRegion::LoadFromCompactBinary(RegionObj, Region);
-						}
-					}
-				}
-			};
-
-			AddChunksFromOplog("packagedata");
-			AddChunksFromOplog("bulkdata");
-			AddChunksFromOplog("files");
-
-			PackageIdToEntry.Add(PackageStoreEntry.GetPackageId(), MoveTemp(PackageStoreEntry));
-		}
-
-		UE_LOG(LogIoStore, Display, TEXT("Parsed %d oplog items %.2lf seconds, %d chunks"),
-			EntriesArray.Num(),
-			FPlatformTime::Seconds() - StartTime,
-			ChunkInfoMap.Num());
-	}
-
-	FIoStatus LoadChunkHashes()
-	{
-		IOSTORE_CPU_SCOPE(LoadChunkHashes);
-		double StartLoadTime = FPlatformTime::Seconds();
-		double StartUpdateTime = StartLoadTime;
-
-		uint32 LoadedHashCount = 0;
-		uint32 UpdatedHashCount = 0;
-
-		if (HasZenStoreClient())
-		{
-			FCbObject ChunksObj;
-			{
-				IOSTORE_CPU_SCOPE(GetChunkInfos);
-				TIoStatusOr<FCbObject> Chunks = ZenStoreClient->GetChunkInfos().Get();
-				if (!Chunks.IsOk())
-				{
-					return Chunks.Status();
-				}
-				ChunksObj = Chunks.ConsumeValueOrDie();
-			}
-
-			StartUpdateTime = FPlatformTime::Seconds();
-
-			IOSTORE_CPU_SCOPE(ParseChunkInfos);
-			for (FCbField& ChunkEntry : ChunksObj["chunkinfos"])
-			{
-				FCbObject ChunkObj = ChunkEntry.AsObject();
-				FIoChunkId ChunkId;
-				if (!LoadFromCompactBinary(ChunkObj["id"], ChunkId))
-				{
-					UE_LOG(LogIoStore, Warning, TEXT("Received invalid chunk id, skipping."));
-					continue;
-				}
-				if (FChunkInfo* FindChunkInfo = ChunkInfoMap.Find(ChunkId); FindChunkInfo != nullptr)
-				{
-					FindChunkInfo->ChunkHash = ChunkObj["rawhash"].AsHash();
-					FindChunkInfo->ChunkSize = ChunkObj["rawsize"].AsUInt64();
-					++UpdatedHashCount;
-				}
-				++LoadedHashCount;
-			}
-		}
-		else
-		{
-			FAssetRegistryState AssetRegistry;
-			if (FindAndLoadMetadataFiles(CookedDir, ECookMetadataFiles::None, AssetRegistry, nullptr, nullptr, nullptr) == ECookMetadataFiles::None)
-			{
-				return FIoStatus(EIoErrorCode::NotFound);
-			}
-
-			StartUpdateTime = FPlatformTime::Seconds();
-
-			const TMap<FName, const FAssetPackageData*>& Packages = AssetRegistry.GetAssetPackageDataMap();
-			for (auto PackageIter : Packages)
-			{
-				for (const TPair<FIoChunkId, FIoHash>& HashIter : PackageIter.Value->ChunkHashes)
-				{
-					// For the moment, only bulk data types are added to teh asset registry - gate here so that
-					// we remember to verify all the hashes match when they eventually get added during cook.
-					if (HashIter.Key.GetChunkType() == EIoChunkType::BulkData ||
-						HashIter.Key.GetChunkType() == EIoChunkType::OptionalBulkData)
-					{
-						if (FChunkInfo* FindChunkInfo = ChunkInfoMap.Find(HashIter.Key); FindChunkInfo != nullptr)
-						{
-							FindChunkInfo->ChunkHash = HashIter.Value;
-							++UpdatedHashCount;
-						}
-						++LoadedHashCount;
-					}
-				}
-			}
-		}
-
-		UE_LOG(LogIoStore, Display, TEXT("Loaded %u chunk hashes from %s in %.2lf seconds, %d hashes updated in %.2lf seconds"),
-			LoadedHashCount,
-			HasZenStoreClient() ? TEXT("Zen") : TEXT("AssetRegistry"),
-			StartUpdateTime - StartLoadTime,
-			UpdatedHashCount,
-			FPlatformTime::Seconds() - StartUpdateTime);
-
-		return FIoStatus::Ok;
-	}
-
-	FIoChunkId GetChunkIdFromFileName(const FString& Filename) const
-	{
-		return FilenameToChunkIdMap.FindRef(Filename);
-	}
-
-	const FChunkInfo* GetChunkInfoFromChunkId(const FIoChunkId& ChunkId) const
-	{
-		return ChunkInfoMap.Find(ChunkId);
-	}
-
-	const FChunkInfo* GetChunkInfoFromFileName(const FString& Filename) const
-	{
-		FIoChunkId ChunkId = GetChunkIdFromFileName(Filename);
-		return ChunkInfoMap.Find(ChunkId);
-	}
-
-	FString GetRelativeFilenameFromChunkId(const FIoChunkId& ChunkId) const
-	{
-		const FChunkInfo* FindChunkInfo = ChunkInfoMap.Find(ChunkId);
-		if (!FindChunkInfo)
-		{
-			return FString();
-		}
-		return FindChunkInfo->RelativeFileName;
-	}
-
-	FName GetPackageNameFromChunkId(const FIoChunkId& ChunkId) const
-	{
-		const FChunkInfo* FindChunkInfo = ChunkInfoMap.Find(ChunkId);
-		if (!FindChunkInfo)
-		{
-			return NAME_None;
-		}
-		return FindChunkInfo->PackageName;
-	}
-
-	FName GetPackageNameFromFileName(const FString& Filename) const
-	{
-		FIoChunkId ChunkId = GetChunkIdFromFileName(Filename);
-		return GetPackageNameFromChunkId(ChunkId);
-	}
-
-	const FPackageStoreEntryResource* GetPackageStoreEntry(FPackageId PackageId) const
-	{
-		return PackageIdToEntry.Find(PackageId);
-	}
-
-	bool HasZenStoreClient() const
-	{
-		return ZenStoreClient.IsValid();
-	}
-
-	TIoStatusOr<FIoBuffer> ReadChunk(const FIoChunkId& ChunkId)
-	{
-		FIoReadOptions ReadOptions;
-		return ZenStoreClient->ReadChunk(ChunkId, ReadOptions.GetOffset(), ReadOptions.GetSize());
-	}
-
-	UE::Tasks::FTask ReadChunkAsync(const FIoChunkId& ChunkId, TFunction<void(TIoStatusOr<FIoBuffer>)>&& Callback)
-	{
-		return UE::Tasks::Launch(TEXT("ReadChunkAsync"), [this, ChunkId, Callback = MoveTemp(Callback)]()
-		{
-			FIoReadOptions ReadOptions;
-			Callback(ZenStoreClient->ReadChunk(ChunkId, ReadOptions.GetOffset(), ReadOptions.GetSize()));
-		}, UE::Tasks::ETaskPriority::Normal);
-	}
-
-private:
-	TUniquePtr<UE::FZenStoreHttpClient> ZenStoreClient;
-	FString CookedDir;
-	TMap<FPackageId, FPackageStoreEntryResource> PackageIdToEntry;
-	TMap<FString, FIoChunkId> FilenameToChunkIdMap;
-	TMap<FIoChunkId, FChunkInfo> ChunkInfoMap;
-};
-
 struct FFileOrderMap
 {
 	TMap<FName, int64> PackageNameToOrder;
@@ -4938,32 +4555,6 @@ static void AddChunkInfoToAssetRegistry(TMap<FPackageId, TArray<FIoStoreChunkSou
 	}
 }
 
-// Returns the hash of the development asset registry or 0 on failure.
-static uint64 LoadAssetRegistry(const FString& InAssetRegistryFileName, FAssetRegistryState& OutAssetRegistry)
-{
-	FAssetRegistryVersion::Type Version;
-	FAssetRegistryLoadOptions Options(UE::AssetRegistry::ESerializationTarget::ForDevelopment);
-
-	TUniquePtr<FArchive> FileReader(IFileManager::Get().CreateFileReader(*InAssetRegistryFileName));
-	if (FileReader)
-	{
-		TArray64<uint8> Data;
-		Data.SetNumUninitialized(FileReader->TotalSize());
-		FileReader->Serialize(Data.GetData(), Data.Num());
-		check(!FileReader->IsError());
-
-		uint64 DevArHash = UE::Cook::FCookMetadataState::ComputeHashOfDevelopmentAssetRegistry(MakeMemoryView(Data));
-
-		FLargeMemoryReader MemoryReader(Data.GetData(), Data.Num());
-		if (OutAssetRegistry.Load(MemoryReader, Options, &Version))
-		{
-			return DevArHash;
-		}
-	}
-
-	return 0;;
-}
-
 static bool SaveAssetRegistry(const FString& InAssetRegistryFileName, FAssetRegistryState& InAssetRegistry, uint64* OutDevArHash)
 {
 	TRACE_CPUPROFILER_EVENT_SCOPE(SavingAssetRegistry);
@@ -5011,115 +4602,6 @@ static bool SaveAssetRegistry(const FString& InAssetRegistryFileName, FAssetRegi
 	return true;
 }
 
-static ECookMetadataFiles FindAndLoadMetadataFiles(
-	const FString& InCookedDir, ECookMetadataFiles InRequiredFiles, 
-	FAssetRegistryState& OutAssetRegistry, FString* OutAssetRegistryFileName /*optional, set on success*/,
-	UE::Cook::FCookMetadataState* OutCookMetadata, FString* OutCookMetadataFileName /*optional, set on success or need*/)
-{
-	TRACE_CPUPROFILER_EVENT_SCOPE(LoadingAssetRegistry);
-
-	// Look for the development registry. Should be in \\GameName\\Metadata\\DevelopmentAssetRegistry.bin, but we don't know what "GameName" is.
-	TArray<FString> PossibleAssetRegistryFiles;
-	IFileManager::Get().FindFilesRecursive(PossibleAssetRegistryFiles, *InCookedDir, GetDevelopmentAssetRegistryFilename(), true, false);
-
-	if (PossibleAssetRegistryFiles.Num() > 1)
-	{
-		UE_LOG(LogIoStore, Warning, TEXT("Found multiple possible development asset registries:"));
-		for (FString& Filename : PossibleAssetRegistryFiles)
-		{
-			UE_LOG(LogIoStore, Warning, TEXT("    %s"), *Filename);
-		}
-	}
-
-	if (PossibleAssetRegistryFiles.Num() == 0)
-	{
-		if (EnumHasAnyFlags(InRequiredFiles, ECookMetadataFiles::AssetRegistry))
-		{
-			UE_LOG(LogIoStore, Error, TEXT("No development asset registry file found!"));
-		}
-		else
-		{
-			UE_LOG(LogIoStore, Display, TEXT("No development asset registry file found!"));
-		}
-		return ECookMetadataFiles::None;
-	}
-
-	UE_LOG(LogIoStore, Display, TEXT("Using input asset registry: %s"), *PossibleAssetRegistryFiles[0]);
-	uint64 LoadedDevArHash = LoadAssetRegistry(PossibleAssetRegistryFiles[0], OutAssetRegistry);
-
-	if (LoadedDevArHash == 0)
-	{
-		return ECookMetadataFiles::None; // already logged
-	}
-
-	// If we found the asset registry, try and find the cook metadata that should be next to it.
-	ECookMetadataFiles ResultFiles = ECookMetadataFiles::AssetRegistry;
-
-	if (OutCookMetadata)
-	{
-		// The cook metadata file should be adjacent to the development asset registry.
-		FString CookMetadataFileName = FPaths::GetPath(PossibleAssetRegistryFiles[0]) / UE::Cook::GetCookMetadataFilename();
-		if (IFileManager::Get().FileExists(*CookMetadataFileName))
-		{
-			if (OutCookMetadata->ReadFromFile(CookMetadataFileName) == false)
-			{
-				UE_LOG(LogIoStore, Error, TEXT("Failed to deserialize cook metadata file - invalid data. [%s]"), *CookMetadataFileName);
-				if (EnumHasAnyFlags(InRequiredFiles, ECookMetadataFiles::CookMetadata))
-				{
-					return ECookMetadataFiles::None;
-				}
-			}
-			else if (OutCookMetadata->GetAssociatedDevelopmentAssetRegistryHash() != LoadedDevArHash &&
-				OutCookMetadata->GetAssociatedDevelopmentAssetRegistryHashPostWriteback() != LoadedDevArHash) // during testing we can repeat stage after cook so we might have already edited it.
-			{
-				if (EnumHasAnyFlags(InRequiredFiles, ECookMetadataFiles::CookMetadata))
-				{
-					UE_LOG(LogIoStore, Error,
-						TEXT("Cook metadata file mismatch: Hash of associated development asset registry does not match. [%s] %llx vs %llx (%llx post writeback)"),
-						*CookMetadataFileName, LoadedDevArHash, OutCookMetadata->GetAssociatedDevelopmentAssetRegistryHash(), OutCookMetadata->GetAssociatedDevelopmentAssetRegistryHashPostWriteback());
-					return ECookMetadataFiles::None;
-				}
-				else
-				{
-					UE_LOG(LogIoStore, Display,
-						TEXT("Cook metadata file mismatch: Hash of associated development asset registry does not match. [%s] %llx vs %llx (%llx post writeback)"),
-						*CookMetadataFileName, LoadedDevArHash, OutCookMetadata->GetAssociatedDevelopmentAssetRegistryHash(), OutCookMetadata->GetAssociatedDevelopmentAssetRegistryHashPostWriteback());
-					OutCookMetadata->Reset();
-				}
-			}
-			else
-			{
-				EnumAddFlags(ResultFiles, ECookMetadataFiles::CookMetadata);
-				if (OutCookMetadataFileName)
-				{
-					*OutCookMetadataFileName = MoveTemp(CookMetadataFileName);
-				}
-			}
-		}
-		else
-		{
-			if (EnumHasAnyFlags(InRequiredFiles, ECookMetadataFiles::CookMetadata))
-			{
-				UE_LOG(LogIoStore, Error, TEXT("Failed to open and read cook metadata file %s"), *CookMetadataFileName);
-				return ECookMetadataFiles::None;
-			}
-
-			UE_LOG(LogIoStore, Display, TEXT("No cook metadata file found, checked %s"), *CookMetadataFileName);
-			if (OutCookMetadataFileName)
-			{
-				*OutCookMetadataFileName = FString("");
-			}
-		}
-	}
-
-
-	if (OutAssetRegistryFileName)
-	{
-		*OutAssetRegistryFileName = MoveTemp(PossibleAssetRegistryFiles[0]);
-	}
-	return ResultFiles;
-}
-
 struct FIoStoreWriterInfo
 {
 	UE::Cook::EPluginSizeTypes SizeType;
diff --git a/Engine/Source/Developer/IoStoreUtilities/Private/ZenFileSystemManifest.cpp b/Engine/Source/Developer/IoStoreUtilities/Private/ZenFileSystemManifest.cpp
--- a/Engine/Source/Developer/IoStoreUtilities/Private/ZenFileSystemManifest.cpp
+++ b/Engine/Source/Developer/IoStoreUtilities/Private/ZenFileSystemManifest.cpp
@@ -527,11 +527,17 @@ int32 FZenFileSystemManifest::Generate()
 
 const FZenFileSystemManifest::FManifestEntry& FZenFileSystemManifest::CreateManifestEntry(const FString& Filename)
 {
+	const FString FullFilename = FPaths::ConvertRelativePathToFull(Filename);
+
 	FString CookedEngineDirectory = FPaths::Combine(CookDirectory, TEXT("Engine"));
+	FString CookedEngineDirectoryTrailingSeparator;
+	CookedEngineDirectoryTrailingSeparator.Reserve(CookedEngineDirectory.Len() + 1);
+	CookedEngineDirectoryTrailingSeparator.Append(CookedEngineDirectory);
+	CookedEngineDirectoryTrailingSeparator.AppendChar(TEXT('/'));
 
-	auto AddEntry = [this, &Filename](const FString& ClientDirectory, const FString& LocalDirectory) -> const FManifestEntry&
+	auto AddEntry = [this, &FullFilename](const FString& ClientDirectory, const FString& LocalDirectory) -> const FManifestEntry&
 	{
-		FStringView RelativePath = Filename;
+		FStringView RelativePath = FullFilename;
 		RelativePath.RightChopInline(LocalDirectory.Len() + 1);
 
 		FString ServerRelativeDirectory = LocalDirectory;
@@ -545,13 +551,18 @@ const FZenFileSystemManifest::FManifestEntry& FZenFileSystemManifest::CreateMani
 		return AddManifestEntry(FileChunkId, MoveTemp(ServerPath), MoveTemp(ClientPath));
 	};
 
-	if (Filename.StartsWith(CookedEngineDirectory))
+	if (FullFilename.StartsWith(CookedEngineDirectoryTrailingSeparator))
 	{
 		return AddEntry(TEXT("/{engine}"), CookedEngineDirectory);
 	}
 
 	FString CookedProjectDirectory = FPaths::Combine(CookDirectory, FApp::GetProjectName());
-	if (Filename.StartsWith(CookedProjectDirectory))
+	FString CookedProjectDirectoryTrailingSeparator;
+	CookedProjectDirectoryTrailingSeparator.Reserve(CookedProjectDirectory.Len() + 1);
+	CookedProjectDirectoryTrailingSeparator.Append(CookedProjectDirectory);
+	CookedProjectDirectoryTrailingSeparator.AppendChar(TEXT('/'));
+
+	if (FullFilename.StartsWith(CookedProjectDirectoryTrailingSeparator))
 	{
 		return AddEntry(TEXT("/{project}"), CookedProjectDirectory);
 	}
diff --git a/Engine/Source/Developer/PakFileUtilities/Private/PakFileUtilities.cpp b/Engine/Source/Developer/PakFileUtilities/Private/PakFileUtilities.cpp
--- a/Engine/Source/Developer/PakFileUtilities/Private/PakFileUtilities.cpp
+++ b/Engine/Source/Developer/PakFileUtilities/Private/PakFileUtilities.cpp
@@ -34,6 +34,8 @@
 #include "Async/Async.h"
 #include "Async/Future.h"
 #include "Virtualization/VirtualizationSystem.h"
+#include "CookedPackageStore.h"
+#include "ZenStoreHttpClient.h"
 
 IMPLEMENT_MODULE(FDefaultModuleImpl, PakFileUtilities);
 
@@ -553,6 +555,7 @@ struct FPakCommandLineParameters
 	FString InputFinalPakFilename; // This is the resulting pak file we want to end up with after we generate the pak patch.  This is used instead of passing in the raw content.
 	FString ChangedFilesOutputFilename;
 	FString CsvPath;
+	FString ProjectStoreFilename;
 	bool EncryptIndex;
 	bool UseCustomCompressor;
 	FGuid EncryptionKeyGuid;
@@ -672,7 +675,7 @@ struct FCompressedFileBuffer
 	}
 
 	void ResetSource();
-	bool ReadSource(const FPakInputPair& InFile);
+	bool ReadSource(const FPakInputPair& InFile, FCookedPackageStore* PackageStore);
 	void SetSourceAsWorkingBuffer();
 
 	TUniquePtr<FMemoryCompressor> BeginCompressFileToWorkingBuffer(const FPakInputPair& InFile, FName CompressionMethod, const int32 CompressionBlockSize, FGraphEventRef EndCompressionBarrier);
@@ -828,35 +831,46 @@ void FCompressedFileBuffer::ResetSource()
 	UncompressedBuffer.Empty();
 }
 
-bool FCompressedFileBuffer::ReadSource(const FPakInputPair& InputInfo)
+bool FCompressedFileBuffer::ReadSource(const FPakInputPair& InputInfo, FCookedPackageStore* PackageStore)
 {
 	if (!InputInfo.bNeedRehydration)
 	{
 		TUniquePtr<FArchive> FileHandle(IFileManager::Get().CreateFileReader(*InputInfo.Source));
-		if (!FileHandle)
+		if (FileHandle)
 		{
-			OriginalSize = -1;
-			return false;
-		}
-
-		OriginalSize = FileHandle->TotalSize();
-		const int64 PaddedEncryptedFileSize = Align(OriginalSize, FAES::AESBlockSize);
+			OriginalSize = FileHandle->TotalSize();
+			const int64 PaddedEncryptedFileSize = Align(OriginalSize, FAES::AESBlockSize);
 
-		UncompressedBuffer.SetNumUninitialized(PaddedEncryptedFileSize);
+			UncompressedBuffer.SetNumUninitialized(PaddedEncryptedFileSize);
 
-		FileHandle->Serialize(UncompressedBuffer.GetData(), OriginalSize);
+			FileHandle->Serialize(UncompressedBuffer.GetData(), OriginalSize);
 
-		if (!FileHandle->IsError())
-		{
-			return true;
+			if (!FileHandle->IsError())
+			{
+				return true;
+			}
 		}
-		else
+		else if (PackageStore && PackageStore->HasZenStoreClient())
 		{
-			UncompressedBuffer.Empty();
-			OriginalSize = -1;
-
-			return false;
+			FString FullFilename = FPaths::ConvertRelativePathToFull(InputInfo.Source);
+			FIoChunkId ChunkId = PackageStore->GetChunkIdFromFileName(FullFilename);
+			if (ChunkId.IsValid())
+			{
+				TIoStatusOr<FIoBuffer> ChunkReadStatus = PackageStore->ReadChunk(ChunkId);
+				if (ChunkReadStatus.IsOk())
+				{
+					FIoBuffer Buffer = ChunkReadStatus.ConsumeValueOrDie();
+					OriginalSize = Buffer.GetSize();
+					const int64 PaddedEncryptedFileSize = Align(OriginalSize, FAES::AESBlockSize);
+					UncompressedBuffer.SetNumUninitialized(PaddedEncryptedFileSize);
+					FMemory::Memcpy(UncompressedBuffer.GetData(), Buffer.GetData(), Buffer.GetSize());
+					return true;
+				}
+			}
 		}
+		UncompressedBuffer.Empty();
+		OriginalSize = -1;
+		return false;
 	}
 	else
 	{
@@ -1266,6 +1280,8 @@ void ProcessCommonCommandLine(const TCHAR* CmdLine, FPakCommandLineParameters& C
 	FParse::Value(CmdLine, TEXT("-patchSeekOptMode="), (int32&)CmdLineParameters.SeekOptParams.Mode);
 
 	FParse::Value(CmdLine, TEXT("csv="), CmdLineParameters.CsvPath);
+
+	FParse::Value(FCommandLine::Get(), TEXT("ProjectStore="), CmdLineParameters.ProjectStoreFilename);
 }
 
 void ProcessPakFileSpecificCommandLine(const TCHAR* CmdLine, const TArray<FString>& NonOptionArguments, TArray<FPakInputPair>& Entries, FPakCommandLineParameters& CmdLineParameters)
@@ -2163,6 +2179,7 @@ private:
 	TSpscQueue<FOutputPakFileEntry*> WriteQueue;
 	TFuture<void> CompressionThread;
 	TFuture<void> WriterThread;
+	TUniquePtr<FCookedPackageStore> PackageStore;
 	FEventRef CompressionQueueEntryAddedEvent;
 	FEventRef WriteQueueEntryAddedEvent;
 	FEventRef EntryRetiredEvent;
@@ -2205,6 +2222,20 @@ bool FPakWriterContext::Initialize(const FPakCommandLineParameters& InCmdLinePar
 		UE_LOG(LogPakFile, Display, TEXT("%s"), *FormatLogLine);
 	}
 
+	if (!CmdLineParameters.ProjectStoreFilename.IsEmpty())
+	{
+		TUniquePtr<FCookedPackageStore> NewPackageStore = MakeUnique<FCookedPackageStore>(FPaths::GetPath(CmdLineParameters.ProjectStoreFilename));
+		FIoStatus Status = NewPackageStore->LoadProjectStore(*CmdLineParameters.ProjectStoreFilename);
+		if (Status.IsOk())
+		{
+			PackageStore = MoveTemp(NewPackageStore);
+		}
+		else
+		{
+			UE_LOG(LogPakFile, Fatal, TEXT("Failed loading project store '%s'"), *CmdLineParameters.ProjectStoreFilename);
+		}
+	}
+
 	// Oodle is built into the Engine now and can be used to decode startup phase files (ini,res,uplugin)
 	//	 which used to be forced to Zlib
 	//	 set bDoUseOodleDespiteNoPluginCompression = true to let Oodle compress those files
@@ -2366,7 +2397,7 @@ void FPakWriterContext::BeginCompress(FOutputPakFileEntry* Entry)
 {
 	TRACE_CPUPROFILER_EVENT_SCOPE(BeginCompress);
 
-	if (!Entry->AccessCompressedBuffer().ReadSource(Entry->InputPair))
+	if (!Entry->AccessCompressedBuffer().ReadSource(Entry->InputPair, PackageStore.Get()))
 	{
 		// TODO: Should we give an error?
 		Entry->bSomeCompressionSucceeded = true; // Prevent loading in EndCompress
diff --git a/Engine/Source/Programs/AutomationTool/Scripts/CopyBuildToStagingDirectory.Automation.cs b/Engine/Source/Programs/AutomationTool/Scripts/CopyBuildToStagingDirectory.Automation.cs
--- a/Engine/Source/Programs/AutomationTool/Scripts/CopyBuildToStagingDirectory.Automation.cs
+++ b/Engine/Source/Programs/AutomationTool/Scripts/CopyBuildToStagingDirectory.Automation.cs
@@ -3486,6 +3486,12 @@ namespace AutomationScripts
 					CommonAdditionalArgs += " -allowForIndexUnload";
 				}
 			}
+			
+			if (!string.IsNullOrEmpty(SC.PackageStoreData.MarkerFullPath))
+			{
+				CommonAdditionalArgs += String.Format(" -ProjectStore={0}",
+					MakePathSafeToUseWithCommandLine(SC.PackageStoreData.MarkerFullPath));
+			}
 
 			// pass the targetplatform so the index may be able to be frozen
 			CommonAdditionalArgs += " -platform=" + ConfigHierarchy.GetIniPlatformName(SC.StageTargetPlatform.IniPlatformType);
