Ensure asset registry and cook metadata read and write-back functionality works when operating on an imported zen snapshot.

diff --git a/Engine/Source/Developer/IoStoreUtilities/Internal/CookedPackageStore.h b/Engine/Source/Developer/IoStoreUtilities/Internal/CookedPackageStore.h
--- a/Engine/Source/Developer/IoStoreUtilities/Internal/CookedPackageStore.h
+++ b/Engine/Source/Developer/IoStoreUtilities/Internal/CookedPackageStore.h
@@ -4,6 +4,7 @@
 
 #include "Containers/Array.h"
 #include "Containers/Map.h"
+#include "Containers/StringFwd.h"
 #include "Containers/UnrealString.h"
 #include "IO/IoBuffer.h"
 #include "IO/IoChunkId.h"
@@ -33,8 +34,7 @@ public:
 
 	FIoStatus LoadManifest(const TCHAR* ManifestFilename);
 	IOSTOREUTILITIES_API FIoStatus LoadProjectStore(const TCHAR* ProjectStoreFilename);
-	IOSTOREUTILITIES_API void ParseOplog(FCbObject& OplogObject);
-	IOSTOREUTILITIES_API FIoStatus LoadChunkHashes();
+
 	IOSTOREUTILITIES_API FIoChunkId GetChunkIdFromFileName(const FString& Filename) const;
 	IOSTOREUTILITIES_API const FChunkInfo* GetChunkInfoFromChunkId(const FIoChunkId& ChunkId) const;
 	IOSTOREUTILITIES_API const FChunkInfo* GetChunkInfoFromFileName(const FString& Filename) const;
@@ -42,11 +42,17 @@ public:
 	IOSTOREUTILITIES_API FName GetPackageNameFromChunkId(const FIoChunkId& ChunkId) const;
 	IOSTOREUTILITIES_API FName GetPackageNameFromFileName(const FString& Filename) const;
 	IOSTOREUTILITIES_API const FPackageStoreEntryResource* GetPackageStoreEntry(FPackageId PackageId) const;
+
 	IOSTOREUTILITIES_API bool HasZenStoreClient() const;
+	IOSTOREUTILITIES_API UE::FZenStoreHttpClient* GetZenStoreClient();
+
 	IOSTOREUTILITIES_API TIoStatusOr<FIoBuffer> ReadChunk(const FIoChunkId& ChunkId);
 	IOSTOREUTILITIES_API UE::Tasks::FTask ReadChunkAsync(const FIoChunkId& ChunkId, TFunction<void(TIoStatusOr<FIoBuffer>)>&& Callback);
 
 private:
+	void ParseOplog(FCbObject& OplogObject);
+	FIoStatus LoadChunkHashes();
+
 	TUniquePtr<UE::FZenStoreHttpClient> ZenStoreClient;
 	FString CookedDir;
 	TMap<FPackageId, FPackageStoreEntryResource> PackageIdToEntry;
diff --git a/Engine/Source/Developer/IoStoreUtilities/Internal/ZenStoreHttpClient.h b/Engine/Source/Developer/IoStoreUtilities/Internal/ZenStoreHttpClient.h
--- a/Engine/Source/Developer/IoStoreUtilities/Internal/ZenStoreHttpClient.h
+++ b/Engine/Source/Developer/IoStoreUtilities/Internal/ZenStoreHttpClient.h
@@ -62,6 +62,7 @@ public:
 	uint16 GetPort() const { return 8558; }
 #endif
 
+	TFuture<TIoStatusOr<FCbObject>> GetProjectInfo();
 	TFuture<TIoStatusOr<FCbObject>> GetOplog();
 	TFuture<TIoStatusOr<FCbObject>> GetFiles();
 	TFuture<TIoStatusOr<FCbObject>> GetChunkInfos();
@@ -95,6 +96,7 @@ private:
 #endif
 	TUniquePtr<Zen::FZenHttpRequestPool> RequestPool;
 	SaltGenerator SaltGen;
+	FString ProjectPath;
 	FString OplogPath;
 	FString OplogNewEntryPath;
 	FString OplogPrepNewEntryPath;
diff --git a/Engine/Source/Developer/IoStoreUtilities/Private/CookMetadataFiles.cpp b/Engine/Source/Developer/IoStoreUtilities/Private/CookMetadataFiles.cpp
--- a/Engine/Source/Developer/IoStoreUtilities/Private/CookMetadataFiles.cpp
+++ b/Engine/Source/Developer/IoStoreUtilities/Private/CookMetadataFiles.cpp
@@ -6,39 +6,59 @@
 #include "Containers/Array.h"
 #include "Containers/UnrealString.h"
 #include "CookMetadata.h"
+#include "CookedPackageStore.h"
 #include "HAL/FileManager.h"
+#include "IO/IoBuffer.h"
 #include "IO/IoStore.h"
+#include "Memory/MemoryView.h"
 #include "Misc/Paths.h"
 #include "ProfilingDebugging/CountersTrace.h"
 #include "Serialization/LargeMemoryReader.h"
+#include "ZenStoreHttpClient.h"
 
 // Returns the hash of the development asset registry or 0 on failure.
-static uint64 LoadAssetRegistry(const FString& InAssetRegistryFileName, FAssetRegistryState& OutAssetRegistry)
+static uint64 LoadAssetRegistry(FCookedPackageStore* InPackageStore, const FString& InAssetRegistryFileName, FIoChunkId InAssetRegistryChunkId, FAssetRegistryState& OutAssetRegistry)
 {
 	FAssetRegistryVersion::Type Version;
 	FAssetRegistryLoadOptions Options(UE::AssetRegistry::ESerializationTarget::ForDevelopment);
 
-	TUniquePtr<FArchive> FileReader(IFileManager::Get().CreateFileReader(*InAssetRegistryFileName));
-	if (FileReader)
+	if (InPackageStore && InPackageStore->HasZenStoreClient() && InAssetRegistryChunkId.IsValid())
 	{
-		TArray64<uint8> Data;
-		Data.SetNumUninitialized(FileReader->TotalSize());
-		FileReader->Serialize(Data.GetData(), Data.Num());
-		check(!FileReader->IsError());
+		FIoBuffer Buffer;
+		Buffer = InPackageStore->ReadChunk(InAssetRegistryChunkId).ConsumeValueOrDie();
+		uint64 DevArHash = UE::Cook::FCookMetadataState::ComputeHashOfDevelopmentAssetRegistry(Buffer.GetView());
 
-		uint64 DevArHash = UE::Cook::FCookMetadataState::ComputeHashOfDevelopmentAssetRegistry(MakeMemoryView(Data));
-
-		FLargeMemoryReader MemoryReader(Data.GetData(), Data.Num());
+		FLargeMemoryReader MemoryReader(Buffer.GetData(), Buffer.GetSize());
 		if (OutAssetRegistry.Load(MemoryReader, Options, &Version))
 		{
 			return DevArHash;
 		}
 	}
+	else
+	{
+		TUniquePtr<FArchive> FileReader(IFileManager::Get().CreateFileReader(*InAssetRegistryFileName));
+		if (FileReader)
+		{
+			TArray64<uint8> Data;
+			Data.SetNumUninitialized(FileReader->TotalSize());
+			FileReader->Serialize(Data.GetData(), Data.Num());
+			check(!FileReader->IsError());
+
+			uint64 DevArHash = UE::Cook::FCookMetadataState::ComputeHashOfDevelopmentAssetRegistry(MakeMemoryView(Data));
 
-	return 0;;
+			FLargeMemoryReader MemoryReader(Data.GetData(), Data.Num());
+			if (OutAssetRegistry.Load(MemoryReader, Options, &Version))
+			{
+				return DevArHash;
+			}
+		}
+	}
+
+	return 0;
 }
 
 ECookMetadataFiles FindAndLoadMetadataFiles(
+	FCookedPackageStore* InPackageStore,
 	const FString& InCookedDir, ECookMetadataFiles InRequiredFiles, 
 	FAssetRegistryState& OutAssetRegistry, FString* OutAssetRegistryFileName /*optional, set on success*/,
 	UE::Cook::FCookMetadataState* OutCookMetadata, FString* OutCookMetadataFileName /*optional, set on success or need*/)
@@ -58,6 +78,31 @@ ECookMetadataFiles FindAndLoadMetadataFiles(
 		}
 	}
 
+	FIoChunkId AssetRegistryChunkId = FIoChunkId::InvalidChunkId;
+	if ((PossibleAssetRegistryFiles.Num() == 0) && InPackageStore)
+	{
+		if (UE::FZenStoreHttpClient* ZenStoreClient = InPackageStore->GetZenStoreClient())
+		{
+			TIoStatusOr<FCbObject> ProjectInfoStatus = ZenStoreClient->GetProjectInfo().Get();
+			if (ProjectInfoStatus.IsOk())
+			{
+				FCbObject ProjectInfo = ProjectInfoStatus.ConsumeValueOrDie();
+				FString ProjectFile(ProjectInfo["projectfile"].AsString());
+				FString ProjectName = FPaths::GetBaseFilename(ProjectFile);
+				if (!ProjectName.IsEmpty())
+				{
+					FString CandidateAssetRegistryFilename = FPaths::Combine(InCookedDir, ProjectName, "Metadata", GetDevelopmentAssetRegistryFilename());
+					FPaths::NormalizeFilename(CandidateAssetRegistryFilename);
+					AssetRegistryChunkId = InPackageStore->GetChunkIdFromFileName(CandidateAssetRegistryFilename);
+					if (AssetRegistryChunkId.IsValid())
+					{
+						PossibleAssetRegistryFiles.Add(MoveTemp(CandidateAssetRegistryFilename));
+					}
+				}
+			}
+		}
+	}
+
 	if (PossibleAssetRegistryFiles.Num() == 0)
 	{
 		if (EnumHasAnyFlags(InRequiredFiles, ECookMetadataFiles::AssetRegistry))
@@ -71,8 +116,10 @@ ECookMetadataFiles FindAndLoadMetadataFiles(
 		return ECookMetadataFiles::None;
 	}
 
+	FPaths::NormalizeFilename(PossibleAssetRegistryFiles[0]);
+
 	UE_LOG(LogIoStore, Display, TEXT("Using input asset registry: %s"), *PossibleAssetRegistryFiles[0]);
-	uint64 LoadedDevArHash = LoadAssetRegistry(PossibleAssetRegistryFiles[0], OutAssetRegistry);
+	uint64 LoadedDevArHash = LoadAssetRegistry(InPackageStore, PossibleAssetRegistryFiles[0], AssetRegistryChunkId, OutAssetRegistry);
 
 	if (LoadedDevArHash == 0)
 	{
@@ -86,17 +133,11 @@ ECookMetadataFiles FindAndLoadMetadataFiles(
 	{
 		// The cook metadata file should be adjacent to the development asset registry.
 		FString CookMetadataFileName = FPaths::GetPath(PossibleAssetRegistryFiles[0]) / UE::Cook::GetCookMetadataFilename();
-		if (IFileManager::Get().FileExists(*CookMetadataFileName))
+
+		auto ValidateAndOutputCookMetadata = [InRequiredFiles, LoadedDevArHash, &CookMetadataFileName, &OutCookMetadata, &OutCookMetadataFileName, &ResultFiles]()
 		{
-			if (OutCookMetadata->ReadFromFile(CookMetadataFileName) == false)
-			{
-				UE_LOG(LogIoStore, Error, TEXT("Failed to deserialize cook metadata file - invalid data. [%s]"), *CookMetadataFileName);
-				if (EnumHasAnyFlags(InRequiredFiles, ECookMetadataFiles::CookMetadata))
-				{
-					return ECookMetadataFiles::None;
-				}
-			}
-			else if (OutCookMetadata->GetAssociatedDevelopmentAssetRegistryHash() != LoadedDevArHash &&
+
+			if (OutCookMetadata->GetAssociatedDevelopmentAssetRegistryHash() != LoadedDevArHash &&
 				OutCookMetadata->GetAssociatedDevelopmentAssetRegistryHashPostWriteback() != LoadedDevArHash) // during testing we can repeat stage after cook so we might have already edited it.
 			{
 				if (EnumHasAnyFlags(InRequiredFiles, ECookMetadataFiles::CookMetadata))
@@ -104,7 +145,7 @@ ECookMetadataFiles FindAndLoadMetadataFiles(
 					UE_LOG(LogIoStore, Error,
 						TEXT("Cook metadata file mismatch: Hash of associated development asset registry does not match. [%s] %llx vs %llx (%llx post writeback)"),
 						*CookMetadataFileName, LoadedDevArHash, OutCookMetadata->GetAssociatedDevelopmentAssetRegistryHash(), OutCookMetadata->GetAssociatedDevelopmentAssetRegistryHashPostWriteback());
-					return ECookMetadataFiles::None;
+					return false;
 				}
 				else
 				{
@@ -122,6 +163,52 @@ ECookMetadataFiles FindAndLoadMetadataFiles(
 					*OutCookMetadataFileName = MoveTemp(CookMetadataFileName);
 				}
 			}
+			return true;
+		};
+
+		if (InPackageStore && InPackageStore->HasZenStoreClient() && AssetRegistryChunkId.IsValid())
+		{
+			FIoChunkId CookMetadataChunkId = InPackageStore->GetChunkIdFromFileName(CookMetadataFileName);
+			if (!CookMetadataChunkId.IsValid())
+			{
+				UE_LOG(LogIoStore, Error, TEXT("Failed to find cook metadata file - chunk missing from package store. [%s]"), *CookMetadataFileName);
+				if (EnumHasAnyFlags(InRequiredFiles, ECookMetadataFiles::CookMetadata))
+				{
+					return ECookMetadataFiles::None;
+				}
+			}
+
+			FIoBuffer Buffer;
+			Buffer = InPackageStore->ReadChunk(CookMetadataChunkId).ConsumeValueOrDie();
+
+			FLargeMemoryReader MemoryReader(Buffer.GetData(), Buffer.GetSize());
+			if (!OutCookMetadata->Serialize(MemoryReader))
+			{
+				UE_LOG(LogIoStore, Error, TEXT("Failed to deserialize cook metadata from package store (%s)"), *CookMetadataFileName);
+				if (EnumHasAnyFlags(InRequiredFiles, ECookMetadataFiles::CookMetadata))
+				{
+					return ECookMetadataFiles::None;
+				}
+			}
+			else if (!ValidateAndOutputCookMetadata())
+			{
+				return ECookMetadataFiles::None;
+			}
+		}
+		else if (IFileManager::Get().FileExists(*CookMetadataFileName))
+		{
+			if (!OutCookMetadata->ReadFromFile(CookMetadataFileName))
+			{
+				UE_LOG(LogIoStore, Error, TEXT("Failed to deserialize cook metadata file - invalid data. [%s]"), *CookMetadataFileName);
+				if (EnumHasAnyFlags(InRequiredFiles, ECookMetadataFiles::CookMetadata))
+				{
+					return ECookMetadataFiles::None;
+				}
+			}
+			else if (!ValidateAndOutputCookMetadata())
+			{
+				return ECookMetadataFiles::None;
+			}
 		}
 		else
 		{
diff --git a/Engine/Source/Developer/IoStoreUtilities/Private/CookMetadataFiles.h b/Engine/Source/Developer/IoStoreUtilities/Private/CookMetadataFiles.h
--- a/Engine/Source/Developer/IoStoreUtilities/Private/CookMetadataFiles.h
+++ b/Engine/Source/Developer/IoStoreUtilities/Private/CookMetadataFiles.h
@@ -5,6 +5,7 @@
 #include "Misc/EnumClassFlags.h"
 
 class FAssetRegistryState;
+class FCookedPackageStore;
 class FString;
 namespace UE::Cook { class FCookMetadataState; }
 
@@ -18,6 +19,7 @@ enum class ECookMetadataFiles
 ENUM_CLASS_FLAGS(ECookMetadataFiles);
 
 ECookMetadataFiles FindAndLoadMetadataFiles(
+	FCookedPackageStore* InPackageStore,
 	const FString& InCookedDir, ECookMetadataFiles InRequiredFiles, 
 	FAssetRegistryState& OutAssetRegistry, FString* OutAssetRegistryFileName /*optional, set on success*/,
 	UE::Cook::FCookMetadataState* OutCookMetadata, FString* OutCookMetadataFileName /*optional, set on success or need*/);
diff --git a/Engine/Source/Developer/IoStoreUtilities/Private/CookedPackageStore.cpp b/Engine/Source/Developer/IoStoreUtilities/Private/CookedPackageStore.cpp
--- a/Engine/Source/Developer/IoStoreUtilities/Private/CookedPackageStore.cpp
+++ b/Engine/Source/Developer/IoStoreUtilities/Private/CookedPackageStore.cpp
@@ -261,7 +261,7 @@ FIoStatus FCookedPackageStore::LoadChunkHashes()
 	else
 	{
 		FAssetRegistryState AssetRegistry;
-		if (FindAndLoadMetadataFiles(CookedDir, ECookMetadataFiles::None, AssetRegistry, nullptr, nullptr, nullptr) == ECookMetadataFiles::None)
+		if (FindAndLoadMetadataFiles(this, CookedDir, ECookMetadataFiles::None, AssetRegistry, nullptr, nullptr, nullptr) == ECookMetadataFiles::None)
 		{
 			return FIoStatus(EIoErrorCode::NotFound);
 		}
@@ -351,6 +351,11 @@ bool FCookedPackageStore::HasZenStoreClient() const
 	return ZenStoreClient.IsValid();
 }
 
+UE::FZenStoreHttpClient* FCookedPackageStore::GetZenStoreClient()
+{
+	return ZenStoreClient.Get();
+}
+
 TIoStatusOr<FIoBuffer> FCookedPackageStore::ReadChunk(const FIoChunkId& ChunkId)
 {
 	FIoReadOptions ReadOptions;
diff --git a/Engine/Source/Developer/IoStoreUtilities/Private/IoStoreUtilities.cpp b/Engine/Source/Developer/IoStoreUtilities/Private/IoStoreUtilities.cpp
--- a/Engine/Source/Developer/IoStoreUtilities/Private/IoStoreUtilities.cpp
+++ b/Engine/Source/Developer/IoStoreUtilities/Private/IoStoreUtilities.cpp
@@ -4611,6 +4611,7 @@ struct FIoStoreWriterInfo
 static bool DoAssetRegistryWritebackDuringStage(
 	EAssetRegistryWritebackMethod InMethod, 
 	bool bInWritePluginMetadata,
+	FCookedPackageStore* InPackageStore,
 	const FString& InCookedDir, 
 	bool bInCompressionEnabled,
 	TArray<TSharedPtr<IIoStoreWriter>>& InIoStoreWriters, 
@@ -4632,7 +4633,7 @@ static bool DoAssetRegistryWritebackDuringStage(
 	// We always need the cook metadata in order to update the corresponding hash.
 	EnumAddFlags(FilesNeeded, ECookMetadataFiles::CookMetadata);
 
-	if (FindAndLoadMetadataFiles(InCookedDir, FilesNeeded, AssetRegistry, &AssetRegistryFileName, &CookMetadata, &CookMetadataFileName) == ECookMetadataFiles::None)
+	if (FindAndLoadMetadataFiles(InPackageStore, InCookedDir, FilesNeeded, AssetRegistry, &AssetRegistryFileName, &CookMetadata, &CookMetadataFileName) == ECookMetadataFiles::None)
 	{
 		// already logged
 		return false;
@@ -5547,6 +5548,7 @@ int32 CreateTarget(const FIoStoreArguments& Arguments, const FIoStoreWriterSetti
 		DoAssetRegistryWritebackDuringStage(
 			Arguments.WriteBackMetadataToAssetRegistry, 
 			Arguments.bWritePluginSizeSummaryJsons, 
+			Arguments.PackageStore.Get(),
 			Arguments.CookedDir, 
 			GeneralIoWriterSettings.CompressionMethod != NAME_None, 
 			IoStoreWriters, 
diff --git a/Engine/Source/Developer/IoStoreUtilities/Private/ZenStoreHttpClient.cpp b/Engine/Source/Developer/IoStoreUtilities/Private/ZenStoreHttpClient.cpp
--- a/Engine/Source/Developer/IoStoreUtilities/Private/ZenStoreHttpClient.cpp
+++ b/Engine/Source/Developer/IoStoreUtilities/Private/ZenStoreHttpClient.cpp
@@ -128,6 +128,7 @@ FZenStoreHttpClient::TryCreateProject(FStringView InProjectId,
 		}
 	}
 
+	ProjectPath = WriteToString<128>("/prj/", InProjectId);
 	OplogPath = WriteToString<128>("/prj/", InProjectId, "/oplog/", InOplogId);
 	OplogNewEntryPath = WriteToString<128>("/prj/", InProjectId, "/oplog/", InOplogId, "/new");
 	OplogPrepNewEntryPath = WriteToString<128>("/prj/", InProjectId, "/oplog/", InOplogId, "/prep");
@@ -237,6 +238,7 @@ void FZenStoreHttpClient::InitializeReadOnly(FStringView InProjectId, FStringVie
 		UE::Zen::FZenScopedRequestPtr Request(RequestPool.Get());
 
 		OplogPath = WriteToString<128>("/prj/", InProjectId, "/oplog/", InOplogId);
+		ProjectPath = WriteToString<128>("/prj/", InProjectId);
 
 		TArray64<uint8> GetBuffer;
 		if (!Download(*Request, OplogPath, &GetBuffer, Zen::EContentType::CbObject) || Request->GetResponseCode() != 200)
@@ -592,6 +594,30 @@ TIoStatusOr<FIoBuffer> FZenStoreHttpClient::ReadOpLogUri(FStringBuilderBase& Chu
 	return FIoStatus(EIoErrorCode::NotFound);
 }
 
+TFuture<TIoStatusOr<FCbObject>> FZenStoreHttpClient::GetProjectInfo()
+{
+#if WITH_EDITOR
+	EAsyncExecution ThreadPool = EAsyncExecution::LargeThreadPool;
+#else
+	EAsyncExecution ThreadPool = EAsyncExecution::ThreadPool;
+#endif
+	return Async(ThreadPool, [this]
+	{
+		UE::Zen::FZenScopedRequestPtr Request(RequestPool.Get());
+
+		TArray64<uint8> GetBuffer;
+		if (Download(*Request, ProjectPath, &GetBuffer, Zen::EContentType::CbObject) && Request->GetResponseCode() == 200)
+		{
+			FCbObjectView Response(GetBuffer.GetData());
+			return TIoStatusOr<FCbObject>(FCbObject::Clone(Response));
+		}
+		else
+		{
+			return TIoStatusOr<FCbObject>(FIoStatus(EIoErrorCode::NotFound));
+		}
+	});
+}
+
 TFuture<TIoStatusOr<FCbObject>> FZenStoreHttpClient::GetOplog()
 {
 #if WITH_EDITOR
@@ -783,6 +809,11 @@ TIoStatusOr<uint64> FZenStoreHttpClient::AppendOp(FCbPackage OpEntry)
 	return TIoStatusOr<uint64>();
 }
 
+TFuture<TIoStatusOr<FCbObject>> FZenStoreHttpClient::GetProjectInfo()
+{
+	return TFuture<TIoStatusOr<FCbObject>>();
+}
+
 TFuture<TIoStatusOr<FCbObject>> FZenStoreHttpClient::GetOplog()
 {
 	return TFuture<TIoStatusOr<FCbObject>>();
