Change zen filesystem manifest filter to allow chunk manifest files to be included in the filesystem manifest and be incorporated into snapshots.

diff --git a/Engine/Source/Developer/IoStoreUtilities/Private/ZenFileSystemManifest.cpp b/Engine/Source/Developer/IoStoreUtilities/Private/ZenFileSystemManifest.cpp
--- a/Engine/Source/Developer/IoStoreUtilities/Private/ZenFileSystemManifest.cpp
+++ b/Engine/Source/Developer/IoStoreUtilities/Private/ZenFileSystemManifest.cpp
@@ -244,7 +244,7 @@ int32 FZenFileSystemManifest::Generate()
 	FFileFilter CookedMetadataFilter = FFileFilter()
 		.ExcludeDirectory(TEXT("ShaderLibrarySource"))
 		.ExcludeExtension(TEXT("manifest"));
-	AddFilesFromDirectory(TEXT("/{project}/Metadata"), FPaths::Combine(CookDirectory, FApp::GetProjectName(), "Metadata"), false, &CookedMetadataFilter);
+	AddFilesFromDirectory(TEXT("/{project}/Metadata"), FPaths::Combine(CookDirectory, FApp::GetProjectName(), "Metadata"), true, &CookedMetadataFilter);
 
 	FFileFilter ProjectSourceFilter = FFileFilter()
 		.IncludeExtension(TEXT("uproject"));
