Fix an issue where incorrect data could be pushed to the DDC during compression for container files.

#rnx

- The problem occurs because before we kick off the job to serialize a compressed entry to the DDC (so that we don't have to run compression on the same data each time the container is built) we were also kicking off a job to potentially encrypt the compressed data. It was possible for the encryption job to start modifying the data before we had serialized a copy for the DDC job, meaning that we'd be writing incorrect data to the DDC. Next time the container is built that bad data might be retrieved and the container would become corrupted.
- We now make sure that we finish setting up the DDC job before we kick off the encryption job.
- NOTE: The key for the compressed DDC entries has been changed to flush out the corrupted data already present.

diff --git a/Engine/Source/Developer/IoStoreUtilities/Private/IoStoreWriter.cpp b/Engine/Source/Developer/IoStoreUtilities/Private/IoStoreWriter.cpp
--- a/Engine/Source/Developer/IoStoreUtilities/Private/IoStoreWriter.cpp
+++ b/Engine/Source/Developer/IoStoreUtilities/Private/IoStoreWriter.cpp
@@ -43,7 +43,7 @@ TRACE_DECLARE_ATOMIC_INT_COUNTER(IoStoreDDCPutCount, TEXT("IoStoreWriter/DDCPutC
 
 static UE::DerivedData::FCacheBucket IoStoreDDCBucket = UE::DerivedData::FCacheBucket(ANSITEXTVIEW("IoStoreCompression"));
 static UE::DerivedData::ECachePolicy IoStoreDDCPolicy = UE::DerivedData::ECachePolicy::Default;
-static FStringView IoStoreDDCVersion = TEXTVIEW("985D3FAD-71D0-4758-A777-A910B49CC4BF");
+static FStringView IoStoreDDCVersion = TEXTVIEW("36EEC49B-E63B-498B-87D0-55FD11E4F9D6");
 
 struct FChunkBlock
 {
@@ -2314,8 +2314,7 @@ void FIoStoreWriterContextImpl::BeginEncryptionAndSigningThreadFunc()
 		{
 			FIoStoreWriteQueueEntry* Next = Entry->Next;
 			Entry->FinishCompressionBarrier.Wait();
-			TRACE_COUNTER_INCREMENT(IoStoreBeginEncryptionAndSigningCount);
-			Entry->Writer->BeginEncryptAndSign(Entry);
+			
 			if (Entry->bStoreCompressedDataInDDC)
 			{
 				TRACE_CPUPROFILER_EVENT_SCOPE(AddDDCPutRequest);
@@ -2332,6 +2331,12 @@ void FIoStoreWriterContextImpl::BeginEncryptionAndSigningThreadFunc()
 				}
 			}
 			DDCPutRequestDispatcher.DispatchPutRequests(HandleDDCPutResult);
+
+			// Must be done after we have serialized the compressed data for DDC as it can potentially modify the
+			// data stored by Entry!
+			TRACE_COUNTER_INCREMENT(IoStoreBeginEncryptionAndSigningCount);
+			Entry->Writer->BeginEncryptAndSign(Entry);
+
 			WriterQueue.Enqueue(Entry);
 			Entry = Next;
 		}
