Build fixes (compile shipping, and cook if zenserver not running) for last submission.

diff --git a/Engine/Source/Developer/IoStoreUtilities/Private/ZenCookArtifactReader.cpp b/Engine/Source/Developer/IoStoreUtilities/Private/ZenCookArtifactReader.cpp
--- a/Engine/Source/Developer/IoStoreUtilities/Private/ZenCookArtifactReader.cpp
+++ b/Engine/Source/Developer/IoStoreUtilities/Private/ZenCookArtifactReader.cpp
@@ -2,6 +2,7 @@
 
 #include "ZenCookArtifactReader.h"
 
+#include "Experimental/ZenServerInterface.h"
 #include "HAL/PlatformFileManager.h"
 #include "Misc/App.h"
 #include "Misc/Paths.h"
@@ -13,7 +14,9 @@ FZenCookArtifactReader::FZenCookArtifactReader(
 	const FString& InMetadataDirectoryPath, 
 	const ITargetPlatform* InTargetPlatform
 )
-	: ZenRootPath(InputPath)
+#if !UE_BUILD_SHIPPING
+	: ScopeZenService(MakeUnique<UE::Zen::FScopeZenService>())
+	, ZenRootPath(InputPath)
 	, StorageServerPlatformFile(IStorageServerClientModule::Get().TryCreateCustomPlatformFile(*InputPath, &FPlatformFileManager::Get().GetPlatformFile()))
 {
 	if (StorageServerPlatformFile)
@@ -21,6 +24,10 @@ FZenCookArtifactReader::FZenCookArtifactReader(
 		StorageServerPlatformFile->SetLowerLevel(nullptr);
 	}
 }
+#else
+{
+}
+#endif // !UE_BUILD_SHIPPING
 
 FZenCookArtifactReader::~FZenCookArtifactReader()
 {
@@ -28,6 +35,7 @@ FZenCookArtifactReader::~FZenCookArtifactReader()
 
 bool FZenCookArtifactReader::FileExists(const TCHAR* Filename)
 {
+#if !UE_BUILD_SHIPPING
 	if (StorageServerPlatformFile)
 	{
 		FString StandardFilename;
@@ -36,11 +44,14 @@ bool FZenCookArtifactReader::FileExists(const TCHAR* Filename)
 			return StorageServerPlatformFile->FileExists(*StandardFilename);
 		}
 	}
+#endif // !UE_BUILD_SHIPPING
+
 	return false;
 }
 
 int64 FZenCookArtifactReader::FileSize(const TCHAR* Filename)
 {
+#if !UE_BUILD_SHIPPING
 	if (StorageServerPlatformFile)
 	{
 		FString StandardFilename;
@@ -49,12 +60,14 @@ int64 FZenCookArtifactReader::FileSize(const TCHAR* Filename)
 			return StorageServerPlatformFile->FileSize(*StandardFilename);
 		}
 	}
+#endif // !UE_BUILD_SHIPPING
 
 	return -1;
 }
 
 IFileHandle* FZenCookArtifactReader::OpenRead(const TCHAR* Filename)
 {
+#if !UE_BUILD_SHIPPING
 	if (StorageServerPlatformFile)
 	{
 		FString StandardFilename;
@@ -63,11 +76,14 @@ IFileHandle* FZenCookArtifactReader::OpenRead(const TCHAR* Filename)
 			return StorageServerPlatformFile->OpenRead(*StandardFilename);
 		}
 	}
+#endif // !UE_BUILD_SHIPPING
+
 	return nullptr;
 }
 
 bool FZenCookArtifactReader::IterateDirectory(const TCHAR* Directory, IPlatformFile::FDirectoryVisitor& Visitor)
 {
+#if !UE_BUILD_SHIPPING
 	if (StorageServerPlatformFile)
 	{
 		FString StandardDirectory;
@@ -76,10 +92,12 @@ bool FZenCookArtifactReader::IterateDirectory(const TCHAR* Directory, IPlatformF
 			return StorageServerPlatformFile->IterateDirectory(Directory, Visitor);
 		}
 	}
+#endif // !UE_BUILD_SHIPPING
 
 	return false;
 }
 
+#if !UE_BUILD_SHIPPING
 bool FZenCookArtifactReader::MakeStorageServerPath(const TCHAR* Filename, FString& OutFilename) const
 {
 	OutFilename = Filename;
@@ -98,4 +116,5 @@ bool FZenCookArtifactReader::MakeStorageServerPath(const TCHAR* Filename, FStrin
 		return true;
 	}
 	return false;
-}
\ No newline at end of file
+}
+#endif // !UE_BUILD_SHIPPING
\ No newline at end of file
diff --git a/Engine/Source/Developer/IoStoreUtilities/Public/ZenCookArtifactReader.h b/Engine/Source/Developer/IoStoreUtilities/Public/ZenCookArtifactReader.h
--- a/Engine/Source/Developer/IoStoreUtilities/Public/ZenCookArtifactReader.h
+++ b/Engine/Source/Developer/IoStoreUtilities/Public/ZenCookArtifactReader.h
@@ -8,6 +8,8 @@
 
 class IStorageServerPlatformFile;
 
+namespace UE::Zen { class FScopeZenService; }
+
 class FZenCookArtifactReader
 	: public FCookArtifactReaderCommon
 {
@@ -23,8 +25,11 @@ public:
 
 	IOSTOREUTILITIES_API bool IterateDirectory(const TCHAR* Directory, IPlatformFile::FDirectoryVisitor& Visitor) override;
 private:
+#if !UE_BUILD_SHIPPING
 	bool MakeStorageServerPath(const TCHAR* Filename, FString& OutFilename) const;
 
+	TUniquePtr<UE::Zen::FScopeZenService> ScopeZenService;
 	FString ZenRootPath;
 	TUniquePtr<IStorageServerPlatformFile> StorageServerPlatformFile;
+#endif // !UE_BUILD_SHIPPING
 };
