Enable DiffContainer command from UnrealPak.exe

#rnx

diff --git a/Engine/Source/Developer/IoStoreUtilities/Private/IoStoreUtilities.cpp b/Engine/Source/Developer/IoStoreUtilities/Private/IoStoreUtilities.cpp
--- a/Engine/Source/Developer/IoStoreUtilities/Private/IoStoreUtilities.cpp
+++ b/Engine/Source/Developer/IoStoreUtilities/Private/IoStoreUtilities.cpp
@@ -8237,6 +8237,66 @@ static int32 Diff(
 	return 0;
 }
 
+bool DiffIoStoreContainer(const TCHAR* CmdLine)
+{
+	FString SourcePath, TargetPath, OutPath;
+	FKeyChain SourceKeyChain, TargetKeyChain;
+
+	if (!FParse::Value(CmdLine, TEXT("DiffContainer="), SourcePath))
+	{
+		UE_LOG(LogIoStore, Error, TEXT("Incorrect arguments. Expected: -DiffContainer=<Path> -Target=<Path>"));
+		return false;
+	}
+
+	if (!IFileManager::Get().DirectoryExists(*SourcePath))
+	{
+		UE_LOG(LogIoStore, Error, TEXT("Source directory '%s' doesn't exist"), *SourcePath);
+		return false;
+	}
+
+	if (!FParse::Value(CmdLine, TEXT("Target="), TargetPath))
+	{
+		UE_LOG(LogIoStore, Error, TEXT("Incorrect arguments. Expected: -DiffContainer=<Path> -Target=<Path>"));
+	}
+
+	if (!IFileManager::Get().DirectoryExists(*TargetPath))
+	{
+		UE_LOG(LogIoStore, Error, TEXT("Target directory '%s' doesn't exist"), *TargetPath);
+		return false;
+	}
+
+	FParse::Value(CmdLine, TEXT("DumpToFile="), OutPath);
+
+	FString CryptoKeysCacheFilename;
+	if (FParse::Value(CmdLine, TEXT("CryptoKeys="), CryptoKeysCacheFilename))
+	{
+		UE_LOG(LogIoStore, Display, TEXT("Parsing source crypto keys from '%s'"), *CryptoKeysCacheFilename);
+		KeyChainUtilities::LoadKeyChainFromFile(CryptoKeysCacheFilename, SourceKeyChain);
+	}
+
+	if (FParse::Value(CmdLine, TEXT("TargetCryptoKeys="), CryptoKeysCacheFilename))
+	{
+		UE_LOG(LogIoStore, Display, TEXT("Parsing target crypto keys from '%s'"), *CryptoKeysCacheFilename);
+		KeyChainUtilities::LoadKeyChainFromFile(CryptoKeysCacheFilename, TargetKeyChain);
+	}
+	else
+	{
+		TargetKeyChain = SourceKeyChain;
+	}
+
+	EChunkTypeFilter ChunkTypeFilter = EChunkTypeFilter::None;
+	if (FParse::Param(CmdLine, TEXT("FilterBulkData")))
+	{
+		ChunkTypeFilter = EChunkTypeFilter::BulkData;
+	}
+	else if (FParse::Param(FCommandLine::Get(), TEXT("FilterPackageData")))
+	{
+		ChunkTypeFilter = EChunkTypeFilter::PackageData;
+	}
+
+	return Diff(SourcePath, SourceKeyChain, TargetPath, TargetKeyChain, OutPath, ChunkTypeFilter) == 0;
+}
+
 bool LegacyDiffIoStoreContainers(const TCHAR* InContainerFilename1, const TCHAR* InContainerFilename2, bool bInLogUniques1, bool bInLogUniques2, const FKeyChain& InKeyChain1, const FKeyChain* InKeyChain2)
 {
 	TGuardValue<ELogTimes::Type> DisableLogTimes(GPrintLogTimes, ELogTimes::None);
diff --git a/Engine/Source/Developer/IoStoreUtilities/Public/IoStoreUtilities.h b/Engine/Source/Developer/IoStoreUtilities/Public/IoStoreUtilities.h
--- a/Engine/Source/Developer/IoStoreUtilities/Public/IoStoreUtilities.h
+++ b/Engine/Source/Developer/IoStoreUtilities/Public/IoStoreUtilities.h
@@ -28,6 +28,7 @@ UE_API bool LegacyListIoStoreContainer(
 UE_API bool ListIoStoreContainer(const TCHAR* CmdLine);
 UE_API bool ListIoStoreContainerBulkData(const TCHAR* CmdLine);
 
+UE_API bool DiffIoStoreContainer(const TCHAR* CmdLine);
 UE_API bool LegacyDiffIoStoreContainers(
 	const TCHAR* InContainerFilename1,
 	const TCHAR* InContainerFilename2,
diff --git a/Engine/Source/Developer/PakFileUtilities/Private/PakFileUtilities.cpp b/Engine/Source/Developer/PakFileUtilities/Private/PakFileUtilities.cpp
--- a/Engine/Source/Developer/PakFileUtilities/Private/PakFileUtilities.cpp
+++ b/Engine/Source/Developer/PakFileUtilities/Private/PakFileUtilities.cpp
@@ -5496,6 +5496,11 @@ bool ExecuteUnrealPak(const TCHAR* CmdLine)
 		{
 			return ListIoStoreContainerBulkData(CmdLine);
 		}
+
+		if (FParse::Value(CmdLine, TEXT("-DiffContainer="), IoStoreArg))
+		{
+			return DiffIoStoreContainer(CmdLine);
+		}
 	}
 
 	// Parse all the non-option arguments from the command line
