IncrementalCook: EndCook files: Write the entire set of files in the manifest to the op on every cook, rather than writing a delta from the previous incremental cook.
It's better to rewrite the entire op rather than doing deltas because an op is supposed to represent all of the data from the transformation. It's worse to rewrite the entire op because it will wastefully store the ~1000 strings about loose files every time we cook. We will fix that in the long term by breaking the endcook op up into smaller system-specific ops and not rewrite those ops if they have not changed.
#rnx

diff --git a/Engine/Source/Developer/IoStoreUtilities/Private/ZenStoreWriter.cpp b/Engine/Source/Developer/IoStoreUtilities/Private/ZenStoreWriter.cpp
--- a/Engine/Source/Developer/IoStoreUtilities/Private/ZenStoreWriter.cpp
+++ b/Engine/Source/Developer/IoStoreUtilities/Private/ZenStoreWriter.cpp
@@ -1376,19 +1376,14 @@ void FZenStoreWriter::CreateProjectMetaData(FCbPackage& Pkg, FCbWriter& PackageO
 {
 	// File Manifest
 	{
-		// Only append new file entries to the Oplog
-		
-		int32 NumEntries = ZenFileSystemManifest->NumEntries();
-		const int32 NumNewEntries = ZenFileSystemManifest->Generate();
-
-		if (NumNewEntries > 0)
+		ZenFileSystemManifest->Generate();
+		if (ZenFileSystemManifest->NumEntries() > 0)
 		{
 			TArrayView<FZenFileSystemManifest::FManifestEntry const> Entries = ZenFileSystemManifest->ManifestEntries();
-			TArrayView<FZenFileSystemManifest::FManifestEntry const> NewEntries = Entries.Slice(NumEntries, NumNewEntries);
 			
 			PackageObj.BeginArray("files");
 
-			for (const FZenFileSystemManifest::FManifestEntry& NewEntry : NewEntries)
+			for (const FZenFileSystemManifest::FManifestEntry& NewEntry : Entries)
 			{
 				FCbObjectId FileOid = ToObjectId(NewEntry.FileChunkId);
 
