IncrementalCook: ReferencedSet: Change storage of the ReferencedSet in zenserver to be in a separate op rather than one of the files in the EndCook op. This allows zenserver to find the ReferencedSet data with less coupling of pathname information between it and the Unreal client.
#rnx

diff --git a/Engine/Source/Developer/IoStoreUtilities/Private/ZenFileSystemManifest.cpp b/Engine/Source/Developer/IoStoreUtilities/Private/ZenFileSystemManifest.cpp
--- a/Engine/Source/Developer/IoStoreUtilities/Private/ZenFileSystemManifest.cpp
+++ b/Engine/Source/Developer/IoStoreUtilities/Private/ZenFileSystemManifest.cpp
@@ -13,10 +13,11 @@
 #include "Misc/DataDrivenPlatformInfoRegistry.h"
 #include "Misc/PathViews.h"
 #include "Settings/ProjectPackagingSettings.h"
+#include "UObject/ICookInfo.h"
 
 DEFINE_LOG_CATEGORY_STATIC(LogZenFileSystemManifest, Display, All);
 
-const FZenFileSystemManifest::FManifestEntry FZenFileSystemManifest::InvalidEntry = FManifestEntry();
+const FZenFileSystemManifestEntry FZenFileSystemManifest::InvalidEntry = FZenFileSystemManifestEntry();
 
 FZenFileSystemManifest::FZenFileSystemManifest(const ITargetPlatform& InTargetPlatform, FString InCookDirectory)
 	: TargetPlatform(InTargetPlatform)
@@ -26,6 +27,7 @@ FZenFileSystemManifest::FZenFileSystemManifest(const ITargetPlatform& InTargetPl
 
 	ServerRoot = PlatformFile.ConvertToAbsolutePathForExternalAppForRead(*FPaths::RootDir());
 	FPaths::NormalizeDirectoryName(ServerRoot);
+	ReferencedSetClientPath = FString::Printf(TEXT("/{project}/Metadata/%s"), UE::Cook::GetReferencedSetFilename());
 }
 
 void FZenFileSystemManifest::GetExtensionDirs(TArray<FString>& OutExtensionDirs, const TCHAR* BaseDir, const TCHAR* SubDir, const TArray<FString>& PlatformDirectoryNames)
@@ -525,7 +527,7 @@ int32 FZenFileSystemManifest::Generate()
 	return CurrentEntryCount - PreviousEntryCount;
 }
 
-const FZenFileSystemManifest::FManifestEntry& FZenFileSystemManifest::CreateManifestEntry(const FString& Filename)
+const FZenFileSystemManifestEntry& FZenFileSystemManifest::CreateManifestEntry(const FString& Filename)
 {
 	const FString FullFilename = FPaths::ConvertRelativePathToFull(Filename);
 
@@ -535,7 +537,8 @@ const FZenFileSystemManifest::FManifestEntry& FZenFileSystemManifest::CreateMani
 	CookedEngineDirectoryTrailingSeparator.Append(CookedEngineDirectory);
 	CookedEngineDirectoryTrailingSeparator.AppendChar(TEXT('/'));
 
-	auto AddEntry = [this, &FullFilename](const FString& ClientDirectory, const FString& LocalDirectory) -> const FManifestEntry&
+	auto AddEntry = [this, &FullFilename](const FString& ClientDirectory, const FString& LocalDirectory)
+		-> const FZenFileSystemManifestEntry&
 	{
 		FStringView RelativePath = FullFilename;
 		RelativePath.RightChopInline(LocalDirectory.Len() + 1);
@@ -570,7 +573,8 @@ const FZenFileSystemManifest::FManifestEntry& FZenFileSystemManifest::CreateMani
 	return InvalidEntry;
 }
 
-const FZenFileSystemManifest::FManifestEntry& FZenFileSystemManifest::AddManifestEntry(const FIoChunkId& FileChunkId, FString ServerPath, FString ClientPath)
+const FZenFileSystemManifestEntry& FZenFileSystemManifest::AddManifestEntry(const FIoChunkId& FileChunkId,
+	FString ServerPath, FString ClientPath)
 {
 	check(ServerPath.Len() > 0 && ClientPath.Len() > 0);
 
@@ -592,14 +596,21 @@ const FZenFileSystemManifest::FManifestEntry& FZenFileSystemManifest::AddManifes
 
 	EntryIndex = Entries.Num();
 
-	FManifestEntry Entry;
+	FZenFileSystemManifestEntry Entry;
 	Entry.ServerPath = MoveTemp(ServerPath);
 	Entry.ClientPath = MoveTemp(ClientPath);
 	Entry.FileChunkId = FileChunkId;
 
-	Entries.Add(MoveTemp(Entry));
-
-	return Entries[EntryIndex];
+	if (Entry.ClientPath == ReferencedSetClientPath)
+	{
+		ReferencedSet.Emplace(MoveTemp(Entry));
+		return *ReferencedSet;
+	}
+	else
+	{
+		Entries.Add(MoveTemp(Entry));
+		return Entries[EntryIndex];
+	}
 }
 
 bool FZenFileSystemManifest::Save(const TCHAR* Filename)
@@ -611,15 +622,15 @@ bool FZenFileSystemManifest::Save(const TCHAR* Filename)
 	CsvLines.Add(TEXT("FileId, ServerPath, ClientPath"));
 
 	TStringBuilder<2048> Sb;
-	TArray<const FManifestEntry*> SortedEntries;
+	TArray<const FZenFileSystemManifestEntry*> SortedEntries;
 	SortedEntries.Reserve(Entries.Num());
-	for (const FManifestEntry& Entry : Entries)
+	for (const FZenFileSystemManifestEntry& Entry : Entries)
 	{
 		SortedEntries.Add(&Entry);
 	}
-	Algo::Sort(SortedEntries, [](const FManifestEntry* A, const FManifestEntry* B) { return A->ClientPath < B->ClientPath; });
+	Algo::Sort(SortedEntries, [](const FZenFileSystemManifestEntry* A, const FZenFileSystemManifestEntry* B) { return A->ClientPath < B->ClientPath; });
 
-	for (const FManifestEntry* Entry : SortedEntries)
+	for (const FZenFileSystemManifestEntry* Entry : SortedEntries)
 	{
 		Sb.Reset();
 		Sb << Entry->FileChunkId << TEXT(", ") << Entry->ServerPath << TEXT(", ") << *Entry->ClientPath;
diff --git a/Engine/Source/Developer/IoStoreUtilities/Private/ZenFileSystemManifest.h b/Engine/Source/Developer/IoStoreUtilities/Private/ZenFileSystemManifest.h
--- a/Engine/Source/Developer/IoStoreUtilities/Private/ZenFileSystemManifest.h
+++ b/Engine/Source/Developer/IoStoreUtilities/Private/ZenFileSystemManifest.h
@@ -4,30 +4,35 @@
 
 #include "CoreMinimal.h"
 #include "IO/IoDispatcher.h"
+#include "Misc/Optional.h"
 
 class ITargetPlatform;
 class FArchive;
 
+struct FZenFileSystemManifestEntry
+{
+	FString ServerPath;
+	FString ClientPath;
+	FIoChunkId FileChunkId;
+};
+
 class FZenFileSystemManifest
 {
 public:
-
-	struct FManifestEntry
-	{
-		FString ServerPath;
-		FString ClientPath;
-		FIoChunkId FileChunkId;
-	};
-	
 	FZenFileSystemManifest(const ITargetPlatform& InTargetPlatform, FString InCookDirectory);
 	
+	/**
+	 * Scan the local Cooked folder for files on disk, create and append Entries for any files not already
+	 * present in the manifest. Files in the manifest do not contain any payload data, so no update is necessary
+	 * for Entries that already exist. Returns the number of new entries.
+	 */
 	int32 Generate();
 
-	const FManifestEntry& CreateManifestEntry(const FString& Filename);
+	const FZenFileSystemManifestEntry& CreateManifestEntry(const FString& Filename);
 	
-	const FManifestEntry& AddManifestEntry(const FIoChunkId& FileChunkId, FString ServerPath, FString ClientPath);
+	const FZenFileSystemManifestEntry& AddManifestEntry(const FIoChunkId& FileChunkId, FString ServerPath, FString ClientPath);
 
-	TArrayView<const FManifestEntry> ManifestEntries() const
+	TArrayView<const FZenFileSystemManifestEntry> ManifestEntries() const
 	{
 		return Entries;
 	}
@@ -44,6 +49,17 @@ public:
 		return ServerRoot;
 	}
 
+	/**
+	 * Return the ManifestEntry that was found for the ReferencedSet.txt artifact left by the cooker in the Metadata
+	 * directory. The ReferencedSet reports which packages in the oplog were cooked in the most recent incremental
+	 * cook. ZenServer needs to read it, so we report it separately to ZenServer rather than reporting it with the
+	 * rest of the cook artifacts reported in the EndCook op.
+	 */
+	TOptional<FZenFileSystemManifestEntry> GetReferencedSet() const
+	{
+		return ReferencedSet;
+	}
+
 private:
 	static void GetExtensionDirs(TArray<FString>& OutExtensionDirs, const TCHAR* BaseDir, const TCHAR* SubDir, const TArray<FString>& PlatformDirectoryNames);
 
@@ -51,7 +67,9 @@ private:
 	FString CookDirectory;
 	FString ServerRoot;
 	TMap<FString, int32> ServerPathToEntry;
-	TArray<FManifestEntry> Entries;
-	
-	static const FManifestEntry InvalidEntry;
+	TArray<FZenFileSystemManifestEntry> Entries;
+	TOptional<FZenFileSystemManifestEntry> ReferencedSet;
+	FString ReferencedSetClientPath;
+
+	static const FZenFileSystemManifestEntry InvalidEntry;
 };
diff --git a/Engine/Source/Developer/IoStoreUtilities/Private/ZenStoreWriter.cpp b/Engine/Source/Developer/IoStoreUtilities/Private/ZenStoreWriter.cpp
--- a/Engine/Source/Developer/IoStoreUtilities/Private/ZenStoreWriter.cpp
+++ b/Engine/Source/Developer/IoStoreUtilities/Private/ZenStoreWriter.cpp
@@ -30,6 +30,7 @@
 #include "Serialization/JsonWriter.h"
 #include "Serialization/LargeMemoryWriter.h" 
 #include "SocketSubsystem.h"
+#include "UObject/ICookInfo.h"
 #include "UObject/SavePackage.h"
 #include "ZenCookArtifactReader.h"
 #include "ZenFileSystemManifest.h"
@@ -390,7 +391,7 @@ void FZenStoreWriter::WriteAdditionalFile(const FAdditionalFileInfo& Info, const
 		return FCompressedBuffer::Compress(FSharedBuffer::MakeView(FileData.GetView()), Compressor, CompressionLevel);
 	});
 
-	const FZenFileSystemManifest::FManifestEntry& ManifestEntry = ZenFileSystemManifest->CreateManifestEntry(Info.Filename);
+	const FZenFileSystemManifestEntry& ManifestEntry = ZenFileSystemManifest->CreateManifestEntry(Info.Filename);
 	FileEntry.Info					= Info;
 	FileEntry.Info.ChunkId			= ManifestEntry.FileChunkId;
 	FileEntry.ZenManifestServerPath = ManifestEntry.ServerPath;
@@ -602,20 +603,8 @@ void FZenStoreWriter::BeginCook(const FCookInfo& Info)
 
 	if (CookMode == ICookedPackageWriter::FCookInfo::CookOnTheFlyMode)
 	{
-		FCbPackage Pkg;
-		FCbWriter PackageObj;
-		
-		PackageObj.BeginObject();
-		PackageObj << "key" << "CookOnTheFly";
-
-		CreateProjectMetaData(Pkg, PackageObj);
-
-		PackageObj.EndObject();
-		FCbObject Obj = PackageObj.Save().AsObject();
-
-		Pkg.SetObject(Obj);
-
-		TIoStatusOr<uint64> Status = HttpClient->AppendOp(Pkg);
+		ZenFileSystemManifest->Generate();
+		TIoStatusOr<uint64> Status = HttpClient->AppendOp(CreateProjectMetaDataOpPackage("CookOnTheFly"));
 		UE_CLOG(!Status.IsOk(), LogZenStoreWriter, Fatal, TEXT("Failed to append OpLog"));
 	}
 
@@ -648,21 +637,16 @@ void FZenStoreWriter::EndCook(const FCookInfo& Info)
 
 	if (!Info.bWorkerOnSharedSandbox)
 	{
-		FCbPackage Pkg;
-		FCbWriter PackageObj;
-	
-		PackageObj.BeginObject();
-		PackageObj << "key" << "EndCook";
-
-		CreateProjectMetaData(Pkg, PackageObj);
-
-		PackageObj.EndObject();
-		FCbObject Obj = PackageObj.Save().AsObject();
-
-		Pkg.SetObject(Obj);
-
-		TIoStatusOr<uint64> Status = HttpClient->EndBuildPass(Pkg);
-		UE_CLOG(!Status.IsOk(), LogZenStoreWriter, Fatal, TEXT("Failed to append OpLog and end the build pass"));
+		ZenFileSystemManifest->Generate();
+		if (FCbPackage ReferencedSetPackage = CreateReferencedSetOpPackage())
+		{
+			TIoStatusOr<uint64> Status = HttpClient->AppendOp(MoveTemp(ReferencedSetPackage));
+			UE_CLOG(!Status.IsOk(), LogZenStoreWriter, Fatal, TEXT("Failed to append ReferencedSetOp"));
+		}
+		{
+			TIoStatusOr<uint64> Status = HttpClient->EndBuildPass(CreateProjectMetaDataOpPackage("EndCook"));
+			UE_CLOG(!Status.IsOk(), LogZenStoreWriter, Fatal, TEXT("Failed to append OpLog and end the build pass"));
+		}
 	}
 
 	UE_LOG(LogZenStoreWriter, Display, TEXT("Output:\t%d Public runtime script objects"), PackageStoreOptimizer->GetTotalScriptObjectCount());
@@ -1372,57 +1356,55 @@ FZenStoreWriter::FPendingPackageState& FZenStoreWriter::AddPendingPackage(const
 	return *Package;
 }
 
+FCbPackage FZenStoreWriter::CreateReferencedSetOpPackage()
+{
+	TOptional<FZenFileSystemManifestEntry> ReferencedSet = ZenFileSystemManifest->GetReferencedSet();
+	if (!ReferencedSet)
+	{
+		return FCbPackage();
+	}
+	FCbPackage Pkg;
+	FCbWriter PackageObj;
+
+	PackageObj.BeginObject();
+	PackageObj << "key" << UE::Cook::GetReferencedSetOpName();
+	PackageObj.BeginArray("files");
+	WriteManifestEntryToPackageWriter(Pkg, PackageObj, *ReferencedSet);
+	PackageObj.EndArray();
+	PackageObj.EndObject();
+	FCbObject Obj = PackageObj.Save().AsObject();
+
+	Pkg.SetObject(Obj);
+	return Pkg;
+}
+
+FCbPackage FZenStoreWriter::CreateProjectMetaDataOpPackage(const ANSICHAR* MetadataOplogKeyName)
+{
+	FCbPackage Pkg;
+	FCbWriter PackageObj;
+
+	PackageObj.BeginObject();
+	PackageObj << "key" << MetadataOplogKeyName;
+	CreateProjectMetaData(Pkg, PackageObj);
+	PackageObj.EndObject();
+
+	Pkg.SetObject(PackageObj.Save().AsObject());
+	return Pkg;
+}
+
 void FZenStoreWriter::CreateProjectMetaData(FCbPackage& Pkg, FCbWriter& PackageObj)
 {
 	// File Manifest
 	{
-		ZenFileSystemManifest->Generate();
 		if (ZenFileSystemManifest->NumEntries() > 0)
 		{
-			TArrayView<FZenFileSystemManifest::FManifestEntry const> Entries = ZenFileSystemManifest->ManifestEntries();
-			
-			PackageObj.BeginArray("files");
+			TArrayView<FZenFileSystemManifestEntry const> Entries = ZenFileSystemManifest->ManifestEntries();
 
-			for (const FZenFileSystemManifest::FManifestEntry& NewEntry : Entries)
+			PackageObj.BeginArray("files");
+			for (const FZenFileSystemManifestEntry& NewEntry : Entries)
 			{
-				FCbObjectId FileOid = ToObjectId(NewEntry.FileChunkId);
-
-				if (IsLocalConnection)
-				{
-					PackageObj.BeginObject();
-					PackageObj << "id" << FileOid;
-					PackageObj << "data" << FIoHash::Zero;
-					PackageObj << "serverpath" << NewEntry.ServerPath;
-					PackageObj << "clientpath" << NewEntry.ClientPath;
-					PackageObj.EndObject();
-				}
-				else
-				{
-					const FString AbsPath = ZenFileSystemManifest->ServerRootPath() / NewEntry.ServerPath;
-					TUniquePtr<FArchive> Reader(IFileManager::Get().CreateFileReader(*AbsPath, 0));
-					if (Reader)
-					{
-						int64 TotalSize = Reader->TotalSize();
-						if (TotalSize > 0)
-						{
-							FIoBuffer FileBuffer(TotalSize);
-							Reader->Serialize(FileBuffer.GetData(), TotalSize);
-							bool Success = Reader->Close();
-							FCbAttachment FileAttachment = CreateAttachment(MoveTemp(FileBuffer));
-
-							PackageObj.BeginObject();
-							PackageObj << "id" << FileOid;
-							PackageObj << "data" << FileAttachment;
-							PackageObj << "serverpath" << NewEntry.ServerPath;
-							PackageObj << "clientpath" << NewEntry.ClientPath;
-							PackageObj.EndObject();
-
-							Pkg.AddAttachment(FileAttachment);
-						}
-					}
-				}
+				WriteManifestEntryToPackageWriter(Pkg, PackageObj, NewEntry);
 			}
-
 			PackageObj.EndArray();
 		}
 
@@ -1452,6 +1434,47 @@ void FZenStoreWriter::CreateProjectMetaData(FCbPackage& Pkg, FCbWriter& PackageO
 	}
 }
 
+void FZenStoreWriter::WriteManifestEntryToPackageWriter(FCbPackage& Pkg, FCbWriter& PackageObj,
+	const FZenFileSystemManifestEntry& Entry) const
+{
+	FCbObjectId FileOid = ToObjectId(Entry.FileChunkId);
+
+	if (IsLocalConnection)
+	{
+		PackageObj.BeginObject();
+		PackageObj << "id" << FileOid;
+		PackageObj << "data" << FIoHash::Zero;
+		PackageObj << "serverpath" << Entry.ServerPath;
+		PackageObj << "clientpath" << Entry.ClientPath;
+		PackageObj.EndObject();
+	}
+	else
+	{
+		const FString AbsPath = ZenFileSystemManifest->ServerRootPath() / Entry.ServerPath;
+		TUniquePtr<FArchive> Reader(IFileManager::Get().CreateFileReader(*AbsPath, 0));
+		if (Reader)
+		{
+			int64 TotalSize = Reader->TotalSize();
+			if (TotalSize > 0)
+			{
+				FIoBuffer FileBuffer(TotalSize);
+				Reader->Serialize(FileBuffer.GetData(), TotalSize);
+				bool Success = Reader->Close();
+				FCbAttachment FileAttachment = CreateAttachment(MoveTemp(FileBuffer));
+
+				PackageObj.BeginObject();
+				PackageObj << "id" << FileOid;
+				PackageObj << "data" << FileAttachment;
+				PackageObj << "serverpath" << Entry.ServerPath;
+				PackageObj << "clientpath" << Entry.ClientPath;
+				PackageObj.EndObject();
+
+				Pkg.AddAttachment(FileAttachment);
+			}
+		}
+	}
+}
+
 void FZenStoreWriter::BroadcastCommit(IPackageStoreWriter::FCommitEventArgs& EventArgs)
 {
 	FScopeLock CommitEventLock(&CommitEventCriticalSection);
@@ -1478,7 +1501,7 @@ void FZenStoreWriter::BroadcastMarkUpToDate(IPackageStoreWriter::FMarkUpToDateEv
 	}
 }
 
-FCbAttachment FZenStoreWriter::CreateAttachment(FSharedBuffer AttachmentData)
+FCbAttachment FZenStoreWriter::CreateAttachment(FSharedBuffer AttachmentData) const
 {
 	check(AttachmentData.GetSize() > 0);
 	FCompressedBuffer CompressedBuffer = FCompressedBuffer::Compress(AttachmentData, Compressor, CompressionLevel);
@@ -1486,7 +1509,7 @@ FCbAttachment FZenStoreWriter::CreateAttachment(FSharedBuffer AttachmentData)
 	return FCbAttachment(CompressedBuffer);
 }
 
-FCbAttachment FZenStoreWriter::CreateAttachment(FIoBuffer AttachmentData)
+FCbAttachment FZenStoreWriter::CreateAttachment(FIoBuffer AttachmentData) const
 {
 	return CreateAttachment(IoBufferToSharedBuffer(AttachmentData));
 }
diff --git a/Engine/Source/Developer/IoStoreUtilities/Public/ZenStoreWriter.h b/Engine/Source/Developer/IoStoreUtilities/Public/ZenStoreWriter.h
--- a/Engine/Source/Developer/IoStoreUtilities/Public/ZenStoreWriter.h
+++ b/Engine/Source/Developer/IoStoreUtilities/Public/ZenStoreWriter.h
@@ -49,6 +49,7 @@ class FPackageStoreOptimizer;
 class FPackageStorePackage;
 class FZenCookArtifactReader;
 class FZenFileSystemManifest;
+struct FZenFileSystemManifestEntry;
 
 /** 
  * A PackageStoreWriter that saves cooked packages for use by IoStore, and stores them in the Zen storage service.
@@ -222,14 +223,18 @@ private:
 		return PendingPackages.FindAndRemoveChecked(PackageName);
 	}
 
+	FCbPackage CreateReferencedSetOpPackage();
+	FCbPackage CreateProjectMetaDataOpPackage(const ANSICHAR* MetadataOplogKeyName);
 	void CreateProjectMetaData(FCbPackage& Pkg, FCbWriter& PackageObj);
+	void WriteManifestEntryToPackageWriter(FCbPackage& Pkg, FCbWriter& PackageObj,
+		const FZenFileSystemManifestEntry& Entry) const;
 	void BroadcastCommit(IPackageStoreWriter::FCommitEventArgs& EventArgs);
 	void BroadcastMarkUpToDate(IPackageStoreWriter::FMarkUpToDateEventArgs& EventArgs);
 	struct FZenCommitInfo;
 
 	void CommitPackageInternal(FZenCommitInfo&& CommitInfo);
-	FCbAttachment CreateAttachment(FSharedBuffer Buffer);
-	FCbAttachment CreateAttachment(FIoBuffer Buffer);
+	FCbAttachment CreateAttachment(FSharedBuffer Buffer) const;
+	FCbAttachment CreateAttachment(FIoBuffer Buffer) const;
 
 	FCriticalSection								PackagesCriticalSection;
 	TMap<FName, TUniquePtr<FPendingPackageState>>	PendingPackages;
diff --git a/Engine/Source/Runtime/CoreUObject/Private/UObject/ICookInfo.cpp b/Engine/Source/Runtime/CoreUObject/Private/UObject/ICookInfo.cpp
--- a/Engine/Source/Runtime/CoreUObject/Private/UObject/ICookInfo.cpp
+++ b/Engine/Source/Runtime/CoreUObject/Private/UObject/ICookInfo.cpp
@@ -58,6 +58,11 @@ const TCHAR* GetReferencedSetFilename()
 	return TEXT("ReferencedSet.txt");
 }
 
+const TCHAR* GetReferencedSetOpName()
+{
+	return TEXT("ReferencedSet");
+}
+
 }
 
 static thread_local ECookLoadType GCookLoadType = ECookLoadType::Unexpected;
diff --git a/Engine/Source/Runtime/CoreUObject/Public/UObject/ICookInfo.h b/Engine/Source/Runtime/CoreUObject/Public/UObject/ICookInfo.h
--- a/Engine/Source/Runtime/CoreUObject/Public/UObject/ICookInfo.h
+++ b/Engine/Source/Runtime/CoreUObject/Public/UObject/ICookInfo.h
@@ -192,6 +192,8 @@ public:
  * It does not include packages from the base game for a DLC cook.
  */
 COREUOBJECT_API const TCHAR* GetReferencedSetFilename();
+/** Return the name of the op used to store the ReferencedSet in the zenserver oplog. */
+COREUOBJECT_API const TCHAR* GetReferencedSetOpName();
 
 } // namespace UE::Cook
 
