Cooker: Fix UnrealPak compilation failure; the CookInfo strings are only available if WITH_EDITOR.
#rnx

diff --git a/Engine/Source/Developer/IoStoreUtilities/Private/ZenFileSystemManifest.cpp b/Engine/Source/Developer/IoStoreUtilities/Private/ZenFileSystemManifest.cpp
--- a/Engine/Source/Developer/IoStoreUtilities/Private/ZenFileSystemManifest.cpp
+++ b/Engine/Source/Developer/IoStoreUtilities/Private/ZenFileSystemManifest.cpp
@@ -27,7 +27,9 @@ FZenFileSystemManifest::FZenFileSystemManifest(const ITargetPlatform& InTargetPl
 
 	ServerRoot = PlatformFile.ConvertToAbsolutePathForExternalAppForRead(*FPaths::RootDir());
 	FPaths::NormalizeDirectoryName(ServerRoot);
+#if WITH_EDITOR
 	ReferencedSetClientPath = FString::Printf(TEXT("/{project}/Metadata/%s"), UE::Cook::GetReferencedSetFilename());
+#endif
 }
 
 void FZenFileSystemManifest::GetExtensionDirs(TArray<FString>& OutExtensionDirs, const TCHAR* BaseDir, const TCHAR* SubDir, const TArray<FString>& PlatformDirectoryNames)
@@ -601,12 +603,14 @@ const FZenFileSystemManifestEntry& FZenFileSystemManifest::AddManifestEntry(cons
 	Entry.ClientPath = MoveTemp(ClientPath);
 	Entry.FileChunkId = FileChunkId;
 
+#if WITH_EDITOR
 	if (Entry.ClientPath == ReferencedSetClientPath)
 	{
 		ReferencedSet.Emplace(MoveTemp(Entry));
 		return *ReferencedSet;
 	}
 	else
+#endif
 	{
 		Entries.Add(MoveTemp(Entry));
 		return Entries[EntryIndex];
diff --git a/Engine/Source/Developer/IoStoreUtilities/Private/ZenFileSystemManifest.h b/Engine/Source/Developer/IoStoreUtilities/Private/ZenFileSystemManifest.h
--- a/Engine/Source/Developer/IoStoreUtilities/Private/ZenFileSystemManifest.h
+++ b/Engine/Source/Developer/IoStoreUtilities/Private/ZenFileSystemManifest.h
@@ -69,7 +69,9 @@ private:
 	TMap<FString, int32> ServerPathToEntry;
 	TArray<FZenFileSystemManifestEntry> Entries;
 	TOptional<FZenFileSystemManifestEntry> ReferencedSet;
+#if WITH_EDITOR
 	FString ReferencedSetClientPath;
+#endif
 
 	static const FZenFileSystemManifestEntry InvalidEntry;
 };
diff --git a/Engine/Source/Developer/IoStoreUtilities/Private/ZenStoreWriter.cpp b/Engine/Source/Developer/IoStoreUtilities/Private/ZenStoreWriter.cpp
--- a/Engine/Source/Developer/IoStoreUtilities/Private/ZenStoreWriter.cpp
+++ b/Engine/Source/Developer/IoStoreUtilities/Private/ZenStoreWriter.cpp
@@ -1358,6 +1358,7 @@ FZenStoreWriter::FPendingPackageState& FZenStoreWriter::AddPendingPackage(const
 
 FCbPackage FZenStoreWriter::CreateReferencedSetOpPackage()
 {
+#if WITH_EDITOR
 	TOptional<FZenFileSystemManifestEntry> ReferencedSet = ZenFileSystemManifest->GetReferencedSet();
 	if (!ReferencedSet)
 	{
@@ -1376,6 +1377,9 @@ FCbPackage FZenStoreWriter::CreateReferencedSetOpPackage()
 
 	Pkg.SetObject(Obj);
 	return Pkg;
+#else // else !WITH_EDITOR
+	return FCbPackage();
+#endif
 }
 
 FCbPackage FZenStoreWriter::CreateProjectMetaDataOpPackage(const ANSICHAR* MetadataOplogKeyName)
