IncrementalCook: Fix crash due to unsynchronized multithreaded access to CookedPackagesInfo. It is written to from async thread in CommitPackageInternal and is read from in main thread by CookRequestCluster when cooking incrementally.
#rnx

diff --git a/Engine/Source/Developer/IoStoreUtilities/Private/ZenStoreWriter.cpp b/Engine/Source/Developer/IoStoreUtilities/Private/ZenStoreWriter.cpp
--- a/Engine/Source/Developer/IoStoreUtilities/Private/ZenStoreWriter.cpp
+++ b/Engine/Source/Developer/IoStoreUtilities/Private/ZenStoreWriter.cpp
@@ -503,6 +503,7 @@ void FZenStoreWriter::Initialize(const FCookInfo& Info)
 
 					if (Oplog["entries"])
 					{
+						FWriteScopeLock _(EntriesLock);
 						for (FCbField& OplogEntry : Oplog["entries"].AsArray())
 						{
 							FCbObject OplogObj = OplogEntry.AsObject();
@@ -807,12 +808,9 @@ void FZenStoreWriter::CommitPackageInternal(FZenCommitInfo&& ZenCommitInfo)
 		FMD5 PkgHashGen;
 		FCbPackage OplogEntry;
 		// Commit attachments
-		FOplogCookInfo& CookInfo = CookedPackagesInfo[CommitEventArgs.EntryIndex];
-		CookInfo.bUpToDate = true;
-
 		const int32 NumAttachments = CommitInfo.Attachments.Num();
+		TArray<FOplogCookInfo::FAttachment> CookInfoAttachments;
 		TArray<FCbAttachment, TInlineAllocator<2>> CbAttachments;
-		CookInfo.Attachments.Empty(NumAttachments);
 
 		if (NumAttachments)
 		{
@@ -830,7 +828,7 @@ void FZenStoreWriter::CommitPackageInternal(FZenCommitInfo&& ZenCommitInfo)
 			});
 			
 			CbAttachments.Reserve(NumAttachments);
-			CookInfo.Attachments.Reserve(NumAttachments);
+			CookInfoAttachments.Reserve(NumAttachments);
 			
 			for (const FCommitAttachmentInfo* AttachmentInfo : SortedAttachments)
 			{
@@ -838,7 +836,7 @@ void FZenStoreWriter::CommitPackageInternal(FZenCommitInfo&& ZenCommitInfo)
 				const FCbAttachment& CbAttachment = CbAttachments.Add_GetRef(CreateAttachment(AttachmentInfo->Value.GetBuffer().ToShared()));
 				OplogEntry.AddAttachment(CbAttachment);
 				
-				CookInfo.Attachments.Add(FOplogCookInfo::FAttachment
+				CookInfoAttachments.Add(FOplogCookInfo::FAttachment
 				{
 					UE::FZenStoreHttpClient::FindOrAddAttachmentId(AttachmentInfo->Key), CbAttachment.GetHash()
 				});
@@ -963,13 +961,20 @@ void FZenStoreWriter::CommitPackageInternal(FZenCommitInfo&& ZenCommitInfo)
 		for (int32 Index = 0; Index < NumAttachments; ++Index)
 		{
 			FCbAttachment& CbAttachment = CbAttachments[Index];
-			FOplogCookInfo::FAttachment& CookInfoAttachment = CookInfo.Attachments[Index];
+			FOplogCookInfo::FAttachment& CookInfoAttachment = CookInfoAttachments[Index];
 			OplogEntryDesc << CookInfoAttachment.Key << CbAttachment;
 		}
 
 		OplogEntryDesc.EndObject();
 		OplogEntry.SetObject(OplogEntryDesc.Save().AsObject());
-		
+
+		{
+			FWriteScopeLock _(EntriesLock);
+			FOplogCookInfo& CookInfo = CookedPackagesInfo[CommitEventArgs.EntryIndex];
+			CookInfo.bUpToDate = true;
+			CookInfo.Attachments = MoveTemp(CookInfoAttachments);
+		}
+
 		TIoStatusOr<uint64> Status = HttpClient->AppendOp(MoveTemp(OplogEntry));
 		UE_CLOG(!Status.IsOk(), LogZenStoreWriter, Error, TEXT("Failed to commit oplog entry '%s' to Zen"), *CommitInfo.PackageName.ToString());
 	}
@@ -1043,36 +1048,37 @@ TUniquePtr<FAssetRegistryState> FZenStoreWriter::LoadPreviousAssetRegistry()
 	PreviousState->Load(*Reader);
 
 	TSet<FName> RemoveSet;
-	const TMap<FName, const FAssetPackageData*>& PreviousStatePackages = PreviousState->GetAssetPackageDataMap(); 
-	for (const TPair<FName, const FAssetPackageData*>& Pair : PreviousStatePackages)
+	TArray<FName> RemoveArray;
+	const TMap<FName, const FAssetPackageData*>& PreviousStatePackages = PreviousState->GetAssetPackageDataMap();
 	{
-		FName PackageName = Pair.Key;
-		if (Pair.Value->DiskSize < 0)
-		{
-			// Keep the FailedSave previous cook packages; some of them (NeverCookPlaceholders) are not expected to exist
-			// in the package store
-			continue;
-		}
-		if (PackageNameToIndex.Find(PackageName))
+		FReadScopeLock _(EntriesLock);
+		for (const TPair<FName, const FAssetPackageData*>& Pair : PreviousStatePackages)
 		{
-			continue;
+			FName PackageName = Pair.Key;
+			if (Pair.Value->DiskSize < 0)
+			{
+				// Keep the FailedSave previous cook packages; some of them (NeverCookPlaceholders) are not expected to exist
+				// in the package store
+				continue;
+			}
+			if (PackageNameToIndex.Find(PackageName))
+			{
+				continue;
+			}
+			RemoveSet.Add(PackageName);
 		}
+		if (RemoveSet.Num())
 		{
-			RemoveSet.Add(PackageName);
+			PreviousState->PruneAssetData(TSet<FName>(), RemoveSet, FAssetRegistrySerializationOptions());
 		}
-	}
-	if (RemoveSet.Num())
-	{
-		PreviousState->PruneAssetData(TSet<FName>(), RemoveSet, FAssetRegistrySerializationOptions());
-	}
 
-	TArray<FName> RemoveArray;
-	for (TPair<FName, int32>& Pair : PackageNameToIndex)
-	{
-		FName PackageName = Pair.Key;
-		if (!PreviousStatePackages.Find(PackageName))
+		for (TPair<FName, int32>& Pair : PackageNameToIndex)
 		{
-			RemoveArray.Add(PackageName);
+			FName PackageName = Pair.Key;
+			if (!PreviousStatePackages.Find(PackageName))
+			{
+				RemoveArray.Add(PackageName);
+			}
 		}
 	}
 	if (RemoveArray.Num())
@@ -1085,35 +1091,41 @@ TUniquePtr<FAssetRegistryState> FZenStoreWriter::LoadPreviousAssetRegistry()
 
 FCbObject FZenStoreWriter::GetOplogAttachment(FName PackageName, FUtf8StringView AttachmentKey)
 {
-	const int32* Idx = PackageNameToIndex.Find(PackageName);
-	if (!Idx)
+	FIoHash AttachmentHash;
 	{
-		return FCbObject();
-	}
+		FReadScopeLock _(EntriesLock);
 
-	const UTF8CHAR* AttachmentId = UE::FZenStoreHttpClient::FindAttachmentId(AttachmentKey);
-	if (!AttachmentId)
-	{
-		return FCbObject();
-	}
-	FUtf8StringView AttachmentIdView(AttachmentId);
+		const int32* Idx = PackageNameToIndex.Find(PackageName);
+		if (!Idx)
+		{
+			return FCbObject();
+		}
 
-	const FOplogCookInfo& CookInfo = CookedPackagesInfo[*Idx];
-	int32 AttachmentIndex = Algo::LowerBound(CookInfo.Attachments, AttachmentIdView,
-		[](const FOplogCookInfo::FAttachment& Existing, FUtf8StringView AttachmentIdView)
+		const UTF8CHAR* AttachmentId = UE::FZenStoreHttpClient::FindAttachmentId(AttachmentKey);
+		if (!AttachmentId)
 		{
-			return FUtf8StringView(Existing.Key).Compare(AttachmentIdView, ESearchCase::IgnoreCase) < 0;
-		});
-	if (AttachmentIndex == CookInfo.Attachments.Num())
-	{
-		return FCbObject();
-	}
-	const FOplogCookInfo::FAttachment& Existing = CookInfo.Attachments[AttachmentIndex];
-	if (!FUtf8StringView(Existing.Key).Equals(AttachmentIdView, ESearchCase::IgnoreCase))
-	{
-		return FCbObject();
+			return FCbObject();
+		}
+		FUtf8StringView AttachmentIdView(AttachmentId);
+
+		const FOplogCookInfo& CookInfo = CookedPackagesInfo[*Idx];
+		int32 AttachmentIndex = Algo::LowerBound(CookInfo.Attachments, AttachmentIdView,
+			[](const FOplogCookInfo::FAttachment& Existing, FUtf8StringView AttachmentIdView)
+			{
+				return FUtf8StringView(Existing.Key).Compare(AttachmentIdView, ESearchCase::IgnoreCase) < 0;
+			});
+		if (AttachmentIndex == CookInfo.Attachments.Num())
+		{
+			return FCbObject();
+		}
+		const FOplogCookInfo::FAttachment& Existing = CookInfo.Attachments[AttachmentIndex];
+		if (!FUtf8StringView(Existing.Key).Equals(AttachmentIdView, ESearchCase::IgnoreCase))
+		{
+			return FCbObject();
+		}
+		AttachmentHash = Existing.Hash;
 	}
-	TIoStatusOr<FIoBuffer> BufferResult = HttpClient->ReadOpLogAttachment(WriteToString<48>(Existing.Hash));
+	TIoStatusOr<FIoBuffer> BufferResult = HttpClient->ReadOpLogAttachment(WriteToString<48>(AttachmentHash));
 	if (!BufferResult.IsOk())
 	{
 		return FCbObject();
@@ -1130,6 +1142,8 @@ FCbObject FZenStoreWriter::GetOplogAttachment(FName PackageName, FUtf8StringView
 
 void FZenStoreWriter::RemoveCookedPackages(TArrayView<const FName> PackageNamesToRemove)
 {
+	FWriteScopeLock _(EntriesLock);
+
 	TSet<int32> PackageIndicesToKeep;
 	for (int32 Idx = 0, Num = PackageStoreEntries.Num(); Idx < Num; ++Idx)
 	{
@@ -1170,6 +1184,8 @@ void FZenStoreWriter::RemoveCookedPackages(TArrayView<const FName> PackageNamesT
 
 void FZenStoreWriter::RemoveCookedPackages()
 {
+	FWriteScopeLock _(EntriesLock);
+
 	PackageStoreEntries.Empty();
 	CookedPackagesInfo.Empty();
 	PackageNameToIndex.Empty();
