Cooker: IncrementalCook: use only the ReferencedSet when staging so that previously cooked packages that are no longer referenced are not staged.
#rnx

diff --git a/Engine/Source/Developer/IoStoreUtilities/Internal/ZenStoreHttpClient.h b/Engine/Source/Developer/IoStoreUtilities/Internal/ZenStoreHttpClient.h
--- a/Engine/Source/Developer/IoStoreUtilities/Internal/ZenStoreHttpClient.h
+++ b/Engine/Source/Developer/IoStoreUtilities/Internal/ZenStoreHttpClient.h
@@ -63,7 +63,7 @@ public:
 #endif
 
 	TFuture<TIoStatusOr<FCbObject>> GetProjectInfo();
-	TFuture<TIoStatusOr<FCbObject>> GetOplog();
+	TFuture<TIoStatusOr<FCbObject>> GetOplog(bool bTrimByReferencedset = true);
 	TFuture<TIoStatusOr<FCbObject>> GetFiles();
 	TFuture<TIoStatusOr<FCbObject>> GetChunkInfos();
 
diff --git a/Engine/Source/Developer/IoStoreUtilities/Private/ZenStoreHttpClient.cpp b/Engine/Source/Developer/IoStoreUtilities/Private/ZenStoreHttpClient.cpp
--- a/Engine/Source/Developer/IoStoreUtilities/Private/ZenStoreHttpClient.cpp
+++ b/Engine/Source/Developer/IoStoreUtilities/Private/ZenStoreHttpClient.cpp
@@ -622,19 +622,20 @@ TFuture<TIoStatusOr<FCbObject>> FZenStoreHttpClient::GetProjectInfo()
 	});
 }
 
-TFuture<TIoStatusOr<FCbObject>> FZenStoreHttpClient::GetOplog()
+TFuture<TIoStatusOr<FCbObject>> FZenStoreHttpClient::GetOplog(bool bTrimByReferencedSet)
 {
 #if WITH_EDITOR
 	EAsyncExecution ThreadPool = EAsyncExecution::LargeThreadPool;
 #else
 	EAsyncExecution ThreadPool = EAsyncExecution::ThreadPool;
 #endif
-	return Async(ThreadPool, [this]
+	return Async(ThreadPool, [this, bTrimByReferencedSet]
 	{
 		UE::Zen::FZenScopedRequestPtr Request(RequestPool.Get());
 		
 		TStringBuilder<128> Uri;
-		Uri << OplogPath << "/entries";
+		Uri << OplogPath << TEXT("/entries");
+		Uri << TEXT("?") << TEXT("trim_by_referencedset=") << (bTrimByReferencedSet ? TEXT("true") : TEXT("false"));
 
 		TArray64<uint8> GetBuffer;
 		if (Download(*Request, Uri, &GetBuffer, Zen::EContentType::CbObject) && Request->GetResponseCode() == 200)
@@ -818,7 +819,7 @@ TFuture<TIoStatusOr<FCbObject>> FZenStoreHttpClient::GetProjectInfo()
 	return TFuture<TIoStatusOr<FCbObject>>();
 }
 
-TFuture<TIoStatusOr<FCbObject>> FZenStoreHttpClient::GetOplog()
+TFuture<TIoStatusOr<FCbObject>> FZenStoreHttpClient::GetOplog(bool bTrimByReferencedSet)
 {
 	return TFuture<TIoStatusOr<FCbObject>>();
 }
diff --git a/Engine/Source/Developer/IoStoreUtilities/Private/ZenStoreWriter.cpp b/Engine/Source/Developer/IoStoreUtilities/Private/ZenStoreWriter.cpp
--- a/Engine/Source/Developer/IoStoreUtilities/Private/ZenStoreWriter.cpp
+++ b/Engine/Source/Developer/IoStoreUtilities/Private/ZenStoreWriter.cpp
@@ -492,7 +492,8 @@ void FZenStoreWriter::Initialize(const FCookInfo& Info)
 		{
 			UE_LOG(LogZenStoreWriter, Display, TEXT("Fetching oplog..."), *ProjectId, *OplogId);
 
-			TFuture<FIoStatus> FutureOplogStatus = HttpClient->GetOplog().Next([this](TIoStatusOr<FCbObject> OplogStatus)
+			TFuture<FIoStatus> FutureOplogStatus = HttpClient->GetOplog(false /* bTrimByReferencedSet*/)
+				.Next([this](TIoStatusOr<FCbObject> OplogStatus)
 				{
 					if (!OplogStatus.IsOk())
 					{
