IncrementalCook: Store CookDependencies in zenserver ops even for packages that failed to save for the cook (e.g. because they only contained editoronly objects). This allows other packages to transitively depend on their dependencies rather than always recooking because those transitive dependencies are not readable.
Add EPackageStoreEntryFlags::HasPackageData to distinguish between valid packages in the oplog and error and editoronly packages in the oplog.
Add ECommitStatus::NothingToCook and GetCookedStatus to PackageWriter, and EPackageStoreEntryFlags::HasCookError, so we can distinguish between editoronly cases that we don't need to recook each cook and error cases that we do.
Move use of critical section in FZenStoreWriter::CommitInternal so that all writes occur in a single critical section block, to reduce contention.
#rnx

diff --git a/Engine/Source/Developer/IoStoreUtilities/Private/PackageStoreOptimizer.cpp b/Engine/Source/Developer/IoStoreUtilities/Private/PackageStoreOptimizer.cpp
--- a/Engine/Source/Developer/IoStoreUtilities/Private/PackageStoreOptimizer.cpp
+++ b/Engine/Source/Developer/IoStoreUtilities/Private/PackageStoreOptimizer.cpp
@@ -1070,7 +1070,7 @@ void FPackageStoreOptimizer::LoadScriptObjectsBuffer(const FIoBuffer& ScriptObje
 FPackageStoreEntryResource FPackageStoreOptimizer::CreatePackageStoreEntry(const FPackageStorePackage* Package, const FPackageStorePackage* OptionalSegmentPackage) const
 {
 	FPackageStoreEntryResource Result;
-	Result.Flags = EPackageStoreEntryFlags::None;
+	Result.Flags = EPackageStoreEntryFlags::HasPackageData;
 
 	if (OptionalSegmentPackage)
 	{
diff --git a/Engine/Source/Developer/IoStoreUtilities/Private/ZenStoreWriter.cpp b/Engine/Source/Developer/IoStoreUtilities/Private/ZenStoreWriter.cpp
--- a/Engine/Source/Developer/IoStoreUtilities/Private/ZenStoreWriter.cpp
+++ b/Engine/Source/Developer/IoStoreUtilities/Private/ZenStoreWriter.cpp
@@ -690,7 +690,7 @@ bool FZenStoreWriter::IsReservedOplogKey(FUtf8StringView Key)
 
 void FZenStoreWriter::CommitPackage(FCommitPackageInfo&& Info)
 {
-	if (Info.Status == ECommitStatus::Canceled)
+	if (Info.Status == ECommitStatus::Canceled || Info.Status == ECommitStatus::NotCommitted)
 	{
 		RemovePendingPackage(Info.PackageName);
 		return;
@@ -752,21 +752,26 @@ void FZenStoreWriter::CommitPackageInternal(FZenCommitInfo&& ZenCommitInfo)
 
 	IPackageStoreWriter::FCommitEventArgs CommitEventArgs;
 
-	CommitEventArgs.PlatformName	= TargetPlatformFName;
-	CommitEventArgs.PackageName		= CommitInfo.PackageName;
-	CommitEventArgs.EntryIndex		= INDEX_NONE;
-	
-	const bool bComputeHash			= EnumHasAnyFlags(CommitInfo.WriteOptions, EWriteOptions::ComputeHash);
-	const bool bWritePackage		= EnumHasAnyFlags(CommitInfo.WriteOptions, EWriteOptions::Write);
+	CommitEventArgs.PlatformName = TargetPlatformFName;
+	CommitEventArgs.PackageName = CommitInfo.PackageName;
+	CommitEventArgs.EntryIndex = INDEX_NONE;
+
+	const bool bWriteHashFlag = EnumHasAnyFlags(CommitInfo.WriteOptions, EWriteOptions::ComputeHash);
+	const bool bWriteOp = EnumHasAnyFlags(CommitInfo.WriteOptions, EWriteOptions::Write);
+	const bool bComputeHash = bWriteHashFlag && CommitInfo.Status == ECommitStatus::Success;
+	const bool bComputeValidPackage = bWriteOp && CommitInfo.Status == ECommitStatus::Success;
 
-	if (CommitInfo.Status == ECommitStatus::Success && bWritePackage)
+	if (bWriteOp)
 	{
-		checkf(!PackageState->PackageData.IsEmpty(), TEXT("CommitPackage called with bSucceeded but without first calling WritePackageData"))
 		checkf(EnumHasAllFlags(CommitInfo.WriteOptions, EWriteOptions::Write), TEXT("Partial EWriteOptions::Write options are not yet implemented."));
 		checkf(!EnumHasAnyFlags(CommitInfo.WriteOptions, EWriteOptions::SaveForDiff), TEXT("-diffonly -savefordiff is not yet implemented."));
 
 		FPackageStoreEntryResource PackageStoreEntry;
+
+		if (bComputeValidPackage)
 		{
+			checkf(!PackageState->PackageData.IsEmpty(), TEXT("CommitPackage called with bSucceeded but without first calling WritePackageData"));
+
 			FPackageDataEntry* PkgData = nullptr;
 			FPackageDataEntry* OptionalSegmentPkgData = nullptr;
 			for (FPackageDataEntry& PackageDataEntry : PackageState->PackageData)
@@ -783,28 +788,17 @@ void FZenStoreWriter::CommitPackageInternal(FZenCommitInfo&& ZenCommitInfo)
 					OptionalSegmentPkgData = &PackageDataEntry;
 				}
 			}
-			PackageStoreEntry = PackageStoreOptimizer->CreatePackageStoreEntry(PkgData->OptimizedPackage.Get(), OptionalSegmentPkgData ? OptionalSegmentPkgData->OptimizedPackage.Get() : nullptr);
-
-			FWriteScopeLock _(EntriesLock);
-			CommitEventArgs.EntryIndex = PackageNameToIndex.FindOrAdd(CommitInfo.PackageName, PackageStoreEntries.Num());
-			if (CommitEventArgs.EntryIndex == PackageStoreEntries.Num())
-			{
-				PackageStoreEntries.Emplace();
-				CookedPackagesInfo.Emplace();
-			}
-			PackageStoreEntries[CommitEventArgs.EntryIndex] = PackageStoreEntry;
+			PackageStoreEntry = PackageStoreOptimizer->CreatePackageStoreEntry(
+				PkgData->OptimizedPackage.Get(),
+				OptionalSegmentPkgData ? OptionalSegmentPkgData->OptimizedPackage.Get() : nullptr);
 		}
-		
-		if (EntryCreatedEvent.IsBound())
+		else
 		{
-			IPackageStoreWriter::FEntryCreatedEventArgs EntryCreatedEventArgs
-			{
-				TargetPlatformFName,
-				PackageStoreEntry
-			};
-			EntryCreatedEvent.Broadcast(EntryCreatedEventArgs);
+			const bool bHasCookError = CommitInfo.Status == ECommitStatus::Error;
+			PackageStoreEntry = FPackageStoreEntryResource::CreateEmptyPackage(
+				CommitInfo.PackageName, bHasCookError);
 		}
-		
+
 		FMD5 PkgHashGen;
 		FCbPackage OplogEntry;
 		// Commit attachments
@@ -850,7 +844,7 @@ void FZenStoreWriter::CommitPackageInternal(FZenCommitInfo&& ZenCommitInfo)
 		PackageNameKey.ToLowerInline();
 		OplogEntryDesc << "key" << PackageNameKey;
 		OplogEntryDesc << "packagestoreentry" << PackageStoreEntry;
-		
+
 		auto AppendFileNameAndRegionsToOplog = [this, &OplogEntryDesc](const FString& LooseFilePath, const TArray<FFileRegion>& FileRegions)
 		{
 			FStringView RelativePathView;
@@ -869,29 +863,32 @@ void FZenStoreWriter::CommitPackageInternal(FZenCommitInfo&& ZenCommitInfo)
 			}
 		};
 
-		OplogEntryDesc.BeginArray("packagedata");
-		
-		for (FPackageDataEntry& PkgData : PackageState->PackageData)
+		if (bComputeValidPackage && PackageState->PackageData.Num())
 		{
-			if (bComputeHash)
+			OplogEntryDesc.BeginArray("packagedata");
+
+			for (FPackageDataEntry& PkgData : PackageState->PackageData)
 			{
-				PackageState->PackageHashes->ChunkHashes.Add(PkgData.Info.ChunkId, PkgData.CompressedPayload.Get().GetRawHash());
-			}
+				if (bComputeHash)
+				{
+					PackageState->PackageHashes->ChunkHashes.Add(PkgData.Info.ChunkId, PkgData.CompressedPayload.Get().GetRawHash());
+				}
 
-			FCbAttachment PkgDataAttachment = FCbAttachment(PkgData.CompressedPayload.Get());
-			PkgHashGen.Update(PkgDataAttachment.GetHash().GetBytes(), sizeof(FIoHash::ByteArray));
-			OplogEntry.AddAttachment(PkgDataAttachment);
+				FCbAttachment PkgDataAttachment = FCbAttachment(PkgData.CompressedPayload.Get());
+				PkgHashGen.Update(PkgDataAttachment.GetHash().GetBytes(), sizeof(FIoHash::ByteArray));
+				OplogEntry.AddAttachment(PkgDataAttachment);
 
-			OplogEntryDesc.BeginObject();
-			OplogEntryDesc << "id" << PkgData.ChunkId;
-			OplogEntryDesc << "data" << PkgDataAttachment;
-			AppendFileNameAndRegionsToOplog(PkgData.Info.LooseFilePath, PkgData.FileRegions);
-			OplogEntryDesc.EndObject();
-		}
+				OplogEntryDesc.BeginObject();
+				OplogEntryDesc << "id" << PkgData.ChunkId;
+				OplogEntryDesc << "data" << PkgDataAttachment;
+				AppendFileNameAndRegionsToOplog(PkgData.Info.LooseFilePath, PkgData.FileRegions);
+				OplogEntryDesc.EndObject();
+			}
 
-		OplogEntryDesc.EndArray();
+			OplogEntryDesc.EndArray();
+		}
 
-		if (PackageState->BulkData.Num())
+		if (bComputeValidPackage && PackageState->BulkData.Num())
 		{
 			OplogEntryDesc.BeginArray("bulkdata");
 
@@ -917,7 +914,7 @@ void FZenStoreWriter::CommitPackageInternal(FZenCommitInfo&& ZenCommitInfo)
 			OplogEntryDesc.EndArray();
 		}
 
-		if (PackageState->FileData.Num())
+		if (bComputeValidPackage && PackageState->FileData.Num())
 		{
 			OplogEntryDesc.BeginArray("files");
 
@@ -968,8 +965,27 @@ void FZenStoreWriter::CommitPackageInternal(FZenCommitInfo&& ZenCommitInfo)
 		OplogEntryDesc.EndObject();
 		OplogEntry.SetObject(OplogEntryDesc.Save().AsObject());
 
+		if (bComputeValidPackage && EntryCreatedEvent.IsBound())
+		{
+			IPackageStoreWriter::FEntryCreatedEventArgs EntryCreatedEventArgs
+			{
+				TargetPlatformFName,
+				PackageStoreEntry
+			};
+			EntryCreatedEvent.Broadcast(EntryCreatedEventArgs);
+		}
+
 		{
 			FWriteScopeLock _(EntriesLock);
+			CommitEventArgs.EntryIndex = PackageNameToIndex.FindOrAdd(
+				CommitInfo.PackageName, PackageStoreEntries.Num());
+			if (CommitEventArgs.EntryIndex == PackageStoreEntries.Num())
+			{
+				PackageStoreEntries.Emplace();
+				CookedPackagesInfo.Emplace();
+			}
+			PackageStoreEntries[CommitEventArgs.EntryIndex] = MoveTemp(PackageStoreEntry);
+
 			FOplogCookInfo& CookInfo = CookedPackagesInfo[CommitEventArgs.EntryIndex];
 			CookInfo.bUpToDate = true;
 			CookInfo.Attachments = MoveTemp(CookInfoAttachments);
@@ -978,7 +994,7 @@ void FZenStoreWriter::CommitPackageInternal(FZenCommitInfo&& ZenCommitInfo)
 		TIoStatusOr<uint64> Status = HttpClient->AppendOp(MoveTemp(OplogEntry));
 		UE_CLOG(!Status.IsOk(), LogZenStoreWriter, Error, TEXT("Failed to commit oplog entry '%s' to Zen"), *CommitInfo.PackageName.ToString());
 	}
-	else if (CommitInfo.Status == ECommitStatus::Success && bComputeHash)
+	else if (bComputeHash)
 	{
 		checkf(!PackageState->PackageData.IsEmpty(), TEXT("CommitPackage called with bSucceeded but without first calling WritePackageData"));
 		
@@ -1010,7 +1026,7 @@ void FZenStoreWriter::CommitPackageInternal(FZenCommitInfo&& ZenCommitInfo)
 		PackageState->PackageHashes->PackageHash.Set(PkgHashGen);
 	}
 
-	if (bWritePackage)
+	if (bWriteOp)
 	{
 		BroadcastCommit(CommitEventArgs);
 	}
@@ -1140,6 +1156,26 @@ FCbObject FZenStoreWriter::GetOplogAttachment(FName PackageName, FUtf8StringView
 	return FCbObject(SharedBuffer);
 }
 
+IPackageWriter::ECommitStatus FZenStoreWriter::GetCommitStatus(FName PackageName)
+{
+	FReadScopeLock _(EntriesLock);
+
+	const int32* Idx = PackageNameToIndex.Find(PackageName);
+	if (!Idx)
+	{
+		return ECommitStatus::NotCommitted;
+	}
+	if (PackageStoreEntries[*Idx].HasPackageData())
+	{
+		return ECommitStatus::Success;
+	}
+	if (!PackageStoreEntries[*Idx].HasPackageData())
+	{
+		return ECommitStatus::NothingToCook;
+	}
+	return ECommitStatus::Error;
+}
+
 void FZenStoreWriter::RemoveCookedPackages(TArrayView<const FName> PackageNamesToRemove)
 {
 	FWriteScopeLock _(EntriesLock);
diff --git a/Engine/Source/Developer/IoStoreUtilities/Public/ZenStoreWriter.h b/Engine/Source/Developer/IoStoreUtilities/Public/ZenStoreWriter.h
--- a/Engine/Source/Developer/IoStoreUtilities/Public/ZenStoreWriter.h
+++ b/Engine/Source/Developer/IoStoreUtilities/Public/ZenStoreWriter.h
@@ -146,6 +146,8 @@ public:
 
 	IOSTOREUTILITIES_API TUniquePtr<FAssetRegistryState> LoadPreviousAssetRegistry() override;
 	IOSTOREUTILITIES_API virtual FCbObject GetOplogAttachment(FName PackageName, FUtf8StringView AttachmentKey) override;
+	IOSTOREUTILITIES_API virtual ECommitStatus GetCommitStatus(FName PackageName) override;
+
 	IOSTOREUTILITIES_API virtual void RemoveCookedPackages(TArrayView<const FName> PackageNamesToRemove) override;
 	IOSTOREUTILITIES_API virtual void RemoveCookedPackages() override;
 	IOSTOREUTILITIES_API virtual void UpdatePackageModificationStatus(FName PackageName, bool bIterativelyUnmodified,
diff --git a/Engine/Source/Editor/UnrealEd/Private/CookOnTheFlyServer.cpp b/Engine/Source/Editor/UnrealEd/Private/CookOnTheFlyServer.cpp
--- a/Engine/Source/Editor/UnrealEd/Private/CookOnTheFlyServer.cpp
+++ b/Engine/Source/Editor/UnrealEd/Private/CookOnTheFlyServer.cpp
@@ -4599,7 +4599,6 @@ private: // Used only by UCookOnTheFlyServer, which has private access
 		TArray<FAssetDependency>* OverrideDependencies);
 	IPackageWriter::EWriteOptions GetCommitWriteOptions() const;
 	static void AddDependency(TMap<FPackageData*, EInstigator>& InDependencies, FPackageData* PackageData, bool bHard);
-	static IPackageWriter::ECommitStatus PackageResultToCommitStatus(FSavePackageResultStruct& Result);
 
 	// General Package Data
 	UCookOnTheFlyServer& COTFS;
@@ -6692,7 +6691,7 @@ void FSaveCookedPackageContext::FinishPlatform()
 		ICookedPackageWriter::FCommitPackageInfo Info;
 		// Note GetCommitAttachments mutates the SaveResult; it moves CookDependencies out of it
 		Info.Attachments = GetCommitAttachments(GeneratedResult.GetPtrOrNull(), OverridePackageDependencies.GetPtrOrNull());
-		Info.Status = PackageResultToCommitStatus(SavePackageResult);
+		Info.Status = PackageResultToCommitStatus(SavePackageResult.Result);
 		Info.PackageName = Package->GetFName();
 		Info.PackageHash = AssetPackageData ? AssetPackageData->GetPackageSavedHash() : FIoHash();
 		Info.WriteOptions = GetCommitWriteOptions();
@@ -6991,23 +6990,6 @@ void FSaveCookedPackageContext::AddDependency(TMap<FPackageData*, EInstigator>&
 	}
 }
 
-IPackageWriter::ECommitStatus FSaveCookedPackageContext::PackageResultToCommitStatus(
-	FSavePackageResultStruct& Result)
-{
-	if (Result.IsSuccessful())
-	{
-		return IPackageWriter::ECommitStatus::Success;
-	}
-	else if (Result.Result == ESavePackageResult::Timeout)
-	{
-		return IPackageWriter::ECommitStatus::Canceled;
-	}
-	else
-	{
-		return IPackageWriter::ECommitStatus::Error;
-	}
-}
-
 } // namespace UE::Cook
 
 void UCookOnTheFlyServer::RecordExternalActorDependencies(TConstArrayView<FName> ExternalActorDependencies)
diff --git a/Engine/Source/Editor/UnrealEd/Private/Cooker/CookRequestCluster.cpp b/Engine/Source/Editor/UnrealEd/Private/Cooker/CookRequestCluster.cpp
--- a/Engine/Source/Editor/UnrealEd/Private/Cooker/CookRequestCluster.cpp
+++ b/Engine/Source/Editor/UnrealEd/Private/Cooker/CookRequestCluster.cpp
@@ -1803,6 +1803,17 @@ void FRequestCluster::FGraphSearch::FExploreEdgesContext::ProcessPlatformAttachm
 			}
 		}
 		bool bShouldIterativelySkip = bIterativelyUnmodified;
+		if (QueryPlatformData.CookAttachments.CommitStatus == IPackageWriter::ECommitStatus::Error)
+		{
+			// Recook packages with errors even if they have not changed, so that the error is not hidden from
+			// the output for the incremental cook.
+			bShouldIterativelySkip = false;
+
+			// ITERATIVECOOK_TODO: Need to mark a generator package as not iteratively skippable if any of its
+			// generated packages had errors. This will require waiting on the generated package results to be fetched,
+			// and will require care to avoid causing a softlock in the presence of a transitive build dependency
+			// cycle.
+		}
 		PackageWriter->UpdatePackageModificationStatus(PackageName, bIterativelyUnmodified,
 			bShouldIterativelySkip);
 
diff --git a/Engine/Source/Editor/UnrealEd/Private/Cooker/CookTypes.cpp b/Engine/Source/Editor/UnrealEd/Private/Cooker/CookTypes.cpp
--- a/Engine/Source/Editor/UnrealEd/Private/Cooker/CookTypes.cpp
+++ b/Engine/Source/Editor/UnrealEd/Private/Cooker/CookTypes.cpp
@@ -1096,3 +1096,23 @@ FMultiPackageReaderResults GetSaveExportsAndImports(UPackage* Package, UObject*
 
 	return Result;
 }
+
+IPackageWriter::ECommitStatus PackageResultToCommitStatus(ESavePackageResult Result)
+{
+	if (IsSuccessful(Result))
+	{
+		return IPackageWriter::ECommitStatus::Success;
+	}
+	else if (Result == ESavePackageResult::Timeout)
+	{
+		return IPackageWriter::ECommitStatus::Canceled;
+	}
+	else if (Result == ESavePackageResult::ContainsEditorOnlyData)
+	{
+		return IPackageWriter::ECommitStatus::NothingToCook;
+	}
+	else
+	{
+		return IPackageWriter::ECommitStatus::Error;
+	}
+}
diff --git a/Engine/Source/Editor/UnrealEd/Private/Cooker/CookTypes.h b/Engine/Source/Editor/UnrealEd/Private/Cooker/CookTypes.h
--- a/Engine/Source/Editor/UnrealEd/Private/Cooker/CookTypes.h
+++ b/Engine/Source/Editor/UnrealEd/Private/Cooker/CookTypes.h
@@ -767,3 +767,6 @@ struct FMultiPackageReaderResults
 
 /** Save the package and read the LinkerTables of its saved data. */
 FMultiPackageReaderResults GetSaveExportsAndImports(UPackage* Package, UObject* Asset, FSavePackageArgs SaveArgs);
+
+IPackageWriter::ECommitStatus PackageResultToCommitStatus(ESavePackageResult Result);
+
diff --git a/Engine/Source/Editor/UnrealEd/Private/Cooker/DiffPackageWriter.cpp b/Engine/Source/Editor/UnrealEd/Private/Cooker/DiffPackageWriter.cpp
--- a/Engine/Source/Editor/UnrealEd/Private/Cooker/DiffPackageWriter.cpp
+++ b/Engine/Source/Editor/UnrealEd/Private/Cooker/DiffPackageWriter.cpp
@@ -4,6 +4,7 @@
 
 #include "Containers/Array.h"
 #include "Containers/Map.h"
+#include "Cooker/CookTypes.h"
 #include "CookOnTheSide/CookLog.h"
 #include "CoreGlobals.h"
 #include "EditorDomain/EditorDomainUtils.h"
@@ -472,14 +473,7 @@ bool FLinkerDiffPackageWriter::IsAnotherSaveNeeded(FSavePackageResultStruct& Pre
 		// The contract with the Inner is that every Begin is paired with a single commit.
 		// Send the old commit and the new begin.
 		IPackageWriter::FCommitPackageInfo CommitInfo;
-		if (OtherResult == ESavePackageResult::Success)
-		{
-			CommitInfo.Status = IPackageWriter::ECommitStatus::Success;
-		}
-		else
-		{
-			CommitInfo.Status = IPackageWriter::ECommitStatus::Error;
-		}
+		CommitInfo.Status = PackageResultToCommitStatus(OtherResult.Result);
 		CommitInfo.PackageName = BeginInfo.PackageName;
 		CommitInfo.WriteOptions = IPackageWriter::EWriteOptions::None;
 		Inner->CommitPackage(MoveTemp(CommitInfo));
diff --git a/Engine/Source/Editor/UnrealEd/Private/Cooker/DiffPackageWriter.h b/Engine/Source/Editor/UnrealEd/Private/Cooker/DiffPackageWriter.h
--- a/Engine/Source/Editor/UnrealEd/Private/Cooker/DiffPackageWriter.h
+++ b/Engine/Source/Editor/UnrealEd/Private/Cooker/DiffPackageWriter.h
@@ -90,6 +90,10 @@ public:
 	{
 		return Inner->GetOplogAttachment(PackageName, AttachmentKey);
 	}
+	virtual ECommitStatus GetCommitStatus(FName PackageName) override
+	{
+		return Inner->GetCommitStatus(PackageName);
+	}
 	virtual void RemoveCookedPackages(TArrayView<const FName> PackageNamesToRemove) override
 	{
 		Inner->RemoveCookedPackages(PackageNamesToRemove);
@@ -243,6 +247,10 @@ public:
 	{
 		return Inner->GetOplogAttachment(PackageName, AttachmentKey);
 	}
+	virtual ECommitStatus GetCommitStatus(FName PackageName) override
+	{
+		return Inner->GetCommitStatus(PackageName);
+	}
 	virtual void RemoveCookedPackages(TArrayView<const FName> PackageNamesToRemove) override
 	{
 		Inner->RemoveCookedPackages(PackageNamesToRemove);
diff --git a/Engine/Source/Editor/UnrealEd/Private/Cooker/LooseCookedPackageWriter.cpp b/Engine/Source/Editor/UnrealEd/Private/Cooker/LooseCookedPackageWriter.cpp
--- a/Engine/Source/Editor/UnrealEd/Private/Cooker/LooseCookedPackageWriter.cpp
+++ b/Engine/Source/Editor/UnrealEd/Private/Cooker/LooseCookedPackageWriter.cpp
@@ -97,7 +97,8 @@ void FLooseCookedPackageWriter::CommitPackageInternal(FPackageWriterRecords::FPa
 	{
 		AsyncSave(Record, Info);
 	}
-	if (Info.Status != IPackageWriter::ECommitStatus::Canceled)
+	if (Info.Status != IPackageWriter::ECommitStatus::Canceled &&
+		Info.Status != IPackageWriter::ECommitStatus::NotCommitted)
 	{
 		UpdateManifest(Record);
 	}
@@ -818,6 +819,12 @@ FCbObject FLooseCookedPackageWriter::GetOplogAttachment(FName PackageName, FUtf8
 	return FCbObject();
 }
 
+IPackageWriter::ECommitStatus FLooseCookedPackageWriter::GetCommitStatus(FName PackageName)
+{
+	// GetCommitStatus is not implemented by FLooseCookedPackageWriter
+	return ECommitStatus::NotCommitted;
+}
+
 void FLooseCookedPackageWriter::RemoveCookedPackages(TArrayView<const FName> PackageNamesToRemove)
 {
 	if (PackageNameToCookedFiles.IsEmpty())
diff --git a/Engine/Source/Editor/UnrealEd/Private/Cooker/LooseCookedPackageWriter.h b/Engine/Source/Editor/UnrealEd/Private/Cooker/LooseCookedPackageWriter.h
--- a/Engine/Source/Editor/UnrealEd/Private/Cooker/LooseCookedPackageWriter.h
+++ b/Engine/Source/Editor/UnrealEd/Private/Cooker/LooseCookedPackageWriter.h
@@ -69,6 +69,7 @@ public:
 	virtual void EndCook(const FCookInfo& Info) override;
 	virtual TUniquePtr<FAssetRegistryState> LoadPreviousAssetRegistry() override;
 	virtual FCbObject GetOplogAttachment(FName PackageName, FUtf8StringView AttachmentKey) override;
+	virtual ECommitStatus GetCommitStatus(FName PackageName) override;
 	virtual void RemoveCookedPackages(TArrayView<const FName> PackageNamesToRemove) override;
 	virtual void RemoveCookedPackages() override;
 	virtual void UpdatePackageModificationStatus(FName PackageName, bool bIterativelyUnmodified,
diff --git a/Engine/Source/Editor/UnrealEd/Private/TargetDomain/TargetDomainUtils.cpp b/Engine/Source/Editor/UnrealEd/Private/TargetDomain/TargetDomainUtils.cpp
--- a/Engine/Source/Editor/UnrealEd/Private/TargetDomain/TargetDomainUtils.cpp
+++ b/Engine/Source/Editor/UnrealEd/Private/TargetDomain/TargetDomainUtils.cpp
@@ -299,7 +299,13 @@ EPackageMountPoint GetPackageMountPoint(FName PackageName, FName TransientPackag
 	}
 	TStringBuilder<256> StringBuffer;
 	PackageName.ToString(StringBuffer);
-	if (FPackageName::IsMemoryPackage(StringBuffer) ||
+	if (
+		// Some packages get renamed to "TrashedPackage" during blueprint compilation and are no longer valid for
+		// saving but might have been dereferenced by TObjectPtr during PostLoad/PreSave. We need to discard these
+		// packages, which we can do by requiring a valid package name; all valid packages start with /MountPoint/.
+		!StringBuffer.ToView().StartsWith('/') ||
+		// Ignore /Memory and /Temp packages
+		FPackageName::IsMemoryPackage(StringBuffer) ||
 		FPackageName::IsTempPackage(StringBuffer))
 	{
 		return EPackageMountPoint::Transient;
@@ -982,9 +988,11 @@ void FCookAttachments::Fetch(TArrayView<FName> PackageNames, const ITargetPlatfo
 	{
 		FCbObject DependenciesObj;
 		FCbObject BuildDefinitionsObj;
+		IPackageWriter::ECommitStatus CommitStatus = IPackageWriter::ECommitStatus::NotCommitted;
 		if (TargetPlatform)
 		{
 			check(PackageWriter);
+			CommitStatus = PackageWriter->GetCommitStatus(PackageName);
 			DependenciesObj = PackageWriter->GetOplogAttachment(PackageName, CookDependenciesAttachmentKey);
 			BuildDefinitionsObj = PackageWriter->GetOplogAttachment(PackageName, CookDependenciesAttachmentKey);
 		}
@@ -997,6 +1005,8 @@ void FCookAttachments::Fetch(TArrayView<FName> PackageNames, const ITargetPlatfo
 			}
 			DependenciesObj = GEditorDomainOplog->GetOplogAttachment(PackageName, CookDependenciesAttachmentKey);
 			BuildDefinitionsObj = GEditorDomainOplog->GetOplogAttachment(PackageName, BuildDefinitionsAttachmentKey);
+			CommitStatus = DependenciesObj ? IPackageWriter::ECommitStatus::Success
+				: IPackageWriter::ECommitStatus::NotCommitted;
 		}
 
 		FCookAttachments Result;
@@ -1005,6 +1015,7 @@ void FCookAttachments::Fetch(TArrayView<FName> PackageNames, const ITargetPlatfo
 			Result.Dependencies.PackageName = PackageName;
 		}
 		LoadFromCompactBinary(MoveTemp(BuildDefinitionsObj), Result.BuildDefinitions);
+		Result.CommitStatus = CommitStatus;
 
 		Callback(PackageName, MoveTemp(Result));
 	}
diff --git a/Engine/Source/Editor/UnrealEd/Private/TargetDomain/TargetDomainUtils.h b/Engine/Source/Editor/UnrealEd/Private/TargetDomain/TargetDomainUtils.h
--- a/Engine/Source/Editor/UnrealEd/Private/TargetDomain/TargetDomainUtils.h
+++ b/Engine/Source/Editor/UnrealEd/Private/TargetDomain/TargetDomainUtils.h
@@ -186,6 +186,7 @@ struct FCookAttachments
 {
 	FCookDependencies Dependencies;
 	FBuildDefinitionList BuildDefinitions;
+	IPackageWriter::ECommitStatus CommitStatus = IPackageWriter::ECommitStatus::NotCommitted;
 
 	void Reset();
 	void Empty();
diff --git a/Engine/Source/Runtime/Core/Private/IO/PackageStore.cpp b/Engine/Source/Runtime/Core/Private/IO/PackageStore.cpp
--- a/Engine/Source/Runtime/Core/Private/IO/PackageStore.cpp
+++ b/Engine/Source/Runtime/Core/Private/IO/PackageStore.cpp
@@ -129,6 +129,15 @@ FPackageStoreEntryResource FPackageStoreEntryResource::FromCbObject(FCbObjectVie
 	return Entry;
 }
 
+FPackageStoreEntryResource FPackageStoreEntryResource::CreateEmptyPackage(FName PackageName, bool bHasCookError)
+{
+	FPackageStoreEntryResource Entry;
+	Entry.PackageName = PackageName;
+	Entry.PackageId = FPackageId::FromName(PackageName);
+	Entry.Flags = bHasCookError ? EPackageStoreEntryFlags::HasCookError : EPackageStoreEntryFlags::None;
+	return Entry;
+}
+
 thread_local int32 FPackageStore::ThreadReadCount = 0;
 
 FPackageStoreReadScope::FPackageStoreReadScope(FPackageStore& InPackageStore)
diff --git a/Engine/Source/Runtime/Core/Public/IO/PackageStore.h b/Engine/Source/Runtime/Core/Public/IO/PackageStore.h
--- a/Engine/Source/Runtime/Core/Public/IO/PackageStore.h
+++ b/Engine/Source/Runtime/Core/Public/IO/PackageStore.h
@@ -44,13 +44,16 @@ struct FPackageStoreEntry
 };
 
 /**
- * Package store entry flags.
+ * Package store entry flags. These flags are persisted in the oplog as integers so
+ * do not change their values.
  */
 enum class EPackageStoreEntryFlags : uint32
 {
-	None		= 0,
-	AutoOptional= 0x02,
-	OptionalSegment = 0x04,
+	None				= 0,
+	HasPackageData		= 0x00000001,
+	AutoOptional		= 0x00000002,
+	OptionalSegment		= 0x00000004,
+	HasCookError		= 0x00000008,
 };
 ENUM_CLASS_FLAGS(EPackageStoreEntryFlags);
 
@@ -94,11 +97,28 @@ struct FPackageStoreEntryResource
 		return EnumHasAnyFlags(Flags, EPackageStoreEntryFlags::OptionalSegment);
 	}
 
+	bool HasPackageData() const
+	{
+		return EnumHasAnyFlags(Flags, EPackageStoreEntryFlags::HasPackageData);
+	}
+
+	bool HasCookError() const
+	{
+		return EnumHasAnyFlags(Flags, EPackageStoreEntryFlags::HasCookError);
+	}
+
 	CORE_API friend FArchive& operator<<(FArchive& Ar, FPackageStoreEntryResource& PackageStoreEntry);
 	
 	CORE_API friend FCbWriter& operator<<(FCbWriter& Writer, const FPackageStoreEntryResource& PackageStoreEntry);
 	
 	CORE_API static FPackageStoreEntryResource FromCbObject(FCbObjectView Obj);
+
+	/**
+	 * Creates a PackageStoreEntryResource that records a packagename and will be in an op with possible metadata
+	 * stored in attachments, but for which the op has no packagedata. This is used for tracking build dependencies
+	 * on packages even if those packages fail to cook due to error or editoronly.
+	 */
+	CORE_API static FPackageStoreEntryResource CreateEmptyPackage(FName PackageName, bool bHasCookError);
 };
 
 /**
diff --git a/Engine/Source/Runtime/Core/Public/Serialization/PackageWriter.h b/Engine/Source/Runtime/Core/Public/Serialization/PackageWriter.h
--- a/Engine/Source/Runtime/Core/Public/Serialization/PackageWriter.h
+++ b/Engine/Source/Runtime/Core/Public/Serialization/PackageWriter.h
@@ -101,7 +101,9 @@ public:
 	{
 		Success,
 		Canceled,
-		Error
+		NothingToCook,
+		Error,
+		NotCommitted,
 	};
 	struct FCommitPackageInfo
 	{
@@ -342,6 +344,12 @@ public:
 	 */
 	virtual FCbObject GetOplogAttachment(FName PackageName, FUtf8StringView AttachmentKey) = 0;
 
+	/**
+	 * Report commit status for the given package. If cooking incrementally, this reports the status
+	 * from a previous cook if it has not yet been committed in the current cook.
+	 */
+	virtual ECommitStatus GetCommitStatus(FName PackageName) = 0;
+
 	/**
 	 * Remove the given cooked package(s) from storage; they have been modified since the last cook.
 	 */
diff --git a/Engine/Source/Runtime/CoreUObject/Public/UObject/Package.h b/Engine/Source/Runtime/CoreUObject/Public/UObject/Package.h
--- a/Engine/Source/Runtime/CoreUObject/Public/UObject/Package.h
+++ b/Engine/Source/Runtime/CoreUObject/Public/UObject/Package.h
@@ -87,6 +87,14 @@ enum class ESavePackageResult
 	Timeout,
 };
 
+inline bool IsSuccessful(ESavePackageResult Result)
+{
+	return
+		Result == ESavePackageResult::Success ||
+		Result == ESavePackageResult::GenerateStub ||
+		Result == ESavePackageResult::ReplaceCompletely;
+}
+
 /**
 * Struct returned from save package, contains the enum as well as extra data about what was written
 */
@@ -139,10 +147,7 @@ struct FSavePackageResultStruct
 	/** Returns whether the package save was successful */
 	bool IsSuccessful() const
 	{
-		return
-			Result == ESavePackageResult::Success ||
-			Result == ESavePackageResult::GenerateStub ||
-			Result == ESavePackageResult::ReplaceCompletely;
+		return ::IsSuccessful(Result);
 	}
 };
 
