Include package and bulk data sizes in oplog entries.

Motivation here is primarily to surface this information to Zen's self-host dashboard where begin able to show sizes is has a lot of benefit. While there is a `chunkinfos/` end point, its gathering nature means it's rather slow (approaching 10s for an oplog of 1.3MM ops). Sizes may prove of further utility for future tools that reason about data via the oplog too. The "files" part is an all together different problem and blissfully glossed over for now.

#ushell-unshelve of 38297801
#rnx

diff --git a/Engine/Source/Developer/IoStoreUtilities/Private/ZenStoreWriter.cpp b/Engine/Source/Developer/IoStoreUtilities/Private/ZenStoreWriter.cpp
--- a/Engine/Source/Developer/IoStoreUtilities/Private/ZenStoreWriter.cpp
+++ b/Engine/Source/Developer/IoStoreUtilities/Private/ZenStoreWriter.cpp
@@ -869,17 +869,20 @@ void FZenStoreWriter::CommitPackageInternal(FZenCommitInfo&& ZenCommitInfo)
 
 			for (FPackageDataEntry& PkgData : PackageState->PackageData)
 			{
+				FCompressedBuffer Payload = PkgData.CompressedPayload.Get();
 				if (bComputeHash)
 				{
-					PackageState->PackageHashes->ChunkHashes.Add(PkgData.Info.ChunkId, PkgData.CompressedPayload.Get().GetRawHash());
+					PackageState->PackageHashes->ChunkHashes.Add(PkgData.Info.ChunkId, Payload.GetRawHash());
 				}
 
-				FCbAttachment PkgDataAttachment = FCbAttachment(PkgData.CompressedPayload.Get());
+				FCbAttachment PkgDataAttachment = FCbAttachment(Payload);
 				PkgHashGen.Update(PkgDataAttachment.GetHash().GetBytes(), sizeof(FIoHash::ByteArray));
 				OplogEntry.AddAttachment(PkgDataAttachment);
 
 				OplogEntryDesc.BeginObject();
 				OplogEntryDesc << "id" << PkgData.ChunkId;
+				OplogEntryDesc << "size" << Payload.GetCompressedSize();
+				OplogEntryDesc << "rawsize" << Payload.GetRawSize();
 				OplogEntryDesc << "data" << PkgDataAttachment;
 				AppendFileNameAndRegionsToOplog(PkgData.Info.LooseFilePath, PkgData.FileRegions);
 				OplogEntryDesc.EndObject();
@@ -894,18 +897,21 @@ void FZenStoreWriter::CommitPackageInternal(FZenCommitInfo&& ZenCommitInfo)
 
 			for (FBulkDataEntry& Bulk : PackageState->BulkData)
 			{
+				FCompressedBuffer Payload = Bulk.CompressedPayload.Get();
 				if (bComputeHash)
 				{
-					PackageState->PackageHashes->ChunkHashes.Add(Bulk.Info.ChunkId, Bulk.CompressedPayload.Get().GetRawHash());
+					PackageState->PackageHashes->ChunkHashes.Add(Bulk.Info.ChunkId, Payload.GetRawHash());
 				}
 
-				FCbAttachment BulkAttachment(Bulk.CompressedPayload.Get());
+				FCbAttachment BulkAttachment(Payload);
 				PkgHashGen.Update(BulkAttachment.GetHash().GetBytes(), sizeof(FIoHash::ByteArray));
 				OplogEntry.AddAttachment(BulkAttachment);
 
 				OplogEntryDesc.BeginObject();
 				OplogEntryDesc << "id" << Bulk.ChunkId;
 				OplogEntryDesc << "type" << LexToString(Bulk.Info.BulkDataType);
+				OplogEntryDesc << "size" << Payload.GetCompressedSize();
+				OplogEntryDesc << "rawsize" << Payload.GetRawSize();
 				OplogEntryDesc << "data" << BulkAttachment;
 				AppendFileNameAndRegionsToOplog(Bulk.Info.LooseFilePath, Bulk.FileRegions);
 				OplogEntryDesc.EndObject();
