Correct reserved key name and added "bulkdata"

#ushell-unshelve of 38297801

#rnx

diff --git a/Engine/Source/Developer/IoStoreUtilities/Private/ZenStoreWriter.cpp b/Engine/Source/Developer/IoStoreUtilities/Private/ZenStoreWriter.cpp
--- a/Engine/Source/Developer/IoStoreUtilities/Private/ZenStoreWriter.cpp
+++ b/Engine/Source/Developer/IoStoreUtilities/Private/ZenStoreWriter.cpp
@@ -154,7 +154,7 @@ void FZenStoreWriter::StaticInit()
 		return;
 	}
 
-	ReservedOplogKeys.Append({ UTF8TEXT("files"), UTF8TEXT("key"), UTF8TEXT("package"), UTF8TEXT("packagestoreentry") });
+	ReservedOplogKeys.Append({ UTF8TEXT("files"), UTF8TEXT("key"), UTF8TEXT("packagedata"), UTF8TEXT("bulkdata"), UTF8TEXT("packagestoreentry") });
 	Algo::Sort(ReservedOplogKeys, [](const UTF8CHAR* A, const UTF8CHAR* B)
 		{
 			return FUtf8StringView(A).Compare(FUtf8StringView(B), ESearchCase::IgnoreCase) < 0;
