Split up oplog delete and create into two operation to avoid inconsistent state in zenserver.
If we delete the directory first and then do the delete/create in one operation we will have a temporary state when the zenserver oplog points to non-existing files.

#rnx

diff --git a/Engine/Source/Developer/IoStoreUtilities/Internal/ZenStoreHttpClient.h b/Engine/Source/Developer/IoStoreUtilities/Internal/ZenStoreHttpClient.h
--- a/Engine/Source/Developer/IoStoreUtilities/Internal/ZenStoreHttpClient.h
+++ b/Engine/Source/Developer/IoStoreUtilities/Internal/ZenStoreHttpClient.h
@@ -37,7 +37,8 @@ public:
 	bool TryCreateProject(FStringView InProjectId, FStringView InParentProjectId, FStringView InOplogId, FStringView ServerRoot, 
 					FStringView EngineRoot, FStringView ProjectRoot,
 					FStringView ProjectFilePath);
-	bool TryCreateOplog(FStringView InProjectId, FStringView InOplogId, FStringView InOplogLifetimeMarkerPath, bool bFullBuild);
+	bool TryCreateOplog(FStringView InProjectId, FStringView InOplogId, FStringView InOplogLifetimeMarkerPath);
+	bool TryDeleteOplog(FStringView InProjectId, FStringView InOplogId);
 
 	void InitializeReadOnly(FStringView InProjectId, FStringView InOplogId);
 
diff --git a/Engine/Source/Developer/IoStoreUtilities/Private/ZenStoreHttpClient.cpp b/Engine/Source/Developer/IoStoreUtilities/Private/ZenStoreHttpClient.cpp
--- a/Engine/Source/Developer/IoStoreUtilities/Private/ZenStoreHttpClient.cpp
+++ b/Engine/Source/Developer/IoStoreUtilities/Private/ZenStoreHttpClient.cpp
@@ -139,22 +139,28 @@ FZenStoreHttpClient::TryCreateProject(FStringView InProjectId,
 	return bConnectionSucceeded;
 }
 
-bool FZenStoreHttpClient::TryCreateOplog(FStringView InProjectId, FStringView InOplogId, FStringView InOplogLifetimeMarkerPath, bool bFullBuild)
+bool FZenStoreHttpClient::TryDeleteOplog(FStringView InProjectId, FStringView InOplogId)
 {
 	if (!IsConnected())
 	{
 		return false;
 	}
-
 	UE::Zen::FZenScopedRequestPtr Request(RequestPool.Get());
 
-	if (bFullBuild)
+	UE_LOG(LogZenStore, Display, TEXT("Deleting oplog '%s/%s' if it exists"), *FString(InProjectId), *FString(InOplogId));
+	bool Success = Delete(*Request, OplogPath);
+	return Success;
+}
+
+bool FZenStoreHttpClient::TryCreateOplog(FStringView InProjectId, FStringView InOplogId, FStringView InOplogLifetimeMarkerPath)
+{
+	if (!IsConnected())
 	{
-		UE_LOG(LogZenStore, Display, TEXT("Deleting oplog '%s/%s' if it exists"), *FString(InProjectId), *FString(InOplogId));
-		Delete(*Request, OplogPath);
-		Request->Reset();
+		return false;
 	}
 
+	UE::Zen::FZenScopedRequestPtr Request(RequestPool.Get());
+
 	TArray64<uint8> GetBuffer;
 	FCbObjectView OplogInfo;
 	if (Download(*Request, OplogPath, &GetBuffer, Zen::EContentType::CbObject) && Request->GetResponseCode() == 200)
@@ -771,7 +777,12 @@ bool FZenStoreHttpClient::TryCreateProject(FStringView InProjectId, FStringView
 	return false;
 }
 
-bool FZenStoreHttpClient::TryCreateOplog(FStringView InProjectId, FStringView InOplogId, FStringView InOplogLifetimeMarkerPath, bool bFullBuild)
+bool FZenStoreHttpClient::TryDeleteOplog(FStringView InProjectId, FStringView InOplogId)
+{
+	return false;
+}
+
+bool FZenStoreHttpClient::TryCreateOplog(FStringView InProjectId, FStringView InOplogId, FStringView InOplogLifetimeMarkerPath)
 {
 	return false;
 }
diff --git a/Engine/Source/Developer/IoStoreUtilities/Private/ZenStoreWriter.cpp b/Engine/Source/Developer/IoStoreUtilities/Private/ZenStoreWriter.cpp
--- a/Engine/Source/Developer/IoStoreUtilities/Private/ZenStoreWriter.cpp
+++ b/Engine/Source/Developer/IoStoreUtilities/Private/ZenStoreWriter.cpp
@@ -433,6 +433,9 @@ void FZenStoreWriter::Initialize(const FCookInfo& Info)
 		bool CleanBuild = Info.bFullBuild && !Info.bWorkerOnSharedSandbox;
 		if (CleanBuild)
 		{
+			bool bOplogDeleted = HttpClient->TryDeleteOplog(ProjectId, OplogId);
+			UE_CLOG(!bOplogDeleted, LogZenStoreWriter, Fatal, TEXT("Failed to delete oplog on the ZenServer"));
+
 			UE_LOG(LogZenStoreWriter, Display, TEXT("Deleting %s..."), *OutputPath);
 			const bool bRequireExists = false;
 			const bool bTree = true;
@@ -442,7 +445,7 @@ void FZenStoreWriter::Initialize(const FCookInfo& Info)
 		FString OplogLifetimeMarkerPath = OutputPath / TEXT("ue.projectstore");
 		TUniquePtr<FArchive> OplogMarker(IFileManager::Get().CreateFileWriter(*OplogLifetimeMarkerPath));
 
-		bool bOplogEstablished = HttpClient->TryCreateOplog(ProjectId, OplogId, OplogLifetimeMarkerPath, CleanBuild);
+		bool bOplogEstablished = HttpClient->TryCreateOplog(ProjectId, OplogId, OplogLifetimeMarkerPath);
 
 		if (bOplogEstablished && OplogMarker)
 		{
diff --git a/Engine/Source/Editor/UnrealEd/Private/TargetDomain/TargetDomainUtils.cpp b/Engine/Source/Editor/UnrealEd/Private/TargetDomain/TargetDomainUtils.cpp
--- a/Engine/Source/Editor/UnrealEd/Private/TargetDomain/TargetDomainUtils.cpp
+++ b/Engine/Source/Editor/UnrealEd/Private/TargetDomain/TargetDomainUtils.cpp
@@ -1123,7 +1123,7 @@ FEditorDomainOplog::FEditorDomainOplog()
 	{
 		bool IsLocalConnection = HttpClient.GetZenServiceInstance().IsServiceRunningLocally();
 		HttpClient.TryCreateProject(ProjectId, FStringView(), OplogId, AbsServerRoot, AbsEngineDir, AbsProjectDir, IsLocalConnection ? ProjectFilePath : FStringView());
-		HttpClient.TryCreateOplog(ProjectId, OplogId, TEXT("") /*InOplogMarkerFile*/, false /* bFullBuild */);
+		HttpClient.TryCreateOplog(ProjectId, OplogId, TEXT("") /*InOplogMarkerFile*/);
 	}
 #endif
 }
