Improve the error given when the file pointed to by '-CryptoKey' does not exist.

#rnx

- Previously if the file was missing the process would assert when attempting to parse the missing crypto key file which in turn would generate a full crash report with callstacks etc which makes the problem look much worse than it is.
- Now we check if the file exists and if not we log an error and exit the program with an error code.
- Other error paths in ::LoadKeyChain have been updated from fatal errors (which also kinda display like a crash) to errors that return error codes instead.

diff --git a/Engine/Source/Developer/IoStoreUtilities/Private/IoStoreUtilities.cpp b/Engine/Source/Developer/IoStoreUtilities/Private/IoStoreUtilities.cpp
--- a/Engine/Source/Developer/IoStoreUtilities/Private/IoStoreUtilities.cpp
+++ b/Engine/Source/Developer/IoStoreUtilities/Private/IoStoreUtilities.cpp
@@ -592,7 +592,7 @@ struct FReleasedPackages
 	TMap<FPackageId, FName> PackageIdToName;
 };
 
-static void LoadKeyChain(const TCHAR* CmdLine, FKeyChain& OutCryptoSettings)
+[[nodiscard]] static bool LoadKeyChain(const TCHAR* CmdLine, FKeyChain& OutCryptoSettings)
 {
 	OutCryptoSettings.SetSigningKey(InvalidRSAKeyHandle);
 	OutCryptoSettings.GetEncryptionKeys().Empty();
@@ -601,8 +601,16 @@ static void LoadKeyChain(const TCHAR* CmdLine, FKeyChain& OutCryptoSettings)
 	FString CryptoKeysCacheFilename;
 	if (FParse::Value(CmdLine, TEXT("cryptokeys="), CryptoKeysCacheFilename))
 	{
-		UE_LOG(LogIoStore, Display, TEXT("Parsing crypto keys from a crypto key cache file '%s'"), *CryptoKeysCacheFilename);
-		KeyChainUtilities::LoadKeyChainFromFile(CryptoKeysCacheFilename, OutCryptoSettings);
+		if(IFileManager::Get().FileExists(*CryptoKeysCacheFilename))
+		{
+			UE_LOG(LogIoStore, Display, TEXT("Parsing crypto keys from a crypto key cache file '%s'"), *CryptoKeysCacheFilename);
+			KeyChainUtilities::LoadKeyChainFromFile(CryptoKeysCacheFilename, OutCryptoSettings);
+		}
+		else
+		{
+			UE_LOG(LogIoStore, Error, TEXT("Unable to parse crypto keys as the cache file '%s' does not exist!"), *CryptoKeysCacheFilename);
+			return false;
+		}
 	}
 	else if (FParse::Param(CmdLine, TEXT("encryptionini")))
 	{
@@ -622,7 +630,7 @@ static void LoadKeyChain(const TCHAR* CmdLine, FKeyChain& OutCryptoSettings)
 
 			if (!bDataCryptoRequired)
 			{
-				return;
+				return true;
 			}
 
 			FConfigFile ConfigFile;
@@ -746,7 +754,8 @@ static void LoadKeyChain(const TCHAR* CmdLine, FKeyChain& OutCryptoSettings)
 			// Error checking
 			if (EncryptionKeyString.Len() < RequiredKeyLength)
 			{
-				UE_LOG(LogIoStore, Fatal, TEXT("AES encryption key must be %d characters long"), RequiredKeyLength);
+				UE_LOG(LogIoStore, Error, TEXT("AES encryption key must be %d characters long"), RequiredKeyLength);
+				return false;
 			}
 
 			if (EncryptionKeyString.Len() > RequiredKeyLength)
@@ -757,7 +766,8 @@ static void LoadKeyChain(const TCHAR* CmdLine, FKeyChain& OutCryptoSettings)
 
 			if (!FCString::IsPureAnsi(*EncryptionKeyString))
 			{
-				UE_LOG(LogIoStore, Fatal, TEXT("AES encryption key must be a pure ANSI string!"));
+				UE_LOG(LogIoStore, Error, TEXT("AES encryption key must be a pure ANSI string!"));
+				return false;
 			}
 
 			const auto AsAnsi = StringCast<ANSICHAR>(*EncryptionKeyString);
@@ -776,6 +786,8 @@ static void LoadKeyChain(const TCHAR* CmdLine, FKeyChain& OutCryptoSettings)
 		FGuid::Parse(EncryptionKeyOverrideGuidString, EncryptionKeyOverrideGuid);
 	}
 	OutCryptoSettings.SetPrincipalEncryptionKey(OutCryptoSettings.GetEncryptionKeys().Find(EncryptionKeyOverrideGuid));
+
+	return true;
 }
 
 struct FContainerSourceFile 
@@ -6221,7 +6233,10 @@ int32 ListContainer(
 bool ListIoStoreContainer(const TCHAR* CmdLine)
 {
 	FKeyChain KeyChain;
-	LoadKeyChain(CmdLine, KeyChain);
+	if (!LoadKeyChain(CmdLine, KeyChain))
+	{
+		return false;
+	}
 
 	FString ContainerPathOrWildcard;
 	if (!FParse::Value(FCommandLine::Get(), TEXT("ListContainer="), ContainerPathOrWildcard))
@@ -6452,7 +6467,10 @@ bool ListContainerBulkData(
 bool ListIoStoreContainerBulkData(const TCHAR* CmdLine)
 {
 	FKeyChain KeyChain;
-	LoadKeyChain(CmdLine, KeyChain);
+	if (!LoadKeyChain(CmdLine, KeyChain))
+	{
+		return false;
+	}
 
 	FString ContainerPathOrWildcard;
 	if (!FParse::Value(FCommandLine::Get(), TEXT("ListContainerBulkData="), ContainerPathOrWildcard))
@@ -9444,7 +9462,10 @@ int32 CreateIoStoreContainerFiles(const TCHAR* CmdLine)
 	FIoStoreArguments Arguments;
 	FIoStoreWriterSettings WriterSettings;
 
-	LoadKeyChain(FCommandLine::Get(), Arguments.KeyChain);
+	if (!LoadKeyChain(FCommandLine::Get(), Arguments.KeyChain))
+	{
+		return 1;
+	}
 	
 	ITargetPlatform* TargetPlatform = nullptr;
 	FString TargetPlatformName;
