Ensure that when reading a cook artifact from zen, we allow fall-through to the underlying filesystem if the artifact is a reference and the bytes are not stored in zen.

diff --git a/Engine/Source/Developer/IoStoreUtilities/Private/ZenCookArtifactReader.cpp b/Engine/Source/Developer/IoStoreUtilities/Private/ZenCookArtifactReader.cpp
--- a/Engine/Source/Developer/IoStoreUtilities/Private/ZenCookArtifactReader.cpp
+++ b/Engine/Source/Developer/IoStoreUtilities/Private/ZenCookArtifactReader.cpp
@@ -39,7 +39,7 @@ bool FZenCookArtifactReader::FileExists(const TCHAR* Filename)
 	if (StorageServerPlatformFile)
 	{
 		FString StandardFilename;
-		if (MakeStorageServerPath(Filename, StandardFilename))
+		if (MakeStorageServerPath(Filename, StandardFilename) && (StorageServerPlatformFile->FileSize(*StandardFilename) >= 0))
 		{
 			return StorageServerPlatformFile->FileExists(*StandardFilename);
 		}
@@ -71,7 +71,7 @@ IFileHandle* FZenCookArtifactReader::OpenRead(const TCHAR* Filename)
 	if (StorageServerPlatformFile)
 	{
 		FString StandardFilename;
-		if (MakeStorageServerPath(Filename, StandardFilename))
+		if (MakeStorageServerPath(Filename, StandardFilename) && (StorageServerPlatformFile->FileSize(*StandardFilename) >= 0))
 		{
 			return StorageServerPlatformFile->OpenRead(*StandardFilename);
 		}
diff --git a/Engine/Source/Developer/IoStoreUtilities/Private/ZenStoreWriter.cpp b/Engine/Source/Developer/IoStoreUtilities/Private/ZenStoreWriter.cpp
--- a/Engine/Source/Developer/IoStoreUtilities/Private/ZenStoreWriter.cpp
+++ b/Engine/Source/Developer/IoStoreUtilities/Private/ZenStoreWriter.cpp
@@ -165,7 +165,7 @@ FZenStoreWriter::FZenStoreWriter(
 	const FString& InOutputPath, 
 	const FString& InMetadataDirectoryPath, 
 	const ITargetPlatform* InTargetPlatform,
-	TSharedRef<FZenCookArtifactReader> InCookArtifactReader
+	TSharedRef<ICookArtifactReader> InCookArtifactReader
 )
 	: CookArtifactReader(InCookArtifactReader)
 	, TargetPlatform(*InTargetPlatform)
diff --git a/Engine/Source/Developer/IoStoreUtilities/Public/ZenStoreWriter.h b/Engine/Source/Developer/IoStoreUtilities/Public/ZenStoreWriter.h
--- a/Engine/Source/Developer/IoStoreUtilities/Public/ZenStoreWriter.h
+++ b/Engine/Source/Developer/IoStoreUtilities/Public/ZenStoreWriter.h
@@ -47,9 +47,9 @@ class FCbPackage;
 class FCbWriter;
 class FPackageStoreOptimizer;
 class FPackageStorePackage;
-class FZenCookArtifactReader;
 class FZenFileSystemManifest;
 struct FZenFileSystemManifestEntry;
+class ICookArtifactReader;
 
 /** 
  * A PackageStoreWriter that saves cooked packages for use by IoStore, and stores them in the Zen storage service.
@@ -61,7 +61,7 @@ public:
 	IOSTOREUTILITIES_API FZenStoreWriter(	const FString& OutputPath, 
 											const FString& MetadataDirectoryPath, 
 											const ITargetPlatform* TargetPlatform,
-											TSharedRef<FZenCookArtifactReader> CookArtifactReader);
+											TSharedRef<ICookArtifactReader> CookArtifactReader);
 
 	IOSTOREUTILITIES_API ~FZenStoreWriter();
 
@@ -243,7 +243,7 @@ private:
 	TUniquePtr<UE::FZenStoreHttpClient>	HttpClient;
 	bool IsLocalConnection = true;
 
-	TSharedRef<FZenCookArtifactReader>	CookArtifactReader;
+	TSharedRef<ICookArtifactReader>		CookArtifactReader;
 	const ITargetPlatform&				TargetPlatform;
 	const FName							TargetPlatformFName;
 	FString								ProjectId;
diff --git a/Engine/Source/Editor/UnrealEd/Private/CookOnTheFlyServer.cpp b/Engine/Source/Editor/UnrealEd/Private/CookOnTheFlyServer.cpp
--- a/Engine/Source/Editor/UnrealEd/Private/CookOnTheFlyServer.cpp
+++ b/Engine/Source/Editor/UnrealEd/Private/CookOnTheFlyServer.cpp
@@ -10400,7 +10400,7 @@ UE::Cook::FCookSavePackageContext* UCookOnTheFlyServer::CreateSaveContext(const
 		LayeredReader->AddLayer(ZenReader);
 		AllContextArtifactReader->AddLayer(ZenReader);
 		CookArtifactReader = LayeredReader;
-		FZenStoreWriter* ZenWriter = new FZenStoreWriter(ResolvedRootPath, ResolvedMetadataPath, TargetPlatform, ZenReader);
+		FZenStoreWriter* ZenWriter = new FZenStoreWriter(ResolvedRootPath, ResolvedMetadataPath, TargetPlatform, LayeredReader);
 		ZenWriter->SetBeginCacheCallback(MoveTemp(BeginCacheCallback));
 		ZenWriter->SetRegisterDeterminismHelperCallback(MoveTemp(RegisterDeterminismHelperCallback));
 		PackageWriter = ZenWriter;
