Skip warning about invalid soft package references in package store optimizer

diff --git a/Engine/Source/Developer/IoStoreUtilities/Private/PackageStoreOptimizer.cpp b/Engine/Source/Developer/IoStoreUtilities/Private/PackageStoreOptimizer.cpp
--- a/Engine/Source/Developer/IoStoreUtilities/Private/PackageStoreOptimizer.cpp
+++ b/Engine/Source/Developer/IoStoreUtilities/Private/PackageStoreOptimizer.cpp
@@ -1144,7 +1144,7 @@ FPackageStoreEntryResource FPackageStoreOptimizer::CreatePackageStoreEntry(const
 		{
 			if (!FPackageName::IsValidLongPackageName(SoftRef))
 			{
-				UE_LOG(LogPackageStoreOptimizer, Warning, TEXT("Invalid soft package reference name '%s'"), *SoftRefName.ToString());
+				UE_LOG(LogPackageStoreOptimizer, Display, TEXT("Skipping invalid soft package reference name '%s'"), *SoftRefName.ToString());
 				continue;
 			}
 
