Async fetch and load of cook dependencies.

diff --git a/Engine/Source/Developer/DerivedDataCache/Private/ZenCacheStore.cpp b/Engine/Source/Developer/DerivedDataCache/Private/ZenCacheStore.cpp
--- a/Engine/Source/Developer/DerivedDataCache/Private/ZenCacheStore.cpp
+++ b/Engine/Source/Developer/DerivedDataCache/Private/ZenCacheStore.cpp
@@ -33,9 +33,10 @@
 #include "Serialization/CompactBinaryPackage.h"
 #include "Serialization/CompactBinaryWriter.h"
 #include "Serialization/LargeMemoryWriter.h"
-#include "Serialization/MemoryReader.h"
 #include "Templates/Function.h"
+#include "ZenAsyncCbPackageReceiver.h"
 #include "ZenBackendUtils.h"
+#include "ZenCbPackageReceiver.h"
 #include "ZenSerialization.h"
 
 TRACE_DECLARE_ATOMIC_INT_COUNTER(ZenDDC_Get, TEXT("ZenDDC Get"));
@@ -220,8 +221,7 @@ private:
 	class FHealthReceiver;
 	class FAsyncHealthReceiver;
 
-	class FCbPackageReceiver;
-	class FAsyncCbPackageReceiver;
+	class FRequestAsyncCbPackageReceiver;
 
 	FString NodeName;
 	FString Namespace;
@@ -1495,84 +1495,31 @@ private:
 	FOnHealthComplete OnHealthComplete;
 };
 
-class FZenCacheStore::FCbPackageReceiver final : public IHttpReceiver
+class FZenCacheStore::FRequestAsyncCbPackageReceiver final : public FRequestBase, public Zen::FAsyncCbPackageReceiver
 {
 public:
-	FCbPackageReceiver(const FCbPackageReceiver&) = delete;
-	FCbPackageReceiver& operator=(const FCbPackageReceiver&) = delete;
+	using FOnComplete = TUniqueFunction<void(const THttpUniquePtr<IHttpResponse>& HttpResponse, const FCbPackage& ResponsePackage, FStringView ResponseText)>;
 
-	explicit FCbPackageReceiver(FCbPackage& OutPackage, IHttpReceiver* InNext = nullptr)
-		: Package(OutPackage)
-		, Next(InNext)
-	{
-	}
-
-	const FMemoryView Body() const
-	{
-		return MakeMemoryView(BodyArray);
-	}
-
-private:
-	IHttpReceiver* OnCreate(IHttpResponse& Response) final
-	{
-		return &BodyReceiver;
-	}
+	FRequestAsyncCbPackageReceiver(const FAsyncCbPackageReceiver&) = delete;
+	FRequestAsyncCbPackageReceiver& operator=(const FAsyncCbPackageReceiver&) = delete;
 
-	IHttpReceiver* OnComplete(IHttpResponse& Response) final
-	{
-		if (Response.GetErrorCode() == EHttpErrorCode::None)
-		{
-			EHttpMediaType ContentType = Response.GetContentType();
-			switch (ContentType)
-			{
-				case EHttpMediaType::CbPackage:
-				{
-					const FMemoryView MemoryView = Body();
-					{
-						FMemoryReaderView Ar(MemoryView);
-						if (Zen::Http::TryLoadCbPackage(Package, Ar))
-						{
-							BodyArray.Reset();
-							return Next;
-						}
-					}
-					FMemoryReaderView Ar(MemoryView);
-					if (Package.TryLoad(Ar))
-					{
-						BodyArray.Reset();
-					}
-				}
-				break;
-			}
-		}
-		return Next;
-	}
-
-private:
-	FCbPackage& Package;
-	IHttpReceiver* Next;
-	TArray64<uint8> BodyArray;
-	FHttpByteArrayReceiver BodyReceiver{BodyArray, this};
-};
-
-class FZenCacheStore::FAsyncCbPackageReceiver final : public FRequestBase, public IHttpReceiver
-{
-public:
-	FAsyncCbPackageReceiver(const FAsyncCbPackageReceiver&) = delete;
-	FAsyncCbPackageReceiver& operator=(const FAsyncCbPackageReceiver&) = delete;
-
-	FAsyncCbPackageReceiver(
+	FRequestAsyncCbPackageReceiver(
 		THttpUniquePtr<IHttpRequest>&& InRequest,
 		IRequestOwner* InOwner,
 		Zen::FZenServiceInstance& InZenServiceInstance,
-		FOnRpcComplete&& InOnRpcComplete)
-		: Request(MoveTemp(InRequest))
+		FOnComplete&& InOnComplete)
+		: FAsyncCbPackageReceiver(MoveTemp(InRequest), InZenServiceInstance,
+		[Self = this, LocalOnComplete = MoveTemp(InOnComplete)]
+		(FAsyncCbPackageReceiver* Receiver, const THttpUniquePtr<IHttpResponse>& HttpResponse, const FCbPackage& InPackage) mutable
+		{
+			Self->Owner->LaunchTask(TEXT("FRequestAsyncCbPackageReceiver::OnComplete"), [Self = TRefCountPtr(Self), InnerOnComplete = MoveTemp(LocalOnComplete)]
+			{
+				InnerOnComplete(Self->Response, Self->Package, Self->GetPayloadAsString());
+			});
+		},
+		/*MaxAttempts*/ 3)
 		, Owner(InOwner)
-		, ZenServiceInstance(InZenServiceInstance)
-		, BaseReceiver(Package, this)
-		, OnRpcComplete(MoveTemp(InOnRpcComplete))
 	{
-		Request->SendAsync(this, Response);
 	}
 
 private:
@@ -1588,86 +1535,21 @@ private:
 	{
 		Monitor = LocalResponse.GetMonitor();
 		Owner->Begin(this);
-		return &BaseReceiver;
-	}
-
-	bool ShouldRecoverAndRetry(IHttpResponse& LocalResponse)
-	{
-		if (!ZenServiceInstance.IsServiceRunningLocally())
-		{
-			return false;
-		}
-
-		if ((LocalResponse.GetErrorCode() == EHttpErrorCode::Connect) ||
-			(LocalResponse.GetErrorCode() == EHttpErrorCode::TlsConnect) ||
-			(LocalResponse.GetErrorCode() == EHttpErrorCode::TimedOut))
-		{
-			return true;
-		}
-
-		return false;
+		return FAsyncCbPackageReceiver::OnCreate(LocalResponse);
 	}
 
 	IHttpReceiver* OnComplete(IHttpResponse& LocalResponse) final
 	{
 		Owner->End(this, [Self = this, &LocalResponse]
 		{
-			if (Self->ShouldRecoverAndRetry(LocalResponse) && Self->ZenServiceInstance.TryRecovery())
-			{
-				new FAsyncCbPackageReceiver(MoveTemp(Self->Request), Self->Owner, Self->ZenServiceInstance, MoveTemp(Self->OnRpcComplete));
-				return;
-			}
-
-			Self->Request.Reset();
-			if (Self->OnRpcComplete)
-			{
-				// Launch a task for the completion function since it can execute arbitrary code.
-				Self->Owner->LaunchTask(TEXT("ZenHttpComplete"), [Self = TRefCountPtr(Self)]
-				{
-					// Ensuring that the OnRpcComplete method is destroyed by the time we exit this method by moving it to a local scope variable
-					FOnRpcComplete LocalOnComplete = MoveTemp(Self->OnRpcComplete);
-
-					LocalOnComplete(Self->Response, Self->Package, Self->GetPayloadAsString());
-				});
-			}
+			Self->FAsyncCbPackageReceiver::OnComplete(LocalResponse);
 		});
 		return nullptr;
 	}
 
 private:
-	FString GetPayloadAsString() const
-	{
-		EHttpMediaType ContentType = Response->GetContentType();
-		FMemoryView BodyView = BaseReceiver.Body();
-		switch (ContentType)
-		{
-			case EHttpMediaType::Text:
-			case EHttpMediaType::Yaml:
-			case EHttpMediaType::Json:
-				if (BodyView.GetSize() > 0)
-				{
-					TStringBuilder<128> PayloadText;
-					PayloadText << FAnsiStringView((const ANSICHAR*)BodyView.GetData(), (int)(BodyView.GetSize()));
-					return *PayloadText;
-				}
-				return {};
-			default:
-			{
-				TStringBuilder<128> SizeText;
-				SizeText << "Payload (" << LexToString(ContentType) << "): " << BodyView.GetSize() << " bytes";
-				return *SizeText;
-			}
-		}
-	}
-
-	THttpUniquePtr<IHttpRequest> Request;
-	THttpUniquePtr<IHttpResponse> Response;
 	TRefCountPtr<IHttpResponseMonitor> Monitor;
 	IRequestOwner* Owner;
-	Zen::FZenServiceInstance& ZenServiceInstance;
-	FCbPackage Package;
-	FCbPackageReceiver BaseReceiver;
-	FOnRpcComplete OnRpcComplete;
 };
 
 FZenCacheStore::FZenCacheStore(
@@ -1892,7 +1774,12 @@ void FZenCacheStore::EnqueueAsyncRpc(IRequestOwner& Owner, FCbObject RequestObje
 		}
 		Request->SetContentType(EHttpMediaType::CbObject);
 		Request->SetBody(RequestObject.GetBuffer().MakeOwned());
-		new FAsyncCbPackageReceiver(MoveTemp(Request), &Owner, ZenService.GetInstance(), MoveTemp(OnComplete));
+		FRequestAsyncCbPackageReceiver* Receiver = new FRequestAsyncCbPackageReceiver(
+			MoveTemp(Request),
+			&Owner,
+			ZenService.GetInstance(),
+			MoveTemp(OnComplete));
+		Receiver->SendAsync();
 	});
 }
 
@@ -1907,7 +1794,12 @@ void FZenCacheStore::EnqueueAsyncRpc(IRequestOwner& Owner, const FCbPackage& Req
 		}
 		Request->SetContentType(EHttpMediaType::CbPackage);
 		Request->SetBody(SaveRpcPackage(RequestPackage));
-		new FAsyncCbPackageReceiver(MoveTemp(Request), &Owner, ZenService.GetInstance(), MoveTemp(OnComplete));
+		FRequestAsyncCbPackageReceiver* Receiver = new FRequestAsyncCbPackageReceiver(
+			MoveTemp(Request),
+			&Owner,
+			ZenService.GetInstance(),
+			MoveTemp(OnComplete));
+		Receiver->SendAsync();
 	});
 }
 
diff --git a/Engine/Source/Developer/IoStoreUtilities/Internal/ZenStoreHttpClient.h b/Engine/Source/Developer/IoStoreUtilities/Internal/ZenStoreHttpClient.h
--- a/Engine/Source/Developer/IoStoreUtilities/Internal/ZenStoreHttpClient.h
+++ b/Engine/Source/Developer/IoStoreUtilities/Internal/ZenStoreHttpClient.h
@@ -3,6 +3,7 @@
 #pragma once
 
 #include "Async/Future.h"
+#include "Containers/AnsiString.h"
 #include "Containers/Array.h"
 #include "Containers/ArrayView.h"
 #include "Containers/StringView.h"
@@ -19,11 +20,15 @@ class FCbPackage;
 class FCbObject;
 
 namespace UE {
-	namespace Zen {
-		struct FZenHttpRequestPool;
-		class FZenHttpRequest;
-		enum class EContentType;
+namespace Zen {
+	struct FZenHttpRequestPool;
+	class FZenHttpRequest;
+	enum class EContentType;
+
+	namespace Internal {
+		struct FZenHttpClientState;
 	}
+}
 
 /**
  * HTTP protocol implementation of Zen Store client interface
@@ -55,8 +60,10 @@ public:
 	TIoStatusOr<FIoBuffer> ReadChunk(const FIoChunkId& Id, uint64 Offset = 0, uint64 Size = ~0ull);
 	TIoStatusOr<FIoBuffer> ReadChunk(const FIoHash& RawHash, uint64 Offset = 0, uint64 Size = ~0ull);
 	// TODO: Expand API to add partial chunk requests and ModTag field
-	UE_INTERNAL void ReadChunks(TConstArrayView<FIoChunkId> Ids, TFunctionRef<void(const FIoChunkId& Id, TIoStatusOr<FIoBuffer> Result)> ChunkCallback);
-	UE_INTERNAL void ReadChunks(TConstArrayView<FIoHash> RawHashes, TFunctionRef<void(const FIoHash& RawHash, TIoStatusOr<FIoBuffer> Result)> ChunkCallback);
+	UE_INTERNAL void ReadChunks(TConstArrayView<FIoChunkId> Ids, TUniqueFunction<void(const FIoChunkId& Id, TIoStatusOr<FIoBuffer> Result)>&& ChunkCallback);
+	UE_INTERNAL void ReadChunks(TConstArrayView<FIoHash> RawHashes, TUniqueFunction<void(const FIoHash& RawHash, TIoStatusOr<FIoBuffer> Result)>&& ChunkCallback);
+	UE_INTERNAL void ReadChunksAsync(TConstArrayView<FIoChunkId> Ids, TUniqueFunction<void(const FIoChunkId& Id, TIoStatusOr<FIoBuffer> Result)>&& ChunkCallback);
+	UE_INTERNAL void ReadChunksAsync(TConstArrayView<FIoHash> RawHashes, TUniqueFunction<void(const FIoHash& RawHash, TIoStatusOr<FIoBuffer> Result)>&& ChunkCallback);
 
 	UE_DEPRECATED(5.6, "Use ReadChunk(FIoHash) instead")
 	TIoStatusOr<FIoBuffer> ReadOpLogAttachment(FStringView Id);
@@ -84,17 +91,20 @@ private:
 	bool Download(Zen::FZenHttpRequest& Request, FStringView Uri, TArray64<uint8>* Buffer, Zen::EContentType AcceptType);
 	bool Post(Zen::FZenHttpRequest& Request, FStringView Uri, FCbObjectView Obj);
 	bool Post(Zen::FZenHttpRequest& Request, FStringView Uri, FMemoryView Payload);
-	bool Rpc(Zen::FZenHttpRequest& Request, FStringView Uri, FCbObjectView RpcRequestObject, FCbPackage &OutResponse);
+	void Rpc(FAnsiStringView Uri, FCbObject RpcRequestObject, FCbPackage& OutResponse);
+	void RpcAsync(FAnsiStringView Uri, FCbObject RpcRequestObject, TUniqueFunction<void(const FCbPackage& Response)>&& ResponseCallback);
 	bool Delete(Zen::FZenHttpRequest& Request, FStringView Uri);
 	bool ShouldRecoverAndRetry(Zen::FZenHttpRequest& Request);
 
-	template<typename T>
-	void ReadChunks_Internal(TConstArrayView<T> Requests, TFunctionRef<void(const T& Id, TIoStatusOr<FIoBuffer> Result)> ChunkCallback);
+	template<bool bAsync, typename T>
+	void ReadChunks_Internal(TConstArrayView<T> Requests, TUniqueFunction<void(const T& Id, TIoStatusOr<FIoBuffer> Result)>&& ChunkCallback);
+
+	void InitCommon();
 
 	static const uint32 PoolEntryCount;
-	struct SaltGenerator
+	struct FSaltGenerator
 	{
-		SaltGenerator();
+		FSaltGenerator();
 		inline int32_t Next()
 		{
 			const uint32_t A = ++GOpCounter;
@@ -107,11 +117,16 @@ private:
 #if UE_WITH_ZEN
 	UE::Zen::FScopeZenService ZenService;
 #endif
+	// Legacy request pool
 	TUniquePtr<Zen::FZenHttpRequestPool> RequestPool;
-	SaltGenerator SaltGen;
+
+	// Newer connection pool and request queue
+	TUniquePtr<Zen::Internal::FZenHttpClientState> ClientState;
+
+	FSaltGenerator SaltGen;
 	FString ProjectPath;
 	FString OplogPath;
-	FString OplogRpcPath;
+	FAnsiString OplogRpcPath;
 	FString OplogNewEntryPath;
 	FString OplogPrepNewEntryPath;
 	FString TempDirPath;
diff --git a/Engine/Source/Developer/IoStoreUtilities/IoStoreUtilities.Build.cs b/Engine/Source/Developer/IoStoreUtilities/IoStoreUtilities.Build.cs
--- a/Engine/Source/Developer/IoStoreUtilities/IoStoreUtilities.Build.cs
+++ b/Engine/Source/Developer/IoStoreUtilities/IoStoreUtilities.Build.cs
@@ -18,6 +18,7 @@ public class IoStoreUtilities : ModuleRules
 			"CoreUObject",
 			"DerivedDataCache",
 			"DeveloperToolSettings",
+			"DevHttp",
 			"Json",
 			"PakFile",
 			"Projects",
diff --git a/Engine/Source/Developer/IoStoreUtilities/Private/ZenStoreHttpClient.cpp b/Engine/Source/Developer/IoStoreUtilities/Private/ZenStoreHttpClient.cpp
--- a/Engine/Source/Developer/IoStoreUtilities/Private/ZenStoreHttpClient.cpp
+++ b/Engine/Source/Developer/IoStoreUtilities/Private/ZenStoreHttpClient.cpp
@@ -14,6 +14,7 @@
 
 #include "Async/Async.h"
 #include "HAL/PlatformFileManager.h"
+#include "Http/HttpClient.h"
 #include "IO/IoHash.h"
 #include "Misc/Paths.h"
 #include "Misc/ScopeLock.h"
@@ -23,43 +24,79 @@
 #include "Serialization/BufferArchive.h"
 #include "Serialization/CompactBinaryWriter.h"
 #include "Serialization/LargeMemoryWriter.h"
+#include "Tasks/Task.h"
+#include "ZenAsyncCbPackageReceiver.h"
+#include "ZenCbPackageReceiver.h"
 #include "ZenSerialization.h"
 #include "ZenServerHttp.h"
 
 DEFINE_LOG_CATEGORY_STATIC(LogZenStore, Log, All);
 
-namespace UE { 
+namespace UE
+{
+namespace Zen::Internal
+{
+struct FZenHttpClientState
+{
+	THttpUniquePtr<IHttpConnectionPool> ConnectionPool;
+	TUniquePtr<FHttpRequestQueue> RequestQueue;
+};
+}
 
 const uint32 FZenStoreHttpClient::PoolEntryCount = 32;
 
-std::atomic<uint32> FZenStoreHttpClient::SaltGenerator::GOpCounter(0);
+std::atomic<uint32> FZenStoreHttpClient::FSaltGenerator::GOpCounter(0);
 
-FZenStoreHttpClient::SaltGenerator::SaltGenerator()
+FZenStoreHttpClient::FSaltGenerator::FSaltGenerator()
 	: SaltBase(FGenericPlatformProcess::GetCurrentProcessId() + 0x9e3779b9u)
 {
 }
 
 FZenStoreHttpClient::FZenStoreHttpClient()
 {
-	RequestPool = MakeUnique<Zen::FZenHttpRequestPool>(ZenService.GetInstance().GetURL(), PoolEntryCount);
+	InitCommon();
 }
 
 FZenStoreHttpClient::FZenStoreHttpClient(FStringView HostName, uint16 Port)
 : ZenService(HostName.IsEmpty() ? FStringView() : *WriteToString<64>(TEXT("http://"), HostName, TEXT(":"), Port))
 {
-	RequestPool = MakeUnique<Zen::FZenHttpRequestPool>(ZenService.GetInstance().GetURL(), PoolEntryCount);
+	InitCommon();
 }
 
 FZenStoreHttpClient::FZenStoreHttpClient(UE::Zen::FServiceSettings&& InSettings)
 : ZenService(MoveTemp(InSettings))
 {
-	RequestPool = MakeUnique<Zen::FZenHttpRequestPool>(ZenService.GetInstance().GetURL(), PoolEntryCount);
+	InitCommon();
 }
 
 FZenStoreHttpClient::~FZenStoreHttpClient()
 {
 }
 
+void
+FZenStoreHttpClient::InitCommon()
+{
+	RequestPool = MakeUnique<Zen::FZenHttpRequestPool>(ZenService.GetInstance().GetURL(), PoolEntryCount);
+	ClientState = MakeUnique<Zen::Internal::FZenHttpClientState>();
+
+	const uint32 MaxConnections = uint32(FMath::Clamp(FPlatformMisc::NumberOfCoresIncludingHyperthreads(), 8, 64));
+	constexpr uint32 RequestPoolSize = 128;
+	constexpr uint32 RequestPoolOverflowSize = 128;
+
+	FHttpConnectionPoolParams ConnectionPoolParams;
+	ConnectionPoolParams.MaxConnections = MaxConnections;
+	ConnectionPoolParams.MinConnections = MaxConnections;
+	ClientState->ConnectionPool = IHttpManager::Get().CreateConnectionPool(ConnectionPoolParams);
+
+	FHttpClientParams ClientParams;
+	ClientParams.MaxRequests = RequestPoolSize + RequestPoolOverflowSize;
+	ClientParams.MinRequests = RequestPoolSize;
+	ClientParams.LowSpeedLimit = 1;
+	ClientParams.LowSpeedTime = 25;
+	ClientParams.bBypassProxy = true;
+	ClientState->RequestQueue = MakeUnique<FHttpRequestQueue>(*ClientState->ConnectionPool, ClientParams);
+}
+
 bool
 FZenStoreHttpClient::TryCreateProject(FStringView InProjectId,
 	FStringView InParentProjectId,
@@ -130,7 +167,7 @@ FZenStoreHttpClient::TryCreateProject(FStringView InProjectId,
 
 	ProjectPath = WriteToString<128>("/prj/", InProjectId);
 	OplogPath = WriteToString<128>("/prj/", InProjectId, "/oplog/", InOplogId);
-	OplogRpcPath = WriteToString<128>("/prj/", InProjectId, "/oplog/", InOplogId, "/rpc");
+	OplogRpcPath = WriteToAnsiString<128>(ZenService.GetInstance().GetURL(), "/prj/", InProjectId, "/oplog/", InOplogId, "/rpc");
 	OplogNewEntryPath = WriteToString<128>("/prj/", InProjectId, "/oplog/", InOplogId, "/new");
 	OplogPrepNewEntryPath = WriteToString<128>("/prj/", InProjectId, "/oplog/", InOplogId, "/prep");
 
@@ -290,15 +327,62 @@ bool FZenStoreHttpClient::Post(Zen::FZenHttpRequest& Request, FStringView Uri, F
 	return Res == Zen::FZenHttpRequest::Result::Success;
 }
 
-bool FZenStoreHttpClient::Rpc(Zen::FZenHttpRequest& Request, FStringView Uri, FCbObjectView RpcRequestObject, FCbPackage &OutResponse)
+void FZenStoreHttpClient::Rpc(FAnsiStringView Uri, FCbObject RpcRequestObject, FCbPackage& OutResponse)
 {
-	Zen::FZenHttpRequest::Result Res = Request.PerformRpc(Uri, RpcRequestObject, OutResponse);
-	if (ShouldRecoverAndRetry(Request))
+	THttpUniquePtr<IHttpRequest> Request = ClientState->RequestQueue->CreateRequest({});
+	if (!Request)
 	{
-		Request.Reset();
-		Res = Request.PerformRpc(Uri, RpcRequestObject, OutResponse);
+		return;
 	}
-	return Res == Zen::FZenHttpRequest::Result::Success;
+
+	constexpr int MaxAttempts = 2;
+	int Attempt = 0;
+	bool bRetry = false;
+	Request->SetUri(Uri);
+	Request->SetMethod(EHttpMethod::Post);
+	Request->AddAcceptType(EHttpMediaType::CbPackage);
+	Request->SetContentType(EHttpMediaType::CbObject);
+	Request->SetBody(RpcRequestObject.GetBuffer().MakeOwned());
+	do
+	{
+		Zen::FCbPackageReceiver Receiver(OutResponse);
+		THttpUniquePtr<IHttpResponse> HttpResponse;
+		Request->Send(&Receiver, HttpResponse);
+		int StatusCode = HttpResponse->GetStatusCode();
+
+		bRetry = Receiver.ShouldRecoverAndRetry(ZenService.GetInstance(), *HttpResponse) && (++Attempt < MaxAttempts);
+	}
+	while (bRetry);
+}
+
+void FZenStoreHttpClient::RpcAsync(FAnsiStringView Uri, FCbObject RpcRequestObject, TUniqueFunction<void(const FCbPackage& Response)>&& ResponseCallback)
+{
+	THttpUniquePtr<IHttpRequest> Request = ClientState->RequestQueue->CreateRequest({});
+	if (!Request)
+	{
+		FCbPackage Response;
+		ResponseCallback(Response);
+		return;
+	}
+
+	Request->SetUri(Uri);
+	Request->SetMethod(EHttpMethod::Post);
+	Request->AddAcceptType(EHttpMediaType::CbPackage);
+	Request->SetContentType(EHttpMediaType::CbObject);
+	Request->SetBody(RpcRequestObject.GetBuffer().MakeOwned());
+	Zen::FAsyncCbPackageReceiver* PackageReceiver = new Zen::FAsyncCbPackageReceiver(MoveTemp(Request), ZenService.GetInstance(),
+		[ResponseCallback = MoveTemp(ResponseCallback)]
+		(Zen::FAsyncCbPackageReceiver* Receiver, const THttpUniquePtr<IHttpResponse>& HttpResponse, const FCbPackage& Response) mutable
+		{
+			// Launch a task for the completion function since it can execute arbitrary code.
+			Tasks::Launch(TEXT("FZenStoreHttpClient::RpcAsync"), [Receiver, &Response, InnerResponseCallback = MoveTemp(ResponseCallback)]
+			{
+				InnerResponseCallback(Response);
+				delete Receiver;
+			});
+		},
+		/*MaxAttempts*/ 2);
+	PackageReceiver->SendAsync();
 }
 
 bool FZenStoreHttpClient::Delete(Zen::FZenHttpRequest& Request, FStringView Uri)
@@ -577,10 +661,10 @@ static const TCHAR* DescribeRequest(const FIoHash& RawHash)
 	return TEXT("RawHash");
 }
 
-template<typename T>
-void FZenStoreHttpClient::ReadChunks_Internal(TConstArrayView<T> Requests, TFunctionRef<void(const T& Id, TIoStatusOr<FIoBuffer> Result)> ChunkCallback)
+template<bool bAsync, typename T>
+void FZenStoreHttpClient::ReadChunks_Internal(TConstArrayView<T> Requests, TUniqueFunction<void(const T& Id, TIoStatusOr<FIoBuffer> Result)>&& ChunkCallback)
 {
-	TRACE_CPUPROFILER_EVENT_SCOPE(ZenStoreHttp_ReadChunks_Internal);
+	TRACE_CPUPROFILER_EVENT_SCOPE(ZenStoreHttp_ReadChunksAsync_Internal);
 	FCbWriter Writer;
 	Writer.BeginObject();
 	Writer << "method" << "getchunks";
@@ -604,10 +688,19 @@ void FZenStoreHttpClient::ReadChunks_Internal(TConstArrayView<T> Requests, TFunc
 	Writer.EndObject();
 	Writer.EndObject();
 
-	UE::Zen::FZenScopedRequestPtr ZenHttpRequest(RequestPool.Get());
-	FCbPackage ResultPackage;
-	if (Rpc(*ZenHttpRequest, OplogRpcPath, Writer.Save().AsObjectView(), ResultPackage))
+	auto OnGetChunksResponse = [InProgressRequests = MoveTemp(InProgressRequests)]
+	(TUniqueFunction<void(const T& Id, TIoStatusOr<FIoBuffer> Result)>&& ChunkCallback, const FCbPackage& ResultPackage) mutable
 	{
+		if (!ResultPackage)
+		{
+			for (const T& InProgressRequest : InProgressRequests)
+			{
+				ChunkCallback(InProgressRequest, (FIoStatus)(FIoStatusBuilder(EIoErrorCode::NotFound) << TEXT("Transport/request failure, ")
+					<< DescribeRequest(InProgressRequest) << TEXT("=") << *WriteToString<32>(InProgressRequest)));
+			}
+			return;
+		}
+
 		const FCbObjectView ResultObjectView(ResultPackage.GetObject());
 		for (FCbFieldView ChunkField : ResultObjectView[ANSITEXTVIEW("Chunks")])
 		{
@@ -660,25 +753,43 @@ void FZenStoreHttpClient::ReadChunks_Internal(TConstArrayView<T> Requests, TFunc
 			ChunkCallback(InProgressRequest, (FIoStatus)(FIoStatusBuilder(EIoErrorCode::NotFound) << TEXT("Missing response, ")
 				<< DescribeRequest(InProgressRequest) << TEXT("=") << *WriteToString<32>(InProgressRequest)));
 		}
+	};
+
+	if constexpr (bAsync)
+	{
+		RpcAsync(OplogRpcPath, Writer.Save().AsObject(),
+			[ChunkCallback = MoveTemp(ChunkCallback), OnGetChunksResponse = MoveTemp(OnGetChunksResponse)]
+			(const FCbPackage& Response) mutable
+			{
+				OnGetChunksResponse(MoveTemp(ChunkCallback), Response);
+			});
 	}
 	else
 	{
-		for (const T& InProgressRequest : InProgressRequests)
-		{
-			ChunkCallback(InProgressRequest, (FIoStatus)(FIoStatusBuilder(EIoErrorCode::NotFound) << TEXT("Transport/request failure, ")
-				<< DescribeRequest(InProgressRequest) << TEXT("=") << *WriteToString<32>(InProgressRequest)));
-		}
+		FCbPackage ResultPackage;
+		Rpc(OplogRpcPath, Writer.Save().AsObject(), ResultPackage);
+		OnGetChunksResponse(MoveTemp(ChunkCallback), MoveTemp(ResultPackage));
 	}
 }
 
-void FZenStoreHttpClient::ReadChunks(TConstArrayView<FIoChunkId> Ids, TFunctionRef<void(const FIoChunkId& Id, TIoStatusOr<FIoBuffer> Result)> ChunkCallback)
+void FZenStoreHttpClient::ReadChunks(TConstArrayView<FIoChunkId> Ids, TUniqueFunction<void(const FIoChunkId& Id, TIoStatusOr<FIoBuffer> Result)>&& ChunkCallback)
+{
+	return ReadChunks_Internal<false, FIoChunkId>(Ids, MoveTemp(ChunkCallback));
+}
+
+void FZenStoreHttpClient::ReadChunks(TConstArrayView<FIoHash> RawHashes, TUniqueFunction<void(const FIoHash& RawHash, TIoStatusOr<FIoBuffer> Result)>&& ChunkCallback)
+{
+	return ReadChunks_Internal<false, FIoHash>(RawHashes, MoveTemp(ChunkCallback));
+}
+
+void FZenStoreHttpClient::ReadChunksAsync(TConstArrayView<FIoChunkId> Ids, TUniqueFunction<void(const FIoChunkId& Id, TIoStatusOr<FIoBuffer> Result)>&& ChunkCallback)
 {
-	return ReadChunks_Internal(Ids, ChunkCallback);
+	return ReadChunks_Internal<true, FIoChunkId>(Ids, MoveTemp(ChunkCallback));
 }
 
-void FZenStoreHttpClient::ReadChunks(TConstArrayView<FIoHash> RawHashes, TFunctionRef<void(const FIoHash& RawHash, TIoStatusOr<FIoBuffer> Result)> ChunkCallback)
+void FZenStoreHttpClient::ReadChunksAsync(TConstArrayView<FIoHash> RawHashes, TUniqueFunction<void(const FIoHash& RawHash, TIoStatusOr<FIoBuffer> Result)>&& ChunkCallback)
 {
-	return ReadChunks_Internal(RawHashes, ChunkCallback);
+	return ReadChunks_Internal<true, FIoHash>(RawHashes, MoveTemp(ChunkCallback));
 }
 
 TIoStatusOr<FIoBuffer> FZenStoreHttpClient::ReadOpLogAttachment(FStringView Id)
@@ -961,7 +1072,23 @@ TIoStatusOr<FIoBuffer> FZenStoreHttpClient::ReadChunk(const FIoHash& RawHash, ui
 	return FIoBuffer();
 }
 
-void FZenStoreHttpClient::ReadChunks(TConstArrayView<FIoChunkId> Ids, TFunctionRef<void(const FIoChunkId& Id, TIoStatusOr<FIoBuffer> Result)> ChunkCallback)
+void FZenStoreHttpClient::ReadChunks(TConstArrayView<FIoChunkId> Ids, TUniqueFunction<void(const FIoChunkId& Id, TIoStatusOr<FIoBuffer> Result)>&& ChunkCallback)
+{
+	for (const FIoChunkId& Id : Ids)
+	{
+		ChunkCallback(Id, (FIoStatus)(FIoStatusBuilder(EIoErrorCode::Unknown) << TEXT("Unsupported, Oid=") << *WriteToString<32>(Id)));
+	}
+}
+
+void FZenStoreHttpClient::ReadChunks(TConstArrayView<FIoHash> RawHashes, TUniqueFunction<void(const FIoHash& RawHash, TIoStatusOr<FIoBuffer> Result)>&& ChunkCallback)
+{
+	for (const FIoHash& RawHash : RawHashes)
+	{
+		ChunkCallback(RawHash, (FIoStatus)(FIoStatusBuilder(EIoErrorCode::Unknown) << TEXT("Unsupported, RawHash=") << *WriteToString<32>(RawHash)));
+	}
+}
+
+void FZenStoreHttpClient::ReadChunksAsync(TConstArrayView<FIoChunkId> Ids, TUniqueFunction<void(const FIoChunkId& Id, TIoStatusOr<FIoBuffer> Result)>&& ChunkCallback)
 {
 	for (const FIoChunkId& Id : Ids)
 	{
@@ -969,7 +1096,7 @@ void FZenStoreHttpClient::ReadChunks(TConstArrayView<FIoChunkId> Ids, TFunctionR
 	}
 }
 
-void FZenStoreHttpClient::ReadChunks(TConstArrayView<FIoHash> RawHashes, TFunctionRef<void(const FIoHash& RawHash, TIoStatusOr<FIoBuffer> Result)> ChunkCallback)
+void FZenStoreHttpClient::ReadChunksAsync(TConstArrayView<FIoHash> RawHashes, TUniqueFunction<void(const FIoHash& RawHash, TIoStatusOr<FIoBuffer> Result)>&& ChunkCallback)
 {
 	for (const FIoHash& RawHash : RawHashes)
 	{
diff --git a/Engine/Source/Developer/IoStoreUtilities/Private/ZenStoreWriter.cpp b/Engine/Source/Developer/IoStoreUtilities/Private/ZenStoreWriter.cpp
--- a/Engine/Source/Developer/IoStoreUtilities/Private/ZenStoreWriter.cpp
+++ b/Engine/Source/Developer/IoStoreUtilities/Private/ZenStoreWriter.cpp
@@ -1188,7 +1188,7 @@ FCbObject FZenStoreWriter::GetOplogAttachment(FName PackageName, FUtf8StringView
 
 void FZenStoreWriter::GetOplogAttachments(TArrayView<FName> PackageNames,
 	TArrayView<FUtf8StringView> AttachmentKeys,
-	TFunctionRef<void(FName PackageName, FUtf8StringView AttachmentKey, FCbObject&& Attachment)> Callback)
+	TUniqueFunction<void(FName PackageName, FUtf8StringView AttachmentKey, FCbObject&& Attachment)>&& Callback)
 {
 	const int MaximumHashCount = PackageNames.Num() * AttachmentKeys.Num();
 	TArray<FIoHash> AttachmentHashes;
@@ -1279,7 +1279,7 @@ void FZenStoreWriter::GetOplogAttachments(TArrayView<FName> PackageNames,
 		return;
 	}
 
-	HttpClient->ReadChunks(AttachmentHashes, [Callback,&AttachmentHashParams](const FIoHash& RawHash, TIoStatusOr<FIoBuffer> Result)
+	HttpClient->ReadChunksAsync(AttachmentHashes, [Callback = MoveTemp(Callback), AttachmentHashParams = MoveTemp(AttachmentHashParams)](const FIoHash& RawHash, TIoStatusOr<FIoBuffer> Result)
 	{
 		for (auto It(AttachmentHashParams.CreateConstKeyIterator(RawHash)); It; ++It)
 		{
diff --git a/Engine/Source/Developer/IoStoreUtilities/Public/ZenStoreWriter.h b/Engine/Source/Developer/IoStoreUtilities/Public/ZenStoreWriter.h
--- a/Engine/Source/Developer/IoStoreUtilities/Public/ZenStoreWriter.h
+++ b/Engine/Source/Developer/IoStoreUtilities/Public/ZenStoreWriter.h
@@ -149,7 +149,7 @@ public:
 	IOSTOREUTILITIES_API virtual FCbObject GetOplogAttachment(FName PackageName, FUtf8StringView AttachmentKey) override;
 	IOSTOREUTILITIES_API virtual void GetOplogAttachments(TArrayView<FName> PackageNames,
 		TArrayView<FUtf8StringView> AttachmentKeys,
-		TFunctionRef<void(FName PackageName, FUtf8StringView AttachmentKey, FCbObject&& Attachment)> Callback) override;
+		TUniqueFunction<void(FName PackageName, FUtf8StringView AttachmentKey, FCbObject&& Attachment)>&& Callback) override;
 	IOSTOREUTILITIES_API virtual ECommitStatus GetCommitStatus(FName PackageName) override;
 
 	IOSTOREUTILITIES_API virtual void RemoveCookedPackages(TArrayView<const FName> PackageNamesToRemove) override;
diff --git a/Engine/Source/Developer/Zen/Internal/ZenAsyncCbPackageReceiver.h b/Engine/Source/Developer/Zen/Internal/ZenAsyncCbPackageReceiver.h
new file mode 100644
--- /dev/null
+++ b/Engine/Source/Developer/Zen/Internal/ZenAsyncCbPackageReceiver.h
@@ -0,0 +1,56 @@
+// Copyright Epic Games, Inc. All Rights Reserved.
+
+#pragma once
+
+#include "Experimental/ZenGlobals.h"
+
+#if UE_WITH_ZEN
+
+#include "Containers/StringFwd.h"
+#include "Http/HttpClient.h"
+#include "Serialization/CompactBinaryPackage.h"
+#include "ZenCbPackageReceiver.h"
+
+namespace UE::Zen
+{
+
+class FZenServiceInstance;
+
+class FAsyncCbPackageReceiver : public IHttpReceiver
+{
+public:
+	using FOnComplete = TUniqueFunction<void(FAsyncCbPackageReceiver* Receiver, const THttpUniquePtr<IHttpResponse>& HttpResponse, const FCbPackage& ResponsePackage)>;
+
+	FAsyncCbPackageReceiver(const FAsyncCbPackageReceiver&) = delete;
+	FAsyncCbPackageReceiver& operator=(const FAsyncCbPackageReceiver&) = delete;
+
+	ZEN_API FAsyncCbPackageReceiver(
+		THttpUniquePtr<IHttpRequest>&& InRequest,
+		FZenServiceInstance& InZenServiceInstance,
+		FOnComplete&& InOnComplete,
+		int InMaxAttempts = 1);
+
+	ZEN_API void SendAsync();
+
+protected:
+	// IHttpReceiver Interface
+
+	ZEN_API IHttpReceiver* OnCreate(IHttpResponse& LocalResponse);
+	ZEN_API IHttpReceiver* OnComplete(IHttpResponse& LocalResponse);
+
+	ZEN_API FString GetPayloadAsString() const;
+
+protected:
+	THttpUniquePtr<IHttpRequest> Request;
+	THttpUniquePtr<IHttpResponse> Response;
+	FZenServiceInstance& ZenServiceInstance;
+	FCbPackage Package;
+	UE::Zen::FCbPackageReceiver BaseReceiver;
+	FOnComplete OnCompleteCallback;
+	const int MaxAttempts;
+	int Attempt;
+};
+
+} // namespace UE::Zen
+
+#endif // UE_WITH_ZEN
\ No newline at end of file
diff --git a/Engine/Source/Developer/Zen/Internal/ZenCbPackageReceiver.h b/Engine/Source/Developer/Zen/Internal/ZenCbPackageReceiver.h
new file mode 100644
--- /dev/null
+++ b/Engine/Source/Developer/Zen/Internal/ZenCbPackageReceiver.h
@@ -0,0 +1,45 @@
+// Copyright Epic Games, Inc. All Rights Reserved.
+
+#pragma once
+
+#include "Experimental/ZenGlobals.h"
+
+#if UE_WITH_ZEN
+
+#include "Containers/Array.h"
+#include "Http/HttpClient.h"
+
+class FCbPackage;
+
+namespace UE::Zen
+{
+
+class FZenServiceInstance;
+
+class FCbPackageReceiver final : public IHttpReceiver
+{
+public:
+	FCbPackageReceiver(const FCbPackageReceiver&) = delete;
+	FCbPackageReceiver& operator=(const FCbPackageReceiver&) = delete;
+
+	ZEN_API explicit FCbPackageReceiver(FCbPackage& OutPackage, IHttpReceiver* InNext = nullptr);
+
+	ZEN_API void Reset();
+	ZEN_API const FMemoryView Body() const;
+
+	ZEN_API static bool ShouldRecoverAndRetry(FZenServiceInstance& ZenServiceInstance, IHttpResponse& LocalResponse);
+
+private:
+	ZEN_API IHttpReceiver* OnCreate(IHttpResponse& Response) final;
+	ZEN_API IHttpReceiver* OnComplete(IHttpResponse& Response) final;
+
+private:
+	FCbPackage& Package;
+	IHttpReceiver* Next;
+	TArray64<uint8> BodyArray;
+	FHttpByteArrayReceiver BodyReceiver{BodyArray, this};
+};
+
+} // namespace UE::Zen
+
+#endif // UE_WITH_ZEN
\ No newline at end of file
diff --git a/Engine/Source/Developer/Zen/Private/ZenAsyncCbPackageReceiver.cpp b/Engine/Source/Developer/Zen/Private/ZenAsyncCbPackageReceiver.cpp
new file mode 100644
--- /dev/null
+++ b/Engine/Source/Developer/Zen/Private/ZenAsyncCbPackageReceiver.cpp
@@ -0,0 +1,86 @@
+// Copyright Epic Games, Inc. All Rights Reserved.
+
+#include "ZenAsyncCbPackageReceiver.h"
+
+#if UE_WITH_ZEN
+
+#include "Containers/UnrealString.h"
+#include "Experimental/ZenServerInterface.h"
+#include "Serialization/MemoryReader.h"
+#include "ZenSerialization.h"
+
+namespace UE::Zen
+{
+
+FAsyncCbPackageReceiver::FAsyncCbPackageReceiver(
+	THttpUniquePtr<IHttpRequest>&& InRequest,
+	Zen::FZenServiceInstance& InZenServiceInstance,
+	FOnComplete&& InOnComplete,
+	int InMaxAttempts)
+	: Request(MoveTemp(InRequest))
+	, ZenServiceInstance(InZenServiceInstance)
+	, BaseReceiver(Package, this)
+	, OnCompleteCallback(MoveTemp(InOnComplete))
+	, MaxAttempts(InMaxAttempts)
+	, Attempt(0)
+{
+}
+
+void FAsyncCbPackageReceiver::SendAsync()
+{
+	Request->SendAsync(this, Response);
+}
+
+IHttpReceiver* FAsyncCbPackageReceiver::OnCreate(IHttpResponse& LocalResponse)
+{
+	return &BaseReceiver;
+}
+
+IHttpReceiver* FAsyncCbPackageReceiver::OnComplete(IHttpResponse& LocalResponse)
+{
+	if ((++Attempt < MaxAttempts) && FCbPackageReceiver::ShouldRecoverAndRetry(ZenServiceInstance, LocalResponse) && ZenServiceInstance.TryRecovery())
+	{
+		BaseReceiver.Reset();
+		SendAsync();
+		return nullptr;
+	}
+
+	Request.Reset();
+	if (OnCompleteCallback)
+	{
+		// Ensuring that the OnComplete method is destroyed by the time we exit this method by moving it to a local scope variable
+		FOnComplete LocalOnComplete = MoveTemp(OnCompleteCallback);
+
+		// Calling LocalOnComplete may result in "this" being deleted, so no further access can happen to anything on "this"
+		LocalOnComplete(this, Response, Package);
+	}
+	return nullptr;
+}
+
+FString FAsyncCbPackageReceiver::GetPayloadAsString() const
+{
+	EHttpMediaType ContentType = Response->GetContentType();
+	FMemoryView BodyView = BaseReceiver.Body();
+	switch (ContentType)
+	{
+		case EHttpMediaType::Text:
+		case EHttpMediaType::Yaml:
+		case EHttpMediaType::Json:
+			if (BodyView.GetSize() > 0)
+			{
+				TStringBuilder<128> PayloadText;
+				PayloadText << FAnsiStringView((const ANSICHAR*)BodyView.GetData(), (int)(BodyView.GetSize()));
+				return *PayloadText;
+			}
+			return {};
+		default:
+		{
+			TStringBuilder<128> SizeText;
+			SizeText << "Payload (" << LexToString(ContentType) << "): " << BodyView.GetSize() << " bytes";
+			return *SizeText;
+		}
+	}
+}
+
+}
+#endif // UE_WITH_ZEN
diff --git a/Engine/Source/Developer/Zen/Private/ZenCbPackageReceiver.cpp b/Engine/Source/Developer/Zen/Private/ZenCbPackageReceiver.cpp
new file mode 100644
--- /dev/null
+++ b/Engine/Source/Developer/Zen/Private/ZenCbPackageReceiver.cpp
@@ -0,0 +1,84 @@
+// Copyright Epic Games, Inc. All Rights Reserved.
+
+#include "ZenCbPackageReceiver.h"
+
+#if UE_WITH_ZEN
+
+#include "Experimental/ZenServerInterface.h"
+#include "Serialization/CompactBinaryPackage.h"
+#include "Serialization/MemoryReader.h"
+#include "ZenSerialization.h"
+
+namespace UE::Zen
+{
+
+FCbPackageReceiver::FCbPackageReceiver(FCbPackage& OutPackage, IHttpReceiver* InNext)
+	: Package(OutPackage)
+	, Next(InNext)
+{
+}
+
+void FCbPackageReceiver::Reset()
+{
+	BodyArray.Reset();
+}
+
+const FMemoryView FCbPackageReceiver::Body() const
+{
+	return MakeMemoryView(BodyArray);
+}
+
+bool FCbPackageReceiver::ShouldRecoverAndRetry(Zen::FZenServiceInstance& ZenServiceInstance, IHttpResponse& LocalResponse)
+{
+	if (!ZenServiceInstance.IsServiceRunningLocally())
+	{
+		return false;
+	}
+
+	if ((LocalResponse.GetErrorCode() == EHttpErrorCode::Connect) ||
+		(LocalResponse.GetErrorCode() == EHttpErrorCode::TlsConnect) ||
+		(LocalResponse.GetErrorCode() == EHttpErrorCode::TimedOut))
+	{
+		return true;
+	}
+
+	return false;
+}
+
+IHttpReceiver* FCbPackageReceiver::OnCreate(IHttpResponse& Response)
+{
+	return &BodyReceiver;
+}
+
+IHttpReceiver* FCbPackageReceiver::OnComplete(IHttpResponse& Response)
+{
+	if (Response.GetErrorCode() == EHttpErrorCode::None)
+	{
+		EHttpMediaType ContentType = Response.GetContentType();
+		switch (ContentType)
+		{
+			case EHttpMediaType::CbPackage:
+			{
+				const FMemoryView MemoryView = Body();
+				{
+					FMemoryReaderView Ar(MemoryView);
+					if (Zen::Http::TryLoadCbPackage(Package, Ar))
+					{
+						BodyArray.Reset();
+						return Next;
+					}
+				}
+				FMemoryReaderView Ar(MemoryView);
+				if (Package.TryLoad(Ar))
+				{
+					BodyArray.Reset();
+				}
+			}
+			break;
+		}
+	}
+	return Next;
+}
+
+}
+#endif // UE_WITH_ZEN
diff --git a/Engine/Source/Developer/Zen/Private/ZenServerInterface.cpp b/Engine/Source/Developer/Zen/Private/ZenServerInterface.cpp
--- a/Engine/Source/Developer/Zen/Private/ZenServerInterface.cpp
+++ b/Engine/Source/Developer/Zen/Private/ZenServerInterface.cpp
@@ -2617,13 +2617,13 @@ FZenServiceInstance::RequestGC(const bool* OverrideCollectSmallObjects, const ui
 
 	if (OverrideCollectSmallObjects)
 	{
-		Query << Separators[SeparatorIndex] << TEXT("smallobjects=") << LexToString(*OverrideCollectSmallObjects);
+		Query << Separators[SeparatorIndex] << TEXT("smallobjects=") << ::LexToString(*OverrideCollectSmallObjects);
 		SeparatorIndex = FMath::Min(SeparatorIndex + 1, (int32)UE_ARRAY_COUNT(Separators));
 	}
 
 	if (OverrideMaxCacheDuration)
 	{
-		Query << Separators[SeparatorIndex] << TEXT("maxcacheduration=") << LexToString(*OverrideMaxCacheDuration);
+		Query << Separators[SeparatorIndex] << TEXT("maxcacheduration=") << ::LexToString(*OverrideMaxCacheDuration);
 		SeparatorIndex = FMath::Min(SeparatorIndex + 1, (int32)UE_ARRAY_COUNT(Separators));
 	}
 
diff --git a/Engine/Source/Developer/Zen/Zen.Build.cs b/Engine/Source/Developer/Zen/Zen.Build.cs
--- a/Engine/Source/Developer/Zen/Zen.Build.cs
+++ b/Engine/Source/Developer/Zen/Zen.Build.cs
@@ -7,17 +7,16 @@ public class Zen : ModuleRules
 {
 	public Zen(ReadOnlyTargetRules Target) : base(Target)
 	{
-		PrivateDependencyModuleNames.AddRange(new string[] { "Core", "Sockets", "SSL", "Json" });
-		PrivateIncludePathModuleNames.Add("DesktopPlatform");
-		PrivateIncludePathModuleNames.Add("Analytics");
-
+		PrivateDependencyModuleNames.AddRange(new string[] { "Core", "Sockets", "SSL", "Json", "DevHttp" });
+		PrivateIncludePathModuleNames.AddRange(new string[] { "DesktopPlatform", "Analytics" });
+		
 		if (Target.Platform.IsInGroup(UnrealPlatformGroup.Windows) && Target.WindowsPlatform.bUseXCurl)
 		{
 			PublicDependencyModuleNames.Add("XCurl");
 		}
 		else
 		{		
-			AddEngineThirdPartyPrivateStaticDependencies(Target, "libcurl");
+		AddEngineThirdPartyPrivateStaticDependencies(Target, "libcurl");
 		}
 		AddEngineThirdPartyPrivateStaticDependencies(Target, "OpenSSL");
 
diff --git a/Engine/Source/Editor/UnrealEd/Private/Cooker/CookPackageArtifacts.cpp b/Engine/Source/Editor/UnrealEd/Private/Cooker/CookPackageArtifacts.cpp
--- a/Engine/Source/Editor/UnrealEd/Private/Cooker/CookPackageArtifacts.cpp
+++ b/Engine/Source/Editor/UnrealEd/Private/Cooker/CookPackageArtifacts.cpp
@@ -1035,8 +1035,8 @@ void FIncrementalCookAttachments::Fetch(TArrayView<FName> PackageNames, const IT
 	if (TargetPlatform || GEditorDomainOplog)
 	{
 		auto OnOplogAttachment =
-			[PackageWriter, Callback = MoveTemp(Callback), RequestedAttachmentNum = RequestedAttachments.Num(), &InProgressResults]
-			(FName PackageName, FUtf8StringView AttachmentKey, FCbObject&& Attachment)
+			[PackageWriter, Callback = MoveTemp(Callback), RequestedAttachmentNum = RequestedAttachments.Num(), InProgressResults = MoveTemp(InProgressResults)]
+			(FName PackageName, FUtf8StringView AttachmentKey, FCbObject&& Attachment) mutable
 			{
 				FInProgressResult& InProgressResult = InProgressResults.FindOrAdd(PackageName);
 				InProgressResult.ReceivedAttachmentCount++;
@@ -1072,12 +1072,12 @@ void FIncrementalCookAttachments::Fetch(TArrayView<FName> PackageNames, const IT
 		if (PackageWriter)
 		{
 			PackageWriter->GetOplogAttachments(PackageNames,
-				RequestedAttachments, OnOplogAttachment);
+				RequestedAttachments, MoveTemp(OnOplogAttachment));
 		}
 		else
 		{
 			GEditorDomainOplog->GetOplogAttachments(PackageNames,
-				RequestedAttachments, OnOplogAttachment);
+				RequestedAttachments, MoveTemp(OnOplogAttachment));
 		}
 	}
 	else
@@ -1087,7 +1087,6 @@ void FIncrementalCookAttachments::Fetch(TArrayView<FName> PackageNames, const IT
 			Callback(PackageName, FIncrementalCookAttachments());
 		}
 	}
-	check(InProgressResults.Num() == 0);
 }
 
 FPackageArtifacts FPackageArtifacts::Collect(const UPackage* Package, const ITargetPlatform* TargetPlatform,
diff --git a/Engine/Source/Editor/UnrealEd/Private/Cooker/DiffPackageWriter.h b/Engine/Source/Editor/UnrealEd/Private/Cooker/DiffPackageWriter.h
--- a/Engine/Source/Editor/UnrealEd/Private/Cooker/DiffPackageWriter.h
+++ b/Engine/Source/Editor/UnrealEd/Private/Cooker/DiffPackageWriter.h
@@ -92,9 +92,9 @@ public:
 	}
 	virtual void GetOplogAttachments(TArrayView<FName> PackageNames,
 		TArrayView<FUtf8StringView> AttachmentKeys,
-		TFunctionRef<void(FName PackageName, FUtf8StringView AttachmentKey, FCbObject&& Attachment)> Callback) override
+		TUniqueFunction<void(FName PackageName, FUtf8StringView AttachmentKey, FCbObject&& Attachment)>&& Callback) override
 	{
-		return Inner->GetOplogAttachments(PackageNames, AttachmentKeys, Callback);
+		return Inner->GetOplogAttachments(PackageNames, AttachmentKeys, MoveTemp(Callback));
 	}
 	virtual ECommitStatus GetCommitStatus(FName PackageName) override
 	{
@@ -255,9 +255,9 @@ public:
 	}
 	virtual void GetOplogAttachments(TArrayView<FName> PackageNames,
 		TArrayView<FUtf8StringView> AttachmentKeys,
-		TFunctionRef<void(FName PackageName, FUtf8StringView AttachmentKey, FCbObject&& Attachment)> Callback) override
+		TUniqueFunction<void(FName PackageName, FUtf8StringView AttachmentKey, FCbObject&& Attachment)>&& Callback) override
 	{
-		return Inner->GetOplogAttachments(PackageNames, AttachmentKeys, Callback);
+		return Inner->GetOplogAttachments(PackageNames, AttachmentKeys, MoveTemp(Callback));
 	}
 	virtual ECommitStatus GetCommitStatus(FName PackageName) override
 	{
diff --git a/Engine/Source/Editor/UnrealEd/Private/Cooker/LooseCookedPackageWriter.cpp b/Engine/Source/Editor/UnrealEd/Private/Cooker/LooseCookedPackageWriter.cpp
--- a/Engine/Source/Editor/UnrealEd/Private/Cooker/LooseCookedPackageWriter.cpp
+++ b/Engine/Source/Editor/UnrealEd/Private/Cooker/LooseCookedPackageWriter.cpp
@@ -446,7 +446,7 @@ FCbObject FLooseCookedPackageWriter::GetOplogAttachment(FName PackageName, FUtf8
 
 void FLooseCookedPackageWriter::GetOplogAttachments(TArrayView<FName> PackageNames,
 	TArrayView<FUtf8StringView> AttachmentKeys,
-	TFunctionRef<void(FName PackageName, FUtf8StringView AttachmentKey, FCbObject&& Attachment)> Callback)
+	TUniqueFunction<void(FName PackageName, FUtf8StringView AttachmentKey, FCbObject&& Attachment)>&& Callback)
 {
 	/** Oplog attachments are not implemented by FLooseCookedPackageWriter */
 	for (FName PackageName: PackageNames)
diff --git a/Engine/Source/Editor/UnrealEd/Private/Cooker/LooseCookedPackageWriter.h b/Engine/Source/Editor/UnrealEd/Private/Cooker/LooseCookedPackageWriter.h
--- a/Engine/Source/Editor/UnrealEd/Private/Cooker/LooseCookedPackageWriter.h
+++ b/Engine/Source/Editor/UnrealEd/Private/Cooker/LooseCookedPackageWriter.h
@@ -56,7 +56,7 @@ public:
 	virtual FCbObject GetOplogAttachment(FName PackageName, FUtf8StringView AttachmentKey) override;
 	virtual void GetOplogAttachments(TArrayView<FName> PackageNames,
 		TArrayView<FUtf8StringView> AttachmentKeys,
-		TFunctionRef<void(FName PackageName, FUtf8StringView AttachmentKey, FCbObject&& Attachment)> Callback) override;
+		TUniqueFunction<void(FName PackageName, FUtf8StringView AttachmentKey, FCbObject&& Attachment)>&& Callback) override;
 	virtual ECommitStatus GetCommitStatus(FName PackageName) override;
 	virtual void RemoveCookedPackages(TArrayView<const FName> PackageNamesToRemove) override;
 	virtual void RemoveCookedPackages() override;
diff --git a/Engine/Source/Editor/UnrealEd/Private/TargetDomain/TargetDomainUtils.cpp b/Engine/Source/Editor/UnrealEd/Private/TargetDomain/TargetDomainUtils.cpp
--- a/Engine/Source/Editor/UnrealEd/Private/TargetDomain/TargetDomainUtils.cpp
+++ b/Engine/Source/Editor/UnrealEd/Private/TargetDomain/TargetDomainUtils.cpp
@@ -325,7 +325,7 @@ FSharedBuffer IoBufferToSharedBuffer(FIoBuffer& InBuffer)
 
 void FEditorDomainOplog::GetOplogAttachments(TArrayView<FName> PackageNames,
 	TArrayView<FUtf8StringView> AttachmentKeys,
-	TFunctionRef<void(FName PackageName, FUtf8StringView AttachmentKey, FCbObject&& Attachment)> Callback)
+	TUniqueFunction<void(FName PackageName, FUtf8StringView AttachmentKey, FCbObject&& Attachment)>&& Callback)
 {
 	const int MaximumHashCount = PackageNames.Num() * AttachmentKeys.Num();
 	TArray<FIoHash> AttachmentHashes;
@@ -416,7 +416,7 @@ void FEditorDomainOplog::GetOplogAttachments(TArrayView<FName> PackageNames,
 		return;
 	}
 
-	HttpClient.ReadChunks(AttachmentHashes, [Callback,&AttachmentHashParams](const FIoHash& RawHash, TIoStatusOr<FIoBuffer> Result)
+	HttpClient.ReadChunks(AttachmentHashes, [Callback = MoveTemp(Callback),&AttachmentHashParams](const FIoHash& RawHash, TIoStatusOr<FIoBuffer> Result)
 	{
 		for (auto It(AttachmentHashParams.CreateConstKeyIterator(RawHash)); It; ++It)
 		{
diff --git a/Engine/Source/Editor/UnrealEd/Private/TargetDomain/TargetDomainUtils.h b/Engine/Source/Editor/UnrealEd/Private/TargetDomain/TargetDomainUtils.h
--- a/Engine/Source/Editor/UnrealEd/Private/TargetDomain/TargetDomainUtils.h
+++ b/Engine/Source/Editor/UnrealEd/Private/TargetDomain/TargetDomainUtils.h
@@ -42,7 +42,7 @@ public:
 	void CommitPackage(FName PackageName, TArrayView<IPackageWriter::FCommitAttachmentInfo> Attachments);
 	void GetOplogAttachments(TArrayView<FName> PackageName,
 		TArrayView<FUtf8StringView> AttachmentKeys,
-		TFunctionRef<void(FName PackageName, FUtf8StringView AttachmentKey, FCbObject&& Attachment)> Callback);
+		TUniqueFunction<void(FName PackageName, FUtf8StringView AttachmentKey, FCbObject&& Attachment)>&& Callback);
 
 private:
 	struct FOplogEntry
diff --git a/Engine/Source/Runtime/Core/Public/Serialization/PackageWriter.h b/Engine/Source/Runtime/Core/Public/Serialization/PackageWriter.h
--- a/Engine/Source/Runtime/Core/Public/Serialization/PackageWriter.h
+++ b/Engine/Source/Runtime/Core/Public/Serialization/PackageWriter.h
@@ -353,7 +353,7 @@ public:
 	 */
 	virtual void GetOplogAttachments(TArrayView<FName> PackageNames,
 		TArrayView<FUtf8StringView> AttachmentKeys,
-		TFunctionRef<void(FName PackageName, FUtf8StringView AttachmentKey, FCbObject&& Attachment)> Callback) = 0;
+		TUniqueFunction<void(FName PackageName, FUtf8StringView AttachmentKey, FCbObject&& Attachment)>&& Callback) = 0;
 
 	/**
 	 * Report commit status for the given package. If cooking incrementally, this reports the status
