Improvements to the UnrealPak CSV file format so that we can improve the grouping of metrics in Grafana and Horde and reduce the computation complexity when using the data in Tableau

Split out the pakchunk name and the target from container name into separate fields in the CSV output. This allows us to group metrics by PakChunk and Target platform names accordingly. It also means we can remove the string computation for this in Tableau

Added InstallType field ( Base, Optional, IAS, IAD ) so we can group metrics by the Install Type

diff --git a/Engine/Source/Developer/IoStoreUtilities/Private/IoStoreUtilities.cpp b/Engine/Source/Developer/IoStoreUtilities/Private/IoStoreUtilities.cpp
--- a/Engine/Source/Developer/IoStoreUtilities/Private/IoStoreUtilities.cpp
+++ b/Engine/Source/Developer/IoStoreUtilities/Private/IoStoreUtilities.cpp
@@ -1158,7 +1158,7 @@ public:
 		OutputArchive.Reset(IFileManager::Get().CreateFileWriter(OutputFilename));
 		if (OutputArchive.IsValid())
 		{
-			OutputArchive->Logf(TEXT("OrderInContainer, ChunkId, PackageId, PackageName, Filename, ContainerName, Offset, OffsetOnDisk, Size, CompressedSize, Hash, ChunkType, ClassType"));
+			OutputArchive->Logf(TEXT("OrderInContainer, ChunkId, PackageId, PackageName, Filename, ContainerName, Offset, OffsetOnDisk, Size, CompressedSize, Hash, ChunkType, ClassType, PakChunk, Target, InstallType"));
 		}
 	}
 
@@ -1166,7 +1166,28 @@ public:
 	{
 		if (OutputArchive.IsValid())
 		{
-			OutputArchive->Logf(TEXT("%d, %s, 0x%s, %s, %s, %s, %lld, %lld, %lld, %lld, 0x%s, %s, %s"),
+			FString PakChunk = TEXT("global");
+			FString Target = TEXT("Unknown");
+			FString InstallType = TEXT("Installed");
+
+			// Extract the PakChunk name and the Target from the container name
+			ContainerName.Split(TEXT("-"), &PakChunk, &Target);
+
+			// Determine the Install Type
+			if (ContainerName.Contains(TEXT("iad")))
+			{
+				InstallType = TEXT("IAD");
+			}
+			else if (ContainerName.Contains(TEXT("ondemand")))
+			{
+				InstallType = TEXT("IAS");
+			}
+			else if (ContainerName.Contains(TEXT("optional")))
+			{
+				InstallType = TEXT("Optional");
+			}
+			
+			OutputArchive->Logf(TEXT("%d, %s, 0x%s, %s, %s, %s, %lld, %lld, %lld, %lld, 0x%s, %s, %s, %s, %s, %s"),
 				Index,
 				*LexToString(Info.Id),
 				*LexToString(PackageId),
@@ -1179,7 +1200,10 @@ public:
 				Info.CompressedSize,
 				*LexToString(Info.ChunkHash),
 				*LexToString(Info.ChunkType),
-				*ClassType
+				*ClassType,
+				*PakChunk,
+				*Target,
+				*InstallType
 			);
 		}
 	}	
