Changed to log verbosity Display when diffing container files via UnrealPak.exe.

diff --git a/Engine/Source/Developer/IoStoreUtilities/Private/IoStoreUtilities.cpp b/Engine/Source/Developer/IoStoreUtilities/Private/IoStoreUtilities.cpp
--- a/Engine/Source/Developer/IoStoreUtilities/Private/IoStoreUtilities.cpp
+++ b/Engine/Source/Developer/IoStoreUtilities/Private/IoStoreUtilities.cpp
@@ -8432,7 +8432,7 @@ bool DiffIoStoreContainer(const TCHAR* CmdLine)
 bool LegacyDiffIoStoreContainers(const TCHAR* InContainerFilename1, const TCHAR* InContainerFilename2, bool bInLogUniques1, bool bInLogUniques2, const FKeyChain& InKeyChain1, const FKeyChain* InKeyChain2)
 {
 	TGuardValue DisableLogTimes(GPrintLogTimes, ELogTimes::None);
-	UE_LOG(LogIoStore, Log, TEXT("FileEventType, FileName, Size1, Size2"));
+	UE_LOG(LogIoStore, Display, TEXT("FileEventType, FileName, Size1, Size2"));
 
 	TUniquePtr<FIoStoreReader> Reader1 = CreateIoStoreReader(InContainerFilename1, InKeyChain1);
 	if (!Reader1.IsValid())
@@ -8484,12 +8484,12 @@ bool LegacyDiffIoStoreContainers(const TCHAR* InContainerFilename1, const TCHAR*
 			{
 				if (FindContainer1Entry->Size != ChunkInfo.Size)
 				{
-					UE_LOG(LogIoStore, Log, TEXT("FilesizeDifferent, %s, %llu, %llu"), *ChunkInfo.FileName, FindContainer1Entry->Size, ChunkInfo.Size);
+					UE_LOG(LogIoStore, Display, TEXT("FilesizeDifferent, %s, %llu, %llu"), *ChunkInfo.FileName, FindContainer1Entry->Size, ChunkInfo.Size);
 					++NumDifferentContents;
 				}
 				else if (FindContainer1Entry->ChunkHash != ChunkInfo.ChunkHash)
 				{
-					UE_LOG(LogIoStore, Log, TEXT("ContentsDifferent, %s, %llu, %llu"), *ChunkInfo.FileName, FindContainer1Entry->Size, ChunkInfo.Size);
+					UE_LOG(LogIoStore, Display, TEXT("ContentsDifferent, %s, %llu, %llu"), *ChunkInfo.FileName, FindContainer1Entry->Size, ChunkInfo.Size);
 					++NumDifferentContents;
 				}
 				else
@@ -8503,7 +8503,7 @@ bool LegacyDiffIoStoreContainers(const TCHAR* InContainerFilename1, const TCHAR*
 				++NumUniqueContainer2;
 				if (bInLogUniques2)
 				{
-					UE_LOG(LogIoStore, Log, TEXT("UniqueToSecondContainer, %s, 0, %llu"), *ChunkInfo.FileName, ChunkInfo.Size);
+					UE_LOG(LogIoStore, Display, TEXT("UniqueToSecondContainer, %s, 0, %llu"), *ChunkInfo.FileName, ChunkInfo.Size);
 				}
 			}
 			return true;
@@ -8515,12 +8515,12 @@ bool LegacyDiffIoStoreContainers(const TCHAR* InContainerFilename1, const TCHAR*
 		++NumUniqueContainer1;
 		if (bInLogUniques1)
 		{
-			UE_LOG(LogIoStore, Log, TEXT("UniqueToFirstContainer, %s, %llu, 0"), *Entry.FileName, Entry.Size);
+			UE_LOG(LogIoStore, Display, TEXT("UniqueToFirstContainer, %s, %llu, 0"), *Entry.FileName, Entry.Size);
 		}
 	}
 
-	UE_LOG(LogIoStore, Log, TEXT("Comparison complete"));
-	UE_LOG(LogIoStore, Log, TEXT("Unique to first container: %d, Unique to second container: %d, Num Different: %d, NumEqual: %d"), NumUniqueContainer1, NumUniqueContainer2, NumDifferentContents, NumEqualContents);
+	UE_LOG(LogIoStore, Display, TEXT("Comparison complete"));
+	UE_LOG(LogIoStore, Display, TEXT("Unique to first container: %d, Unique to second container: %d, Num Different: %d, NumEqual: %d"), NumUniqueContainer1, NumUniqueContainer2, NumDifferentContents, NumEqualContents);
 	return true;
 }
 
