Allow ZenStore project to be overriden with a cvar.

ContentWorker "prewarms" the cooker while it works out other parameters. One such parameter is the name of the plugin being cooked. This means when FCW is launched there is no meaningful information available to set a value for "-ZenStoreProject=". Later cook parameters are passed via a .txt file and manually processed which means it is not possible to influence FApp::GetZenStoreProjectId(). Adding a FApp::SetZenStoreProjectId() seemed too broad for a niche problem, and ZenStoreWriter's implicit logic around "-dlcname=" was or may be undesirable. Hence a CVar override was added.

#rnx

diff --git a/Engine/Source/Developer/IoStoreUtilities/Private/ZenStoreWriter.cpp b/Engine/Source/Developer/IoStoreUtilities/Private/ZenStoreWriter.cpp
--- a/Engine/Source/Developer/IoStoreUtilities/Private/ZenStoreWriter.cpp
+++ b/Engine/Source/Developer/IoStoreUtilities/Private/ZenStoreWriter.cpp
@@ -10,6 +10,7 @@
 #include "Containers/Queue.h"
 #include "GenericPlatform/GenericPlatformFile.h"
 #include "HAL/FileManager.h"
+#include "HAL/IConsoleManager.h"
 #include "HAL/PlatformFileManager.h"
 #include "Interfaces/ITargetPlatform.h"
 #include "IO/IoDispatcher.h"
@@ -40,6 +41,49 @@ DEFINE_LOG_CATEGORY_STATIC(LogZenStoreWriter, Log, All);
 
 using namespace UE;
 
+static FString GZenStoreWriterProjectIdOverride;
+static FAutoConsoleVariableRef CVar_ZenStoreWriterProjectIdOverride(
+	TEXT("ZenStoreWriter.ProjectIdOverride"),
+	GZenStoreWriterProjectIdOverride,
+	TEXT("Overrides the project ID otherwise taken from FApp/-DLCNAME. Format; '[<ParentId>/]ProjectId'")
+);
+
+static void GetZenStoreProjectId(FString& ProjectId, FString& ParentId)
+{
+	ParentId.Reset();
+
+	if (GZenStoreWriterProjectIdOverride.IsEmpty())
+	{
+		FString DLCName;
+		FParse::Value(FCommandLine::Get(), TEXT("DLCNAME="), DLCName);
+		DLCName.ToLowerInline();
+
+		ProjectId = FApp::GetZenStoreProjectId(DLCName);
+		if (!DLCName.IsEmpty())
+		{
+			ParentId = FApp::GetZenStoreProjectId();
+		}
+
+		return;
+	}
+
+	FStringView Override = GZenStoreWriterProjectIdOverride;
+
+	int32 SlashIndex;
+	if (!Override.FindChar('/', SlashIndex))
+	{
+		ProjectId = Override;
+		return;
+	}
+
+	ProjectId = Override.Left(SlashIndex);
+	check(!ProjectId.IsEmpty());
+
+	Override = Override.Mid(SlashIndex + 1);
+	check(!Override.IsEmpty());
+	ParentId = Override;
+}
+
 // Note that this is destructive - we yank out the buffer memory from the 
 // IoBuffer into the FSharedBuffer
 FSharedBuffer IoBufferToSharedBuffer(FIoBuffer& InBuffer)
@@ -179,16 +223,9 @@ FZenStoreWriter::FZenStoreWriter(
 {
 	StaticInit();
 
-	FString DLCName;
-	FParse::Value(FCommandLine::Get(), TEXT("DLCNAME="), DLCName);
-	DLCName.ToLowerInline();
-	ProjectId = FApp::GetZenStoreProjectId(DLCName);
 	FString ParentProjectId;
-	if (!DLCName.IsEmpty())
-	{
-		ParentProjectId = FApp::GetZenStoreProjectId();
-	}
-	
+	GetZenStoreProjectId(ProjectId, ParentProjectId);
+
 	if (FParse::Value(FCommandLine::Get(), TEXT("-ZenStorePlatform="), OplogId) == false)
 	{
 		OplogId = InTargetPlatform->PlatformName();
