IoStoreUtilities VCell deserialization

#rnx none

diff --git a/Engine/Source/Developer/IoStoreUtilities/Private/IoStoreUtilities.cpp b/Engine/Source/Developer/IoStoreUtilities/Private/IoStoreUtilities.cpp
--- a/Engine/Source/Developer/IoStoreUtilities/Private/IoStoreUtilities.cpp
+++ b/Engine/Source/Developer/IoStoreUtilities/Private/IoStoreUtilities.cpp
@@ -7322,6 +7322,18 @@ namespace DescribeUtils
 				HeaderDataReader << VersioningInfo;
 			}
 
+			FZenPackageCellOffsets CellOffsets;
+			if (!PackageSummary->bHasVersioningInfo || VersioningInfo.PackageVersion >= EUnrealEngineObjectUE5Version::VERSE_CELLS)
+			{
+				HeaderDataReader << CellOffsets.CellImportMapOffset;
+				HeaderDataReader << CellOffsets.CellExportMapOffset;
+			}
+			else
+			{
+				CellOffsets.CellImportMapOffset = PackageSummary->ExportBundleEntriesOffset;
+				CellOffsets.CellExportMapOffset = PackageSummary->ExportBundleEntriesOffset;
+			}
+
 			TArray<FDisplayNameEntryId> PackageNameMap;
 			{
 				TRACE_CPUPROFILER_EVENT_SCOPE(LoadNameBatch);
@@ -7343,7 +7355,7 @@ namespace DescribeUtils
 			}
 
 			const FExportMapEntry* ExportMap = reinterpret_cast<const FExportMapEntry*>(PackageSummaryData + PackageSummary->ExportMapOffset);
-			Job.PackageDesc->Exports.SetNum((PackageSummary->ExportBundleEntriesOffset - PackageSummary->ExportMapOffset) / sizeof(FExportMapEntry));
+			Job.PackageDesc->Exports.SetNum((CellOffsets.CellImportMapOffset - PackageSummary->ExportMapOffset) / sizeof(FExportMapEntry));
 			for (int32 ExportIndex = 0; ExportIndex < Job.PackageDesc->Exports.Num(); ++ExportIndex)
 			{
 				const FExportMapEntry& ExportMapEntry = ExportMap[ExportIndex];
