Added the DeliveryType to the IoStoreChunks csv export
Fixed the InstallType
Bumped the schema version for the data

diff --git a/Engine/Source/Developer/IoStoreUtilities/Private/IoStoreUtilities.cpp b/Engine/Source/Developer/IoStoreUtilities/Private/IoStoreUtilities.cpp
--- a/Engine/Source/Developer/IoStoreUtilities/Private/IoStoreUtilities.cpp
+++ b/Engine/Source/Developer/IoStoreUtilities/Private/IoStoreUtilities.cpp
@@ -1158,7 +1158,7 @@ public:
 		OutputArchive.Reset(IFileManager::Get().CreateFileWriter(OutputFilename));
 		if (OutputArchive.IsValid())
 		{
-			OutputArchive->Logf(TEXT("OrderInContainer, ChunkId, PackageId, PackageName, Filename, ContainerName, Offset, OffsetOnDisk, Size, CompressedSize, Hash, ChunkType, ClassType, PakChunk, Platform, InstallType"));
+			OutputArchive->Logf(TEXT("OrderInContainer, ChunkId, PackageId, PackageName, Filename, ContainerName, Offset, OffsetOnDisk, Size, CompressedSize, Hash, ChunkType, ClassType, PakChunk, Platform, InstallType, DeliveryType"));
 		}
 	}
 
@@ -1168,26 +1168,29 @@ public:
 		{
 			FString PakChunk = TEXT("global");
 			FString Platform = TEXT("Unknown");
-			FString InstallType = TEXT("Installed");
+			FString InstallType = TEXT("Base");
+			FString DeliveryType = TEXT("Installed");
 
 			// Extract the PakChunk name and the Platform from the container name
 			ContainerName.Split(TEXT("-"), &PakChunk, &Platform);
 
-			// Determine the Install Type
+			// Determine the Delivery Type Installed/IAS/IAD
 			if (ContainerName.Contains(TEXT("iad")))
 			{
-				InstallType = TEXT("IAD");
+				DeliveryType = TEXT("IAD");
 			}
 			else if (ContainerName.Contains(TEXT("ondemand")))
 			{
-				InstallType = TEXT("IAS");
+				DeliveryType = TEXT("IAS");
 			}
-			else if (ContainerName.Contains(TEXT("optional")))
+
+			// Determine the Install Type Base/Optional
+			if (ContainerName.Contains(TEXT("optional")))
 			{
 				InstallType = TEXT("Optional");
 			}
 			
-			OutputArchive->Logf(TEXT("%d, %s, 0x%s, %s, %s, %s, %lld, %lld, %lld, %lld, 0x%s, %s, %s, %s, %s, %s"),
+			OutputArchive->Logf(TEXT("%d, %s, 0x%s, %s, %s, %s, %lld, %lld, %lld, %lld, 0x%s, %s, %s, %s, %s, %s, %s"),
 				Index,
 				*LexToString(Info.Id),
 				*LexToString(PackageId),
@@ -1203,7 +1206,8 @@ public:
 				*ClassType,
 				*PakChunk,
 				*Platform,
-				*InstallType
+				*InstallType,
+				*DeliveryType
 			);
 		}
 	}	
