[unrealpak iostore] Add a direct way to find and extract packages from a directory full of containers.

diff --git a/Engine/Source/Developer/IoStoreUtilities/Private/IoStoreUtilities.cpp b/Engine/Source/Developer/IoStoreUtilities/Private/IoStoreUtilities.cpp
--- a/Engine/Source/Developer/IoStoreUtilities/Private/IoStoreUtilities.cpp
+++ b/Engine/Source/Developer/IoStoreUtilities/Private/IoStoreUtilities.cpp
@@ -6163,13 +6163,71 @@ void Split(FStringView Line, FStringView::ViewType Sep, TArray<FStringView>& Out
 	OutValues.Add(Value);
 }
 
+static TArray<FString> GetContainersFromPathOrWildcard(const FString& InContainerPathOrWildcard)
+{
+	TArray<FString> ContainerFilePaths;
+	if (IFileManager::Get().FileExists(*InContainerPathOrWildcard))
+	{
+		ContainerFilePaths.Add(InContainerPathOrWildcard);
+	}
+	else if (IFileManager::Get().DirectoryExists(*InContainerPathOrWildcard))
+	{
+		FString Directory = InContainerPathOrWildcard;
+		FPaths::NormalizeDirectoryName(Directory);
+
+		TArray<FString> FoundContainerFiles;
+		IFileManager::Get().FindFiles(FoundContainerFiles, *(Directory / TEXT("*.utoc")), true, false);
+
+		for (const FString& Filename : FoundContainerFiles)
+		{
+			ContainerFilePaths.Emplace(Directory / Filename);
+		}
+	}
+	else
+	{
+		FString Directory = FPaths::GetPath(InContainerPathOrWildcard);
+		FPaths::NormalizeDirectoryName(Directory);
+
+		TArray<FString> FoundContainerFiles;
+		IFileManager::Get().FindFiles(FoundContainerFiles, *InContainerPathOrWildcard, true, false);
+
+		for (const FString& Filename : FoundContainerFiles)
+		{
+			ContainerFilePaths.Emplace(Directory / Filename);
+		}
+	}
+	return MoveTemp(ContainerFilePaths);
+}
+
+static void FindAndExtractPackages(
+	const FKeyChain& KeyChain,
+	const FString& ContainerPathOrWildcard,
+	const FString& OutputPath,
+	const FString& PackageFilter)
+{
+	
+	TArray<FString> ContainerFilePaths = GetContainersFromPathOrWildcard(ContainerPathOrWildcard);
+	if (ContainerFilePaths.Num() == 0)
+	{
+		UE_LOG(LogIoStore, Error, TEXT("Container '%s' doesn't exist and no container matches wildcard."), *ContainerPathOrWildcard);
+		return;
+	}
+
+	for (const FString& ContainerFilePath : ContainerFilePaths)
+	{
+		ExtractFilesFromIoStoreContainer(
+			*ContainerFilePath, *OutputPath, KeyChain, &PackageFilter, nullptr, nullptr, nullptr);
+	}
+
+}
+
 int32 ListContainer(
 	const FKeyChain& KeyChain,
 	const FString& ContainerPathOrWildcard,
 	const FString& CsvPath)
 {
 	IOSTORE_CPU_SCOPE(ListContainer);
-	TArray<FString> ContainerFilePaths;
+	
 
 	FString AllChunksInfoFilename;
 	TMap<FPackageId, TPair<FName,FName>> AllChunksInfoMap;
@@ -6195,37 +6253,7 @@ int32 ListContainer(
 		FFileHelper::LoadFileToStringWithLineVisitor(*AllChunksInfoFilename, Visitor);
 	}
 
-	if (IFileManager::Get().FileExists(*ContainerPathOrWildcard))
-	{
-		ContainerFilePaths.Add(ContainerPathOrWildcard);
-	}
-	else if (IFileManager::Get().DirectoryExists(*ContainerPathOrWildcard))
-	{
-		FString Directory = ContainerPathOrWildcard;
-		FPaths::NormalizeDirectoryName(Directory);
-
-		TArray<FString> FoundContainerFiles;
-		IFileManager::Get().FindFiles(FoundContainerFiles, *(Directory / TEXT("*.utoc")), true, false);
-
-		for (const FString& Filename : FoundContainerFiles)
-		{
-			ContainerFilePaths.Emplace(Directory / Filename);
-		}
-	}
-	else
-	{
-		FString Directory = FPaths::GetPath(ContainerPathOrWildcard);
-		FPaths::NormalizeDirectoryName(Directory);
-
-		TArray<FString> FoundContainerFiles;
-		IFileManager::Get().FindFiles(FoundContainerFiles, *ContainerPathOrWildcard, true, false);
-
-		for (const FString& Filename : FoundContainerFiles)
-		{
-			ContainerFilePaths.Emplace(Directory / Filename);
-		}
-	}
-
+	TArray<FString> ContainerFilePaths = GetContainersFromPathOrWildcard(ContainerPathOrWildcard);
 	if (ContainerFilePaths.Num() == 0)
 	{
 		UE_LOG(LogIoStore, Error, TEXT("Container '%s' doesn't exist and no container matches wildcard."), *ContainerPathOrWildcard);
@@ -6367,38 +6395,7 @@ bool ListContainerBulkData(
 	};
 
 	TArray<FContainerData> Containers;
-	TArray<FString> ContainerFilePaths;
-
-	if (IFileManager::Get().FileExists(*ContainerPathOrWildcard))
-	{
-		ContainerFilePaths.Add(ContainerPathOrWildcard);
-	}
-	else if (IFileManager::Get().DirectoryExists(*ContainerPathOrWildcard))
-	{
-		FString Directory = ContainerPathOrWildcard;
-		FPaths::NormalizeDirectoryName(Directory);
-
-		TArray<FString> FoundContainerFiles;
-		IFileManager::Get().FindFiles(FoundContainerFiles, *(Directory / TEXT("*.utoc")), true, false);
-
-		for (const FString& Filename : FoundContainerFiles)
-		{
-			ContainerFilePaths.Emplace(Directory / Filename);
-		}
-	}
-	else
-	{
-		FString Directory = FPaths::GetPath(ContainerPathOrWildcard);
-		FPaths::NormalizeDirectoryName(Directory);
-
-		TArray<FString> FoundContainerFiles;
-		IFileManager::Get().FindFiles(FoundContainerFiles, *ContainerPathOrWildcard, true, false);
-
-		for (const FString& Filename : FoundContainerFiles)
-		{
-			ContainerFilePaths.Emplace(Directory / Filename);
-		}
-	}
+	TArray<FString> ContainerFilePaths = GetContainersFromPathOrWildcard(ContainerPathOrWildcard);
 
 	if (ContainerFilePaths.Num() == 0)
 	{
@@ -9619,6 +9616,31 @@ int32 CreateIoStoreContainerFiles(const TCHAR* CmdLine)
 
 		return ListContainer(Arguments.KeyChain, ContainerPathOrWildcard, Arguments.CsvPath);
 	}
+	else if (FParse::Value(FCommandLine::Get(), TEXT("FindAndExtract="), ArgumentValue))
+	{
+		// Often we know the package we want but not what container it's in, so this lets us
+		// just pass in the package and it'll load all the containers and pull it out wherever it is.
+
+		FString ContainerPathOrWildcard = MoveTemp(ArgumentValue);
+		
+		FString PackageFilter;
+		FParse::Value(FCommandLine::Get(), TEXT("PackageFilter="), PackageFilter);
+		if (PackageFilter.Len() == 0)
+		{
+			UE_LOG(LogIoStore, Error, TEXT("-PackageFilter=<wildcard> is required, use * for all packages"));
+			return -1;
+		}
+
+		FString OutPath;
+		FParse::Value(FCommandLine::Get(), TEXT("OutPath="), OutPath);
+		if (OutPath.Len() == 0)
+		{
+			UE_LOG(LogIoStore, Error, TEXT("-OutPath=/path/to/out is required"));
+			return -1;
+		}
+
+		FindAndExtractPackages(Arguments.KeyChain, ContainerPathOrWildcard, OutPath, PackageFilter);
+	}
 	else if (FParse::Value(FCommandLine::Get(), TEXT("Describe="), ArgumentValue))
 	{
 		FString ContainerPathOrWildcard = ArgumentValue;
@@ -9859,7 +9881,8 @@ int32 CreateIoStoreContainerFiles(const TCHAR* CmdLine)
 	{
 		UE_LOG(LogIoStore, Display, TEXT("Usage:"));
 		UE_LOG(LogIoStore, Display, TEXT(" -List=</path/to/[container.utoc|*.utoc]> -CSV=<list.csv> [-CryptoKeys=</path/to/crypto.json>]"));
-		UE_LOG(LogIoStore, Display, TEXT(" -Describe=</path/to/global.utoc> [-PackageFilter=<PackageName>] [-DumpToFile=<describe.txt>] [-CryptoKeys=</path/to/crypto.json>]"));
+		UE_LOG(LogIoStore, Display, TEXT(" -Describe=</path/to/global.utoc> [-PackageFilter=<PackageNameWildcard>] [-DumpToFile=<describe.txt>] [-CryptoKeys=</path/to/crypto.json>]"));
+		UE_LOG(LogIoStore, Display, TEXT(" -FindAndExtract=</path/to/[container.utoc|*.utoc]> -PackageFilter=<PackageNameWildcard> -OutPath=</path/to/output> [-CryptoKeys=</path/to/crypto.json>]"));
 		return -1;
 	}
 
