Allow IO Store to target a specific chunk name when for a maximum partition size
This is done in
[/Script/UnrealEd.ProjectPackagingSettings]
+MaxIoStorePartitionSizeOverrideMB=(pakchunk20-Windows, 256)
to split pakchunk20-Windows.ucas in multiple partitions up to 256 MiB

diff --git a/Engine/Source/Developer/IoStoreUtilities/Private/IoStoreUtilities.cpp b/Engine/Source/Developer/IoStoreUtilities/Private/IoStoreUtilities.cpp
--- a/Engine/Source/Developer/IoStoreUtilities/Private/IoStoreUtilities.cpp
+++ b/Engine/Source/Developer/IoStoreUtilities/Private/IoStoreUtilities.cpp
@@ -9337,6 +9337,47 @@ static bool ParseOrderFileArguments(FIoStoreArguments& Arguments)
 	return true;
 }
 
+static TMap<FString, uint64> GetMaxIoStorePartitionSizeOverride(TConstArrayView<FString> MaxIoStorePartitionSizeOverrideMB)
+{
+	TMap<FString, uint64> MaxPartitionSizeOverrideMap;
+	for (const FString& ConfigString : MaxIoStorePartitionSizeOverrideMB)
+	{
+		FStringView ConfigStringView(ConfigString);
+		ConfigStringView = ConfigStringView.TrimStartAndEnd();
+		if (ConfigStringView.IsEmpty() || ConfigStringView[0] != TEXT('(') || ConfigStringView[ConfigStringView.Len() - 1] != TEXT(')'))
+		{
+			continue;
+		}
+		ConfigStringView.RemovePrefix(1);
+		ConfigStringView.RemoveSuffix(1);
+
+		TArray<FStringView> ConfigStringArrayNameSize;
+		UE::String::ParseTokens(ConfigStringView, TEXTVIEW(","), [&](FStringView Element)
+		{
+			ConfigStringArrayNameSize.Add(Element);
+		});
+
+		if (ConfigStringArrayNameSize.Num() != 2 || ConfigStringArrayNameSize[0].Len() == 0 || ConfigStringArrayNameSize[1].Len() == 0)
+		{
+			continue;
+		}
+
+		const TCHAR* ConfigSizeStringViewEnd = ConfigStringArrayNameSize[1].GetData();
+		uint64 ConfigSize = FCString::Strtoui64(ConfigStringArrayNameSize[1].GetData(), (TCHAR**)&ConfigSizeStringViewEnd, 10);
+
+		if (ConfigSizeStringViewEnd == ConfigStringArrayNameSize[1].GetData())
+		{
+			continue;
+		}
+
+		ConfigSize = ConfigSize << 20; // MiB to bytes
+
+		MaxPartitionSizeOverrideMap.Add(FString(ConfigStringArrayNameSize[0]), ConfigSize);
+	}
+
+	return MaxPartitionSizeOverrideMap;
+}
+
 bool ParseContainerGenerationArguments(FIoStoreArguments& Arguments, FIoStoreWriterSettings& WriterSettings)
 {
 	IOSTORE_CPU_SCOPE(ParseContainerGenerationArguments);
@@ -9450,6 +9491,12 @@ bool ParseContainerGenerationArguments(FIoStoreArguments& Arguments, FIoStoreWri
 		WriterSettings.MaxPartitionSize = MaxPartitionSize;
 	}
 
+	TArray<FString> MaxIoStorePartitionSizeOverrideMB;
+	if (GConfig->GetArray(TEXT("/Script/UnrealEd.ProjectPackagingSettings"), TEXT("MaxIoStorePartitionSizeOverrideMB"), MaxIoStorePartitionSizeOverrideMB, GGameIni))
+	{
+		WriterSettings.MaxPartitionSizeOverride = GetMaxIoStorePartitionSizeOverride(MaxIoStorePartitionSizeOverrideMB);
+	}
+
 	int32 CompressionMinBytesSaved = 0;
 	if (FParse::Value(FCommandLine::Get(), TEXT("-compressionMinBytesSaved="), CompressionMinBytesSaved))
 	{
@@ -9482,6 +9529,11 @@ bool ParseContainerGenerationArguments(FIoStoreArguments& Arguments, FIoStoreWri
 	UE_LOG(LogIoStore, Display, TEXT("Using compression min bytes saved '%d'"), WriterSettings.CompressionMinBytesSaved);
 	UE_LOG(LogIoStore, Display, TEXT("Using compression min percent saved '%d'"), WriterSettings.CompressionMinPercentSaved);
 	UE_LOG(LogIoStore, Display, TEXT("Using max partition size '%lld'"), WriterSettings.MaxPartitionSize);
+	for (const TPair<FString, uint64>& MaxPartitionSizeOverrideIter : WriterSettings.MaxPartitionSizeOverride)
+	{
+		UE_LOG(LogIoStore, Display, TEXT("Using max partition size override %s -> '%lld'"), *MaxPartitionSizeOverrideIter.Key, MaxPartitionSizeOverrideIter.Value);
+	}
+
 	if (WriterSettings.bCompressionEnableDDC)
 	{
 		UE_LOG(LogIoStore, Display, TEXT("Using DDC for compression with min size '%d'"), WriterSettings.CompressionMinSizeToConsiderDDC);
diff --git a/Engine/Source/Developer/IoStoreUtilities/Private/IoStoreWriter.cpp b/Engine/Source/Developer/IoStoreUtilities/Private/IoStoreWriter.cpp
--- a/Engine/Source/Developer/IoStoreUtilities/Private/IoStoreWriter.cpp
+++ b/Engine/Source/Developer/IoStoreUtilities/Private/IoStoreWriter.cpp
@@ -781,6 +781,16 @@ public:
 
 		TocFilePath = ContainerPathAndBaseFileName + TEXT(".utoc");
 		
+		FString PakChunkName = FPaths::GetBaseFilename(ContainerPathAndBaseFileName);
+
+		const uint64* MaxPartitionSizeOverrideValue = WriterContext->WriterSettings.MaxPartitionSizeOverride.Find(PakChunkName);
+		MaxPartitionSize =  MaxPartitionSizeOverrideValue ? *MaxPartitionSizeOverrideValue : WriterContext->WriterSettings.MaxPartitionSize;
+
+		if (MaxPartitionSize > 0)
+		{
+			UE_LOG(LogIoStore, Display, TEXT("FIoStoreWriter: using a max partition size of %llu bytes for %s"), MaxPartitionSize, *PakChunkName);
+		}
+
 		IPlatformFile& Ipf = IPlatformFile::GetPlatformPhysical();
 		Ipf.CreateDirectoryTree(*FPaths::GetPath(TocFilePath));
 
@@ -1284,7 +1294,7 @@ public:
 				ContainerSettings.IsEncrypted() ? ContainerSettings.EncryptionKey : FAES::FAESKey());
 		}
 
-		TIoStatusOr<uint64> TocSize = FIoStoreTocResource::Write(*TocFilePath, TocResource, static_cast<uint32>(WriterSettings.CompressionBlockSize), WriterSettings.MaxPartitionSize, ContainerSettings);
+		TIoStatusOr<uint64> TocSize = FIoStoreTocResource::Write(*TocFilePath, TocResource, static_cast<uint32>(WriterSettings.CompressionBlockSize), MaxPartitionSize, ContainerSettings);
 		checkf(TocSize.IsOk(), TEXT("FIoStoreTocResource::Write failed with error %s"), *TocSize.Status().ToString());
 		
 		Result.ContainerId = ContainerSettings.ContainerId;
@@ -1905,7 +1915,7 @@ private:
 		const FIoStoreWriterSettings& WriterSettings = WriterContext->WriterSettings;
 		bHasMemoryMappedEntry |= Entry->Options.bIsMemoryMapped;
 		const uint64 ChunkAlignment = Entry->Options.bIsMemoryMapped ? WriterSettings.MemoryMappingAlignment : 0;
-		const uint64 PartitionSizeLimit = WriterSettings.MaxPartitionSize > 0 ? WriterSettings.MaxPartitionSize : MAX_uint64;
+		const uint64 PartitionSizeLimit = MaxPartitionSize > 0 ? MaxPartitionSize : MAX_uint64;
 		checkf(Entry->DiskSize <= PartitionSizeLimit, TEXT("Chunk is too large, increase max partition size!"));
 		for (;;)
 		{
@@ -1968,7 +1978,7 @@ private:
 		for (const FChunkBlock& ChunkBlock : Entry->ChunkBlocks)
 		{
 			FIoStoreTocCompressedBlockEntry& BlockEntry = TocBuilder.AddCompressionBlockEntry();
-			BlockEntry.SetOffset(TargetPartition->Index * WriterSettings.MaxPartitionSize + TargetPartition->Offset + OffsetInChunk);
+			BlockEntry.SetOffset(TargetPartition->Index * MaxPartitionSize + TargetPartition->Offset + OffsetInChunk);
 			OffsetInChunk += ChunkBlock.DiskSize;
 			BlockEntry.SetCompressedSize(uint32(ChunkBlock.CompressedSize));
 			BlockEntry.SetUncompressedSize(uint32(ChunkBlock.UncompressedSize));
@@ -2059,6 +2069,7 @@ private:
 	uint64						TotalPaddingSize = 0;
 	uint64						UncompressedContainerSize = 0; // this is the size the container would be if it were uncompressed.
 	uint64						CompressedContainerSize = 0; // this is the size of the container with the given compression (which may be none).
+	uint64 						MaxPartitionSize = 0;
 	int32						CurrentPartitionIndex = 0;
 	bool						bHasMemoryMappedEntry = false;
 	bool						bHasFlushed = false;
diff --git a/Engine/Source/Developer/IoStoreUtilities/Private/IoStoreWriter.h b/Engine/Source/Developer/IoStoreUtilities/Private/IoStoreWriter.h
--- a/Engine/Source/Developer/IoStoreUtilities/Private/IoStoreWriter.h
+++ b/Engine/Source/Developer/IoStoreUtilities/Private/IoStoreWriter.h
@@ -22,6 +22,7 @@ struct FIoStoreWriterSettings
 	int32 CompressionMinSizeToConsiderDDC = 0;
 	uint64 MemoryMappingAlignment = 0;
 	uint64 MaxPartitionSize = 0;
+	TMap<FString, uint64> MaxPartitionSizeOverride;
 	bool bEnableFileRegions = false;
 	bool bCompressionEnableDDC = false;
 	bool bValidateChunkHashes = false;
