Project was including the "host name" for URLs instead of the name used to connect to the server (which may be a Unix domain socket).

#rnx

diff --git a/Engine/Source/Developer/IoStoreUtilities/Private/ZenStoreWriter.cpp b/Engine/Source/Developer/IoStoreUtilities/Private/ZenStoreWriter.cpp
--- a/Engine/Source/Developer/IoStoreUtilities/Private/ZenStoreWriter.cpp
+++ b/Engine/Source/Developer/IoStoreUtilities/Private/ZenStoreWriter.cpp
@@ -496,7 +496,7 @@ void FZenStoreWriter::Initialize(const FCookInfo& Info)
 #if UE_WITH_ZEN
 		const Zen::FZenServiceInstance& ZenServiceInstance = HttpClient->GetZenServiceInstance();
 		IsRunningLocally = ZenServiceInstance.IsServiceRunningLocally();
-		ServiceHostName = ZenServiceInstance.GetEndpoint().GetHostName();
+		ServiceHostName = ZenServiceInstance.GetEndpoint().GetName();
 		ServicePort = ZenServiceInstance.GetEndpoint().GetPort();
 #endif
 
