CookCorrectness: ZenStoreWriter: Initialize FOplogCookInfo.PackageName in committed packages.
#rnx

diff --git a/Engine/Source/Developer/IoStoreUtilities/Private/ZenStoreWriter.cpp b/Engine/Source/Developer/IoStoreUtilities/Private/ZenStoreWriter.cpp
--- a/Engine/Source/Developer/IoStoreUtilities/Private/ZenStoreWriter.cpp
+++ b/Engine/Source/Developer/IoStoreUtilities/Private/ZenStoreWriter.cpp
@@ -1010,7 +1010,7 @@ void FZenStoreWriter::CommitPackageInternal(FZenCommitInfo&& ZenCommitInfo)
 			if (CommitEventArgs.EntryIndex == PackageStoreEntries.Num())
 			{
 				PackageStoreEntries.Emplace();
-				CookedPackagesInfo.Emplace();
+				CookedPackagesInfo.Emplace(CommitInfo.PackageName);
 			}
 			PackageStoreEntries[CommitEventArgs.EntryIndex] = MoveTemp(PackageStoreEntry);
 
