Fix MaxIoStorePartitionSizeOverrideMB per-platform handling
* add '-platform' argument to UnrealPak so that we may specify the config we want to load
* The ini setup is still the same:
[/Script/UnrealEd.ProjectPackagingSettings]
+MaxIoStorePartitionSizeOverrideMB=(pakchunk20-Windows, 256)
to split pakchunk20-Windows.ucas in multiple partitions up to 256 MiB
* FConfigCacheIni::ForPlatform requires ALLOW_OTHER_PLATFORM_CONFIG: allow projects to define it without requiring WITH_UNREAL_DEVELOPER_TOOLS

diff --git a/Engine/Source/Developer/IoStoreUtilities/Private/IoStoreUtilities.cpp b/Engine/Source/Developer/IoStoreUtilities/Private/IoStoreUtilities.cpp
--- a/Engine/Source/Developer/IoStoreUtilities/Private/IoStoreUtilities.cpp
+++ b/Engine/Source/Developer/IoStoreUtilities/Private/IoStoreUtilities.cpp
@@ -9491,8 +9491,19 @@ bool ParseContainerGenerationArguments(FIoStoreArguments& Arguments, FIoStoreWri
 		WriterSettings.MaxPartitionSize = MaxPartitionSize;
 	}
 
+	FConfigCacheIni* Config = GConfig;
+
+	FString IniPlatformName;
+	if (FParse::Value(FCommandLine::Get(), TEXT("platform="), IniPlatformName, false))
+	{
+		if (FConfigCacheIni* PlatformConfig = FConfigCacheIni::ForPlatform(*IniPlatformName))
+		{
+			Config = PlatformConfig;
+		}
+	}
+
 	TArray<FString> MaxIoStorePartitionSizeOverrideMB;
-	if (GConfig->GetArray(TEXT("/Script/UnrealEd.ProjectPackagingSettings"), TEXT("MaxIoStorePartitionSizeOverrideMB"), MaxIoStorePartitionSizeOverrideMB, GGameIni))
+	if (Config->GetArray(TEXT("/Script/UnrealEd.ProjectPackagingSettings"), TEXT("MaxIoStorePartitionSizeOverrideMB"), MaxIoStorePartitionSizeOverrideMB, GGameIni))
 	{
 		WriterSettings.MaxPartitionSizeOverride = GetMaxIoStorePartitionSizeOverride(MaxIoStorePartitionSizeOverrideMB);
 	}
diff --git a/Engine/Source/Programs/AutomationTool/Scripts/CopyBuildToStagingDirectory.Automation.cs b/Engine/Source/Programs/AutomationTool/Scripts/CopyBuildToStagingDirectory.Automation.cs
--- a/Engine/Source/Programs/AutomationTool/Scripts/CopyBuildToStagingDirectory.Automation.cs
+++ b/Engine/Source/Programs/AutomationTool/Scripts/CopyBuildToStagingDirectory.Automation.cs
@@ -4393,7 +4393,8 @@ namespace AutomationScripts
 				{	
 					AdditionalArgs += " -AlphaSortClusterPackageLists";
 				}
-				
+
+				AdditionalArgs += $" -platform={ConfigHierarchy.GetIniPlatformName(SC.StageTargetPlatform.IniPlatformType)}";
 
 				AdditionalArgs += " " + Params.AdditionalIoStoreOptions;
 
diff --git a/Engine/Source/Programs/UnrealPak/UnrealPak.Target.cs b/Engine/Source/Programs/UnrealPak/UnrealPak.Target.cs
--- a/Engine/Source/Programs/UnrealPak/UnrealPak.Target.cs
+++ b/Engine/Source/Programs/UnrealPak/UnrealPak.Target.cs
@@ -43,5 +43,6 @@ public class UnrealPakTarget : TargetRules
 
 		bEnableTrace = true;
 		GlobalDefinitions.Add("UE_EXTERNAL_PROFILING_ENABLED=0");
+		GlobalDefinitions.Add("ALLOW_OTHER_PLATFORM_CONFIG=1"); 
 	}
 }
diff --git a/Engine/Source/Runtime/Core/Public/Misc/Build.h b/Engine/Source/Runtime/Core/Public/Misc/Build.h
--- a/Engine/Source/Runtime/Core/Public/Misc/Build.h
+++ b/Engine/Source/Runtime/Core/Public/Misc/Build.h
@@ -508,7 +508,9 @@
 #endif
 
 // Controls if the config system can stores configs for other platforms than the running one
+#ifndef ALLOW_OTHER_PLATFORM_CONFIG
 #define ALLOW_OTHER_PLATFORM_CONFIG		WITH_UNREAL_DEVELOPER_TOOLS
+#endif
 
 // Controls whether or not process will control OS scheduler priority
 #ifndef WITH_PROCESS_PRIORITY_CONTROL
