Added StudioTelemetry to UnrealPak
UnrealPak will send telemetry ONLY when the -Telemetry is supplied on the commandline, otherwise the overhead is assumed to be near zero.
Added ChunkEventSummary to IoStoreUtilities to aggregate the Compressed and Uncompressed size totals.
Added Core.IoStoreChunk.Summary Telemetry event to UnrealPak that sends summary of the aggregated Chunk information
Added Core.IoStoreHeader.Summary Telemetry event to UnrealPak that sends summary of the header information

diff --git a/Engine/Source/Developer/IoStoreUtilities/IoStoreUtilities.Build.cs b/Engine/Source/Developer/IoStoreUtilities/IoStoreUtilities.Build.cs
--- a/Engine/Source/Developer/IoStoreUtilities/IoStoreUtilities.Build.cs
+++ b/Engine/Source/Developer/IoStoreUtilities/IoStoreUtilities.Build.cs
@@ -26,6 +26,7 @@ public class IoStoreUtilities : ModuleRules
 			"RSA",
 			"SandboxFile",
 			"Sockets",
+			"StudioTelemetry",
 			"Zen",
 		});
 
diff --git a/Engine/Source/Developer/IoStoreUtilities/Private/IoStoreUtilities.cpp b/Engine/Source/Developer/IoStoreUtilities/Private/IoStoreUtilities.cpp
--- a/Engine/Source/Developer/IoStoreUtilities/Private/IoStoreUtilities.cpp
+++ b/Engine/Source/Developer/IoStoreUtilities/Private/IoStoreUtilities.cpp
@@ -87,6 +87,7 @@
 #include "Serialization/CompactBinarySerialization.h"
 #include "Serialization/ZenPackageHeader.h"
 #include "String/ParseTokens.h"
+#include "StudioTelemetry.h"
 #include "Tasks/Task.h"
 #include "Templates/Greater.h"
 
@@ -1142,6 +1143,61 @@ struct FContainerTargetSpec
 using FPackageNameMap = TMap<FName, FCookedPackage*>;
 using FPackageIdMap = TMap<FPackageId, FCookedPackage*>;
 
+enum EAggregate
+{
+	Installed,
+	IAS,
+	IAD,
+	NumTypes
+};
+
+const FString AggregateNames[EAggregate::NumTypes] =
+{
+	TEXT("Installed"),
+	TEXT("IAS"),
+	TEXT("IAD")
+};
+
+class FChunkEntrySummary
+{
+public:
+	uint64 PackageCount;
+	uint64 UncompressedTotals[EAggregate::NumTypes];
+	uint64 CompressedTotals[EAggregate::NumTypes];
+
+	FChunkEntrySummary()
+	{
+		PackageCount = 0;
+
+		for (uint32 i = 0; i < EAggregate::NumTypes; i++)
+		{
+			UncompressedTotals[i] = CompressedTotals[i] = 0;
+		}
+	}
+
+	void SendTelemetryEvent()
+	{
+		if (FStudioTelemetry::Get().IsSessionRunning())
+		{
+			// Write out the summary to telemetry
+			TArray<FAnalyticsEventAttribute> Attributes;
+
+			const int SchemaVersion = 1;
+
+			Attributes.Emplace(TEXT("SchemaVersion"), SchemaVersion);
+			Attributes.Emplace(TEXT("PackageCount"), PackageCount);
+
+			for (uint32 i = 0; i < EAggregate::NumTypes; i++)
+			{
+				Attributes.Emplace(*FString::Printf(TEXT("UncompressedSize%s"), *AggregateNames[i]), UncompressedTotals[i]);
+				Attributes.Emplace(*FString::Printf(TEXT("CompressedSize%s"), *AggregateNames[i]), CompressedTotals[i]);
+			}
+
+			FStudioTelemetry::Get().RecordEvent(TEXT("Core.IoStoreChunk.Summary"), Attributes);
+		}
+	}
+};
+
 class FChunkEntryCsv
 {
 public:
@@ -1151,6 +1207,7 @@ public:
 		if (OutputArchive.IsValid())
 		{
 			OutputArchive->Flush();
+			ChunkEntrySummary.SendTelemetryEvent();
 		}
 	}
 
@@ -1179,11 +1236,22 @@ public:
 			if (ContainerName.Contains(TEXT("iad")))
 			{
 				DeliveryType = TEXT("IAD");
+				ChunkEntrySummary.UncompressedTotals[EAggregate::IAD] += Info.Size;
+				ChunkEntrySummary.CompressedTotals[EAggregate::IAD] += Info.CompressedSize;
 			}
 			else if (ContainerName.Contains(TEXT("ondemand")))
 			{
 				DeliveryType = TEXT("IAS");
+				ChunkEntrySummary.UncompressedTotals[EAggregate::IAD] += Info.Size;
+				ChunkEntrySummary.CompressedTotals[EAggregate::IAD] += Info.CompressedSize;
 			}
+			else
+			{
+				ChunkEntrySummary.UncompressedTotals[EAggregate::Installed] += Info.Size;
+				ChunkEntrySummary.CompressedTotals[EAggregate::Installed] += Info.CompressedSize;
+			}
+
+			ChunkEntrySummary.PackageCount++;
 
 			// Determine the Install Type Base/Optional
 			if (ContainerName.Contains(TEXT("optional")))
@@ -1215,6 +1283,7 @@ public:
 
 private:
 	TUniquePtr<FArchive> OutputArchive;
+	FChunkEntrySummary ChunkEntrySummary;
 };
 
 void SortPackagesInLoadOrderRecursive(TArray<FCookedPackage*>& Result, FCookedPackage* Package, TArray<FCookedPackage*>& S, TArray<FCookedPackage*>& P, int32& C, const TMap<FCookedPackage*, TArray<FCookedPackage*>>& ReverseEdges)
@@ -5865,7 +5934,35 @@ int32 CreateTarget(const FIoStoreArguments& Arguments, const FIoStoreWriterSetti
 	}
 	UE_LOG(LogIoStore, Display, TEXT("Output: %8.2lf MB InitialLoadData"), (double)InitialLoadSize / 1024.0 / 1024.0);
 
+	if (FStudioTelemetry::Get().IsSessionRunning() == true)
+	{
+		// Write out the summary to telemetry
+		TArray<FAnalyticsEventAttribute> Attributes;
+
+		const int SchemaVersion = 1;
 
+		Attributes.Emplace(TEXT("SchemaVersion"), SchemaVersion);
+		Attributes.Emplace(TEXT("PackageCount"), Packages.Num());
+		Attributes.Emplace(TEXT("UExpSize"), UExpSize);
+		Attributes.Emplace(TEXT("UUAssetSize"), UAssetSize);
+		Attributes.Emplace(TEXT("UBulkSize"), UBulkSize);
+		Attributes.Emplace(TEXT("InitialLoadSize"), InitialLoadSize);
+		Attributes.Emplace(TEXT("GlobalShaderSize"), GlobalShaderSize);
+		Attributes.Emplace(TEXT("SharedShaderSize"), SharedShaderSize);
+		Attributes.Emplace(TEXT("UniqueShaderSize"), UniqueShaderSize);
+		Attributes.Emplace(TEXT("ImportedPackagesCount"), ImportedPackagesCount);
+		Attributes.Emplace(TEXT("NoImportedPackagesCount"), NoImportedPackagesCount);
+		Attributes.Emplace(TEXT("ScriptCount"), PackageStoreOptimizer.GetTotalScriptObjectCount());
+
+		if (bIsLegacyStage)
+		{
+			Attributes.Emplace(TEXT("HeaderSize"), HeaderSize);
+			Attributes.Emplace(TEXT("NameMapCount"), NameMapCount);
+		}
+
+		FStudioTelemetry::Get().RecordEvent(TEXT("Core.IoStoreHeader.Summary"), Attributes);
+	}
+	
 	if (ChunkDatabase.IsValid())
 	{
 		TRACE_CPUPROFILER_EVENT_SCOPE(ReferenceChunkLogging);
@@ -9631,7 +9728,7 @@ bool ParseContainerGenerationArguments(FIoStoreArguments& Arguments, FIoStoreWri
 int32 CreateIoStoreContainerFiles(const TCHAR* CmdLine)
 {
 	IOSTORE_CPU_SCOPE(CreateIoStoreContainerFiles);
-
+	
 	UE_LOG(LogIoStore, Display, TEXT("==================== IoStore Utils ===================="));
 
 	FIoStoreArguments Arguments;
@@ -9641,7 +9738,7 @@ int32 CreateIoStoreContainerFiles(const TCHAR* CmdLine)
 	{
 		return 1;
 	}
-	
+
 	ITargetPlatform* TargetPlatform = nullptr;
 	FString TargetPlatformName;
 	if (FParse::Value(FCommandLine::Get(), TEXT("TargetPlatform="), TargetPlatformName))
diff --git a/Engine/Source/Programs/UnrealPak/Private/UnrealPak.cpp b/Engine/Source/Programs/UnrealPak/Private/UnrealPak.cpp
--- a/Engine/Source/Programs/UnrealPak/Private/UnrealPak.cpp
+++ b/Engine/Source/Programs/UnrealPak/Private/UnrealPak.cpp
@@ -5,6 +5,7 @@
 #include "PakFileUtilities.h"
 #include "IPlatformFilePak.h"
 #include "ProjectUtilities.h"
+#include "StudioTelemetry.h"
 
 IMPLEMENT_APPLICATION(UnrealPak, "UnrealPak");
 
@@ -21,8 +22,16 @@ INT32_MAIN_INT32_ARGC_TCHAR_ARGV()
 
 	double StartTime = FPlatformTime::Seconds();
 
+	if (FParse::Param(FCommandLine::Get(), TEXT("Telemetry")))
+	{
+		// Only start a telemetry session if specified
+		FStudioTelemetry::Get().StartSession();
+	}
+	
 	int32 Result = ExecuteUnrealPak(FCommandLine::Get())? 0 : 1;
 
+	FStudioTelemetry::Get().EndSession();
+
 	UE_LOG(LogPakFile, Display, TEXT("UnrealPak executed in %f seconds"), FPlatformTime::Seconds() - StartTime);
 
 	if (FParse::Param(FCommandLine::Get(), TEXT("fastexit")))
diff --git a/Engine/Source/Programs/UnrealPak/UnrealPak.Build.cs b/Engine/Source/Programs/UnrealPak/UnrealPak.Build.cs
--- a/Engine/Source/Programs/UnrealPak/UnrealPak.Build.cs
+++ b/Engine/Source/Programs/UnrealPak/UnrealPak.Build.cs
@@ -19,7 +19,8 @@ public class UnrealPak : ModuleRules
 			"PakFile",
 			"PakFileUtilities",
 			"Projects", 
-			"RSA", 
+			"RSA",
+			"StudioTelemetry",
 			"TargetPlatform",
 		});
 
