Improve the logging of memory values at the end of UnrealPak so that the unit scales with the number of bytes.

#rnx

- Fixes smaller data sets being rounded to 0.0 GB

diff --git a/Engine/Source/Developer/IoStoreUtilities/Private/IoStoreUtilities.cpp b/Engine/Source/Developer/IoStoreUtilities/Private/IoStoreUtilities.cpp
--- a/Engine/Source/Developer/IoStoreUtilities/Private/IoStoreUtilities.cpp
+++ b/Engine/Source/Developer/IoStoreUtilities/Private/IoStoreUtilities.cpp
@@ -115,6 +115,27 @@ static FPackageId PackageIdFromChunkId(const FIoChunkId& ChunkId)
 	return FPackageId::FromValue(*(int64*)(ChunkId.GetData()));
 }
 
+/** Utility to format the number of bytes to a string. The units used will scale with the number of bytes. */
+static FString FormatMemoryString(uint64 Bytes)
+{
+	if (Bytes < 1024)
+	{
+		return FString::Printf(TEXT("%8u   B"), (uint32)Bytes);
+	}
+	else if (Bytes < 1024 * 1024)
+	{
+		return FString::Printf(TEXT("%8.2f KiB"), (double)Bytes / 1024.0);
+	}
+	else if (Bytes < 1024 * 1024 * 1024)
+	{
+		return FString::Printf(TEXT("%8.2f MiB"), (double)Bytes / 1024.0 / 1024.0);
+	}
+	else
+	{
+		return FString::Printf(TEXT("%8.2f GiB"), (double)Bytes / 1024.0 / 1024.0 / 1024.0);
+	}
+}
+
 static const FName DefaultCompressionMethod = NAME_Zlib;
 static const uint64 DefaultCompressionBlockSize = 64 << 10;
 static const uint64 DefaultCompressionBlockAlignment = 64 << 10;
@@ -5911,28 +5932,32 @@ int32 CreateTarget(const FIoStoreArguments& Arguments, const FIoStoreWriterSetti
 
 	LogWriterResults(IoStoreWriterResults);
 	LogContainerPackageInfo(ContainerTargets);
-	
+
 	UE_LOG(LogIoStore, Display, TEXT("Input:  %8d Packages"), Packages.Num());
-	UE_LOG(LogIoStore, Display, TEXT("Input:  %8.2lf GB UExp"), (double)UExpSize / 1024.0 / 1024.0 / 1024.0);
-	UE_LOG(LogIoStore, Display, TEXT("Input:  %8.2lf GB UAsset"), (double)UAssetSize / 1024.0 / 1024.0 / 1024.0);
-	UE_LOG(LogIoStore, Display, TEXT("Input:  %8.2lf GB UBulk"), (double)UBulkSize / 1024.0 / 1024.0 / 1024.0);
-	UE_LOG(LogIoStore, Display, TEXT("Input:  %8.2f MB for %d Global shaders"), (double)GlobalShaderSize / 1024.0 / 1024.0, GlobalShaderCount);
-	UE_LOG(LogIoStore, Display, TEXT("Input:  %8.2f MB for %d Shared shaders"), (double)SharedShaderSize / 1024.0 / 1024.0, SharedShaderCount);
-	UE_LOG(LogIoStore, Display, TEXT("Input:  %8.2f MB for %d Unique shaders"), (double)UniqueShaderSize / 1024.0 / 1024.0, UniqueShaderCount);
-	UE_LOG(LogIoStore, Display, TEXT("Input:  %8.2f MB for %d Inline shaders"), (double)InlineShaderSize / 1024.0 / 1024.0, InlineShaderCount);
+	UE_LOG(LogIoStore, Display, TEXT("Input:  %s UExp"),   *FormatMemoryString(UExpSize));
+	UE_LOG(LogIoStore, Display, TEXT("Input:  %s UAsset"), *FormatMemoryString(UAssetSize));
+	UE_LOG(LogIoStore, Display, TEXT("Input:  %s UBulk"),  *FormatMemoryString(UBulkSize));
+	UE_LOG(LogIoStore, Display, TEXT("Input:  %s for %d Global shaders"), *FormatMemoryString(GlobalShaderSize), GlobalShaderCount);
+	UE_LOG(LogIoStore, Display, TEXT("Input:  %s for %d Shared shaders"), *FormatMemoryString(SharedShaderSize), SharedShaderCount);
+	UE_LOG(LogIoStore, Display, TEXT("Input:  %s for %d Unique shaders"), *FormatMemoryString(UniqueShaderSize), UniqueShaderCount);
+	UE_LOG(LogIoStore, Display, TEXT("Input:  %s for %d Inline shaders"), *FormatMemoryString(InlineShaderSize), InlineShaderCount);
 	UE_LOG(LogIoStore, Display, TEXT(""));
+
 	if (bIsLegacyStage)
 	{
 		UE_LOG(LogIoStore, Display, TEXT("Output: %8llu Name map entries"), NameMapCount);
 	}
+
 	UE_LOG(LogIoStore, Display, TEXT("Output: %8llu Imported package entries"), ImportedPackagesCount);
 	UE_LOG(LogIoStore, Display, TEXT("Output: %8llu Packages without imports"), NoImportedPackagesCount);
 	UE_LOG(LogIoStore, Display, TEXT("Output: %8d Public runtime script objects"), PackageStoreOptimizer.GetTotalScriptObjectCount());
+	
 	if (bIsLegacyStage)
 	{
-		UE_LOG(LogIoStore, Display, TEXT("Output: %8.2lf GB HeaderData"), (double)HeaderSize / 1024.0 / 1024.0 / 1024.0);
+		UE_LOG(LogIoStore, Display, TEXT("Output: %s HeaderData"), *FormatMemoryString(HeaderSize));
 	}
-	UE_LOG(LogIoStore, Display, TEXT("Output: %8.2lf MB InitialLoadData"), (double)InitialLoadSize / 1024.0 / 1024.0);
+
+	UE_LOG(LogIoStore, Display, TEXT("Output: %s InitialLoadData"), *FormatMemoryString(InitialLoadSize));
 
 	if (FStudioTelemetry::Get().IsSessionRunning() == true)
 	{
