Minor tweak to UnrealPak Telemetry
Added Optional/Installed sizes
Fixed missing Flush in StudioTlemetry

diff --git a/Engine/Source/Developer/IoStoreUtilities/Private/IoStoreUtilities.cpp b/Engine/Source/Developer/IoStoreUtilities/Private/IoStoreUtilities.cpp
--- a/Engine/Source/Developer/IoStoreUtilities/Private/IoStoreUtilities.cpp
+++ b/Engine/Source/Developer/IoStoreUtilities/Private/IoStoreUtilities.cpp
@@ -1167,6 +1167,7 @@ using FPackageIdMap = TMap<FPackageId, FCookedPackage*>;
 enum EAggregate
 {
 	Installed,
+	Optional,
 	IAS,
 	IAD,
 	NumTypes
@@ -1175,6 +1176,7 @@ enum EAggregate
 const FString AggregateNames[EAggregate::NumTypes] =
 {
 	TEXT("Installed"),
+	TEXT("Optional"),
 	TEXT("IAS"),
 	TEXT("IAD")
 };
@@ -1263,13 +1265,21 @@ public:
 			else if (ContainerName.Contains(TEXT("ondemand")))
 			{
 				DeliveryType = TEXT("IAS");
-				ChunkEntrySummary.UncompressedTotals[EAggregate::IAD] += Info.Size;
-				ChunkEntrySummary.CompressedTotals[EAggregate::IAD] += Info.CompressedSize;
+				ChunkEntrySummary.UncompressedTotals[EAggregate::IAS] += Info.Size;
+				ChunkEntrySummary.CompressedTotals[EAggregate::IAS] += Info.CompressedSize;
 			}
 			else
 			{
-				ChunkEntrySummary.UncompressedTotals[EAggregate::Installed] += Info.Size;
-				ChunkEntrySummary.CompressedTotals[EAggregate::Installed] += Info.CompressedSize;
+				if (ContainerName.Contains(TEXT("optional")))
+				{
+					ChunkEntrySummary.UncompressedTotals[EAggregate::Optional] += Info.Size;
+					ChunkEntrySummary.CompressedTotals[EAggregate::Optional] += Info.CompressedSize;
+				}
+				else
+				{
+					ChunkEntrySummary.UncompressedTotals[EAggregate::Installed] += Info.Size;
+					ChunkEntrySummary.CompressedTotals[EAggregate::Installed] += Info.CompressedSize;
+				}
 			}
 
 			ChunkEntrySummary.PackageCount++;
