A lightweight IoStore telemetry event for ThinClient that allows us to see trends in our builds in a more efficeint way.
Changed the schema for the IoStoreChunk.Summary event to v2.
PakChunk, PackageCount, Size and UncompressedSize are sent once per PakChunk per InstallType ( Installed, Optional, IAS, IAD ).
Aggregation path is only ever exectued with UnrealPak -list
Telemetry event is only ever posted whith UnrealPak -telemetry

diff --git a/Engine/Source/Developer/IoStoreUtilities/Private/IoStoreUtilities.cpp b/Engine/Source/Developer/IoStoreUtilities/Private/IoStoreUtilities.cpp
--- a/Engine/Source/Developer/IoStoreUtilities/Private/IoStoreUtilities.cpp
+++ b/Engine/Source/Developer/IoStoreUtilities/Private/IoStoreUtilities.cpp
@@ -1181,44 +1181,76 @@ const FString AggregateNames[EAggregate::NumTypes] =
 	TEXT("IAD")
 };
 
-class FChunkEntrySummary
+class FChunkDataSummary
 {
 public:
-	uint64 PackageCount;
-	uint64 UncompressedTotals[EAggregate::NumTypes];
-	uint64 CompressedTotals[EAggregate::NumTypes];
 
-	FChunkEntrySummary()
+	class FChunkData
 	{
-		PackageCount = 0;
+	public:
+		uint64 PackageCount[EAggregate::NumTypes];
+		uint64 UncompressedSize[EAggregate::NumTypes];
+		uint64 CompressedSize[EAggregate::NumTypes];
+
+		FChunkData()
+		{
+			for (uint32 i = 0; i < EAggregate::NumTypes; i++)
+			{
+				PackageCount[i] = UncompressedSize[i] = CompressedSize[i] = 0;
+			}
+		}
+	};
+	
+	void AddPackage(const FString& PakChunk, EAggregate Type, uint64 UncompressedSize, uint64 CompressedSize)
+	{
+		// See if we already have data for this chunk
+		FChunkDataSummary::FChunkData* ChunkData = ChunkDataMap.Find(PakChunk);
 
-		for (uint32 i = 0; i < EAggregate::NumTypes; i++)
+		if (ChunkData == nullptr)
 		{
-			UncompressedTotals[i] = CompressedTotals[i] = 0;
+			// Add a new entry for this chunk
+			ChunkData = &ChunkDataMap.Emplace(PakChunk);
 		}
+
+		ChunkData->UncompressedSize[Type] += UncompressedSize;
+		ChunkData->CompressedSize[Type] += CompressedSize;
+		ChunkData->PackageCount[Type]++;
 	}
 
 	void SendTelemetryEvent()
 	{
 		if (FStudioTelemetry::Get().IsSessionRunning())
 		{
-			// Write out the summary to telemetry
-			TArray<FAnalyticsEventAttribute> Attributes;
+			const int SchemaVersion = 2;
 
-			const int SchemaVersion = 1;
+			for (TMap< FString, FChunkData>::TConstIterator it(ChunkDataMap); it; ++it)
+			{
+				// Write out the aggregated data for each PakChunk to telemetry by delivery type
+				const FChunkData& ChunkData = (*it).Value;
 
-			Attributes.Emplace(TEXT("SchemaVersion"), SchemaVersion);
-			Attributes.Emplace(TEXT("PackageCount"), PackageCount);
+				for (uint32 i = 0; i < EAggregate::NumTypes; i++)
+				{
+					if (ChunkData.PackageCount[i] > 0)
+					{
+						TArray<FAnalyticsEventAttribute> Attributes;
 
-			for (uint32 i = 0; i < EAggregate::NumTypes; i++)
-			{
-				Attributes.Emplace(*FString::Printf(TEXT("UncompressedSize%s"), *AggregateNames[i]), UncompressedTotals[i]);
-				Attributes.Emplace(*FString::Printf(TEXT("CompressedSize%s"), *AggregateNames[i]), CompressedTotals[i]);
-			}
+						Attributes.Emplace(TEXT("SchemaVersion"), SchemaVersion);
+						Attributes.Emplace(TEXT("PakChunk"), (*it).Key);
+						Attributes.Emplace(TEXT("DeliveryType"), *AggregateNames[i]);
+						Attributes.Emplace(TEXT("PackageCount"), ChunkData.PackageCount[i]);
+						Attributes.Emplace(TEXT("Size"), ChunkData.UncompressedSize[i]);
+						Attributes.Emplace(TEXT("CompressedSize"), ChunkData.CompressedSize[i]);
 
-			FStudioTelemetry::Get().RecordEvent(TEXT("Core.IoStoreChunk.Summary"), Attributes);
+						FStudioTelemetry::Get().RecordEvent(TEXT("Core.IoStoreChunk.Summary"), Attributes);
+					}
+				}	
+			}
 		}
 	}
+
+private:
+
+	TMap< FString, FChunkData> ChunkDataMap;
 };
 
 class FChunkEntryCsv
@@ -1230,7 +1262,7 @@ public:
 		if (OutputArchive.IsValid())
 		{
 			OutputArchive->Flush();
-			ChunkEntrySummary.SendTelemetryEvent();
+			ChunkDataSummary.SendTelemetryEvent();
 		}
 	}
 
@@ -1258,32 +1290,29 @@ public:
 			// Determine the Delivery Type Installed/IAS/IAD
 			if (ContainerName.Contains(TEXT("iad")))
 			{
+				// IAD data
 				DeliveryType = TEXT("IAD");
-				ChunkEntrySummary.UncompressedTotals[EAggregate::IAD] += Info.Size;
-				ChunkEntrySummary.CompressedTotals[EAggregate::IAD] += Info.CompressedSize;
+				ChunkDataSummary.AddPackage(PakChunk, EAggregate::IAD, Info.Size, Info.CompressedSize);
 			}
 			else if (ContainerName.Contains(TEXT("ondemand")))
 			{
+				// IAS data
 				DeliveryType = TEXT("IAS");
-				ChunkEntrySummary.UncompressedTotals[EAggregate::IAS] += Info.Size;
-				ChunkEntrySummary.CompressedTotals[EAggregate::IAS] += Info.CompressedSize;
+				ChunkDataSummary.AddPackage(PakChunk, EAggregate::IAS, Info.Size, Info.CompressedSize);
 			}
 			else
 			{
+				// Mandatory or Optionally Installed data
 				if (ContainerName.Contains(TEXT("optional")))
 				{
-					ChunkEntrySummary.UncompressedTotals[EAggregate::Optional] += Info.Size;
-					ChunkEntrySummary.CompressedTotals[EAggregate::Optional] += Info.CompressedSize;
+					ChunkDataSummary.AddPackage(PakChunk, EAggregate::Optional, Info.Size, Info.CompressedSize);
 				}
 				else
 				{
-					ChunkEntrySummary.UncompressedTotals[EAggregate::Installed] += Info.Size;
-					ChunkEntrySummary.CompressedTotals[EAggregate::Installed] += Info.CompressedSize;
+					ChunkDataSummary.AddPackage(PakChunk, EAggregate::Installed, Info.Size, Info.CompressedSize);
 				}
 			}
 
-			ChunkEntrySummary.PackageCount++;
-
 			// Determine the Install Type Base/Optional
 			if (ContainerName.Contains(TEXT("optional")))
 			{
@@ -1314,7 +1343,7 @@ public:
 
 private:
 	TUniquePtr<FArchive> OutputArchive;
-	FChunkEntrySummary ChunkEntrySummary;
+	FChunkDataSummary ChunkDataSummary;
 };
 
 void SortPackagesInLoadOrderRecursive(TArray<FCookedPackage*>& Result, FCookedPackage* Package, TArray<FCookedPackage*>& S, TArray<FCookedPackage*>& P, int32& C, const TMap<FCookedPackage*, TArray<FCookedPackage*>>& ReverseEdges)
