Improved an existing comment and added a new one based on post submit feedback.

diff --git a/Engine/Source/Developer/IoStoreUtilities/Private/IoStoreUtilities.cpp b/Engine/Source/Developer/IoStoreUtilities/Private/IoStoreUtilities.cpp
--- a/Engine/Source/Developer/IoStoreUtilities/Private/IoStoreUtilities.cpp
+++ b/Engine/Source/Developer/IoStoreUtilities/Private/IoStoreUtilities.cpp
@@ -1203,7 +1203,7 @@ public:
 	
 	void AddPackage(const FString& PakChunk, EAggregate Type, uint64 UncompressedSize, uint64 CompressedSize)
 	{
-		// See if we already have data for this chunk
+		// See if we already have an entry for this chunk
 		FChunkDataSummary::FChunkData* ChunkData = ChunkDataMap.Find(PakChunk);
 
 		if (ChunkData == nullptr)
@@ -1225,13 +1225,14 @@ public:
 
 			for (TMap< FString, FChunkData>::TConstIterator it(ChunkDataMap); it; ++it)
 			{
-				// Write out the aggregated data for each PakChunk to telemetry by delivery type
 				const FChunkData& ChunkData = (*it).Value;
 
 				for (uint32 i = 0; i < EAggregate::NumTypes; i++)
 				{
+					// Only send an event if the PakChunk contains one or more packages
 					if (ChunkData.PackageCount[i] > 0)
 					{
+						// Build an event from the aggregated data for each PakChunk for each DeliveryType
 						TArray<FAnalyticsEventAttribute> Attributes;
 
 						Attributes.Emplace(TEXT("SchemaVersion"), SchemaVersion);
@@ -1241,6 +1242,7 @@ public:
 						Attributes.Emplace(TEXT("Size"), ChunkData.UncompressedSize[i]);
 						Attributes.Emplace(TEXT("CompressedSize"), ChunkData.CompressedSize[i]);
 
+						// Send the event to the telemetry system
 						FStudioTelemetry::Get().RecordEvent(TEXT("Core.IoStoreChunk.Summary"), Attributes);
 					}
 				}	
