IncrementalCook: DLC: Second attempt at crash fixes from static analysis.
#rnx

diff --git a/Engine/Source/Developer/IoStoreUtilities/Private/ZenStoreWriter.cpp b/Engine/Source/Developer/IoStoreUtilities/Private/ZenStoreWriter.cpp
--- a/Engine/Source/Developer/IoStoreUtilities/Private/ZenStoreWriter.cpp
+++ b/Engine/Source/Developer/IoStoreUtilities/Private/ZenStoreWriter.cpp
@@ -1184,10 +1184,10 @@ FCbObject FZenStoreWriter::GetOplogAttachment(FName PackageName, FUtf8StringView
 }
 
 void FZenStoreWriter::GetOplogAttachments(TArrayView<FName> PackageNames,
-	TArrayView<FUtf8StringView> AttachmentKeys,
+	TArrayView<FUtf8StringView> InAttachmentKeys,
 	TUniqueFunction<void(FName PackageName, FUtf8StringView AttachmentKey, FCbObject&& Attachment)>&& Callback)
 {
-	const int MaximumHashCount = PackageNames.Num() * AttachmentKeys.Num();
+	const int MaximumHashCount = PackageNames.Num() * InAttachmentKeys.Num();
 	TArray<FIoHash> AttachmentHashes;
 	AttachmentHashes.Reserve(MaximumHashCount);
 	struct FAttachmentHashParam
@@ -1206,10 +1206,20 @@ void FZenStoreWriter::GetOplogAttachments(TArrayView<FName> PackageNames,
 	TArray<FAttachmentHashParam> InvalidAttachmentHashParams;
 	InvalidAttachmentHashParams.Reserve(MaximumHashCount);
 
-	TArray<const UTF8CHAR*, TInlineAllocator<2>> AttachmentIds;
-	for (FUtf8StringView AttachmentKey : AttachmentKeys)
+	TArray<FUtf8StringView, TInlineAllocator<2>> AttachmentIds;
+	for (FUtf8StringView AttachmentKey : InAttachmentKeys)
 	{
-		AttachmentIds.Add(UE::FZenStoreHttpClient::FindAttachmentId(AttachmentKey));
+		const UTF8CHAR* Id = UE::FZenStoreHttpClient::FindOrAddAttachmentId(AttachmentKey);
+		FUtf8StringView NewKey(Id);
+		// The AttachmentId is supposed to be byte identical to the attachment looked up,
+		// we look it up only so that we can have a single pointer value to use as Id.
+		// Enforce the contract that the id is byte identical. If we ever change id to something
+		// else (e.g. an abbreviation), then we will need to get a pointer to a persistent copy
+		// of the key. We can not use the original AttachmentKey as our key because we store it
+		// in an async function and we do not require that the caller maintain their copy of it
+		// in memory after we return.
+		check(NewKey.Equals(AttachmentKey, ESearchCase::CaseSensitive));
+		AttachmentIds.Add(NewKey);
 	}
 
 	{
@@ -1219,10 +1229,9 @@ void FZenStoreWriter::GetOplogAttachments(TArrayView<FName> PackageNames,
 		{
 			const int32* Idx = PackageNameToIndex.Find(PackageName);
 
-			for (int32 InputAttachmentIndex = 0; InputAttachmentIndex < AttachmentKeys.Num(); ++InputAttachmentIndex)
+			for (FUtf8StringView AttachmentKey : AttachmentIds)
 			{
-				FUtf8StringView AttachmentKey = AttachmentKeys[InputAttachmentIndex];
-				const UTF8CHAR* AttachmentId = AttachmentIds[InputAttachmentIndex];
+				const UTF8CHAR* AttachmentId = AttachmentKey.GetData();
 
 				FIoHash AttachmentHash;
 				ON_SCOPE_EXIT
@@ -1300,14 +1309,14 @@ void FZenStoreWriter::GetOplogAttachments(TArrayView<FName> PackageNames,
 }
 
 void FZenStoreWriter::GetBaseGameOplogAttachments(TArrayView<FName> PackageNames,
-	TArrayView<FUtf8StringView> AttachmentKeys,
+	TArrayView<FUtf8StringView> InAttachmentKeys,
 	TUniqueFunction<void(FName PackageName, FUtf8StringView AttachmentKey, FCbObject&& Attachment)>&& Callback)
 {
 	if (ParentProjectId.IsEmpty())
 	{
 		for (FName PackageName : PackageNames)
 		{
-			for (FUtf8StringView AttachmentKey : AttachmentKeys)
+			for (FUtf8StringView AttachmentKey : InAttachmentKeys)
 			{
 				Callback(PackageName, AttachmentKey, FCbObject());
 			}
@@ -1315,10 +1324,21 @@ void FZenStoreWriter::GetBaseGameOplogAttachments(TArrayView<FName> PackageNames
 		return;
 	}
 
-	TArray<FUtf8String> AttachmentKeysBuffer;
-	for (FUtf8StringView AttachmentKey : AttachmentKeys)
+	TArray<FUtf8StringView> AttachmentIds;
+	AttachmentIds.Reserve(InAttachmentKeys.Num());
+	for (FUtf8StringView AttachmentKey : InAttachmentKeys)
 	{
-		AttachmentKeysBuffer.Emplace(AttachmentKey);
+		const UTF8CHAR* Id = UE::FZenStoreHttpClient::FindOrAddAttachmentId(AttachmentKey);
+		FUtf8StringView NewKey(Id);
+		// The AttachmentId is supposed to be byte identical to the attachment looked up,
+		// we look it up only so that we can have a single pointer value to use as Id.
+		// Enforce the contract that the id is byte identical. If we ever change id to something
+		// else (e.g. an abbreviation), then we will need to get a pointer to a persistent copy
+		// of the key. We can not use the original AttachmentKey as our key because we store it
+		// in an async function and we do not require that the caller maintain their copy of it
+		// in memory after we return.
+		check(NewKey.Equals(AttachmentKey, ESearchCase::CaseSensitive));
+		AttachmentIds.Add(NewKey);
 	}
 	TArray<FName> PackageNamesBuffer(PackageNames);
 	TArray<FString> EntryKeys;
@@ -1329,11 +1349,11 @@ void FZenStoreWriter::GetBaseGameOplogAttachments(TArrayView<FName> PackageNames
 	}
 
 	HttpClient->GetOplogEntries(ParentProjectId, OplogId, MoveTemp(EntryKeys)).Next(
-		[HttpClient = this->HttpClient.Get(), AttachmentKeys = MoveTemp(AttachmentKeysBuffer), Callback = MoveTemp(Callback),
+		[HttpClient = this->HttpClient.Get(), AttachmentIds = MoveTemp(AttachmentIds), Callback = MoveTemp(Callback),
 		PackageNames = MoveTemp(PackageNamesBuffer)]
 		(TArray<TIoStatusOr<FCbObject>> Entries) mutable
 		{
-			const int MaximumHashCount = PackageNames.Num() * AttachmentKeys.Num();
+			const int MaximumHashCount = PackageNames.Num() * AttachmentIds.Num();
 			TArray<FIoHash> AttachmentHashes;
 			AttachmentHashes.Reserve(MaximumHashCount);
 			struct FAttachmentHashParam
@@ -1352,12 +1372,6 @@ void FZenStoreWriter::GetBaseGameOplogAttachments(TArrayView<FName> PackageNames
 			TArray<FAttachmentHashParam> InvalidAttachmentHashParams;
 			InvalidAttachmentHashParams.Reserve(MaximumHashCount);
 
-			TArray<const UTF8CHAR*, TInlineAllocator<2>> AttachmentIds;
-			for (FUtf8StringView AttachmentKey : AttachmentKeys)
-			{
-				AttachmentIds.Add(UE::FZenStoreHttpClient::FindAttachmentId(AttachmentKey));
-			}
-
 			TMap<FName, TMap<const UTF8CHAR*, FIoHash>> AttachmentsForPackage;
 			for (TIoStatusOr<FCbObject>& EntryStatus : Entries)
 			{
@@ -1395,10 +1409,9 @@ void FZenStoreWriter::GetBaseGameOplogAttachments(TArrayView<FName> PackageNames
 			{
 				TMap<const UTF8CHAR*, FIoHash>* Attachments = AttachmentsForPackage.Find(PackageName);
 
-				for (int32 InputAttachmentIndex = 0; InputAttachmentIndex < AttachmentKeys.Num(); ++InputAttachmentIndex)
+				for (FUtf8StringView AttachmentKey : AttachmentIds)
 				{
-					FUtf8StringView AttachmentKey = AttachmentKeys[InputAttachmentIndex];
-					const UTF8CHAR* AttachmentId = AttachmentIds[InputAttachmentIndex];
+					const UTF8CHAR* AttachmentId = AttachmentKey.GetData();
 					FIoHash* AttachmentHash = (Attachments && AttachmentId) ? Attachments->Find(AttachmentId) : nullptr;
 					if (AttachmentHash)
 					{
