Make API for async package receiver more concise by having the http response and package response read from the receiver itself instead of passed as arguments to the completion callback.

diff --git a/Engine/Source/Developer/DerivedDataCache/Private/ZenCacheStore.cpp b/Engine/Source/Developer/DerivedDataCache/Private/ZenCacheStore.cpp
--- a/Engine/Source/Developer/DerivedDataCache/Private/ZenCacheStore.cpp
+++ b/Engine/Source/Developer/DerivedDataCache/Private/ZenCacheStore.cpp
@@ -1632,7 +1632,7 @@ public:
 		FOnComplete&& InOnComplete)
 		: FAsyncCbPackageReceiver(MoveTemp(InRequest), InZenServiceInstance,
 		[Self = this, LocalOnComplete = MoveTemp(InOnComplete)]
-		(FAsyncCbPackageReceiver* Receiver, const THttpUniquePtr<IHttpResponse>& HttpResponse, const FCbPackage& InPackage) mutable
+		(FAsyncCbPackageReceiver* Receiver) mutable
 		{
 			Self->Owner->LaunchTask(TEXT("FRequestAsyncCbPackageReceiver::OnComplete"), [Self = TRefCountPtr(Self), InnerOnComplete = MoveTemp(LocalOnComplete)]
 			{
diff --git a/Engine/Source/Developer/IoStoreUtilities/Private/ZenStoreHttpClient.cpp b/Engine/Source/Developer/IoStoreUtilities/Private/ZenStoreHttpClient.cpp
--- a/Engine/Source/Developer/IoStoreUtilities/Private/ZenStoreHttpClient.cpp
+++ b/Engine/Source/Developer/IoStoreUtilities/Private/ZenStoreHttpClient.cpp
@@ -396,12 +396,12 @@ void FZenStoreHttpClient::RpcAsync(FAnsiStringView Uri, FCbObject RpcRequestObje
 	// PackageReceiver will delete itself in the completion callback
 	Zen::FAsyncCbPackageReceiver* PackageReceiver = new Zen::FAsyncCbPackageReceiver(MoveTemp(Request), ZenService.GetInstance(),
 		[ResponseCallback = MoveTemp(ResponseCallback)]
-		(Zen::FAsyncCbPackageReceiver* Receiver, const THttpUniquePtr<IHttpResponse>& HttpResponse, const FCbPackage& Response) mutable
+		(Zen::FAsyncCbPackageReceiver* Receiver) mutable
 		{
 			// Launch a task for the completion function since it can execute arbitrary code.
-			Tasks::Launch(TEXT("FZenStoreHttpClient::RpcAsync"), [Receiver, &Response, InnerResponseCallback = MoveTemp(ResponseCallback)]
+			Tasks::Launch(TEXT("FZenStoreHttpClient::RpcAsync"), [Receiver, InnerResponseCallback = MoveTemp(ResponseCallback)]
 			{
-				InnerResponseCallback(Response);
+				InnerResponseCallback(Receiver->GetResponsePackage());
 				delete Receiver;
 			});
 		},
diff --git a/Engine/Source/Developer/Zen/Internal/ZenAsyncCbPackageReceiver.h b/Engine/Source/Developer/Zen/Internal/ZenAsyncCbPackageReceiver.h
--- a/Engine/Source/Developer/Zen/Internal/ZenAsyncCbPackageReceiver.h
+++ b/Engine/Source/Developer/Zen/Internal/ZenAsyncCbPackageReceiver.h
@@ -19,7 +19,7 @@ class FZenServiceInstance;
 class FAsyncCbPackageReceiver : public IHttpReceiver
 {
 public:
-	using FOnComplete = TUniqueFunction<void(FAsyncCbPackageReceiver* Receiver, const THttpUniquePtr<IHttpResponse>& HttpResponse, const FCbPackage& ResponsePackage)>;
+	using FOnComplete = TUniqueFunction<void(FAsyncCbPackageReceiver* Receiver)>;
 
 	FAsyncCbPackageReceiver(const FAsyncCbPackageReceiver&) = delete;
 	FAsyncCbPackageReceiver& operator=(const FAsyncCbPackageReceiver&) = delete;
@@ -30,6 +30,10 @@ public:
 		FOnComplete&& InOnComplete,
 		int InMaxAttempts = 1);
 
+	ZEN_API const IHttpResponse& GetHttpResponse() { return *Response; }
+	ZEN_API const FCbPackage& GetResponsePackage() { return Package; }
+
+
 	ZEN_API void SendAsync();
 
 protected:
diff --git a/Engine/Source/Developer/Zen/Private/ZenAsyncCbPackageReceiver.cpp b/Engine/Source/Developer/Zen/Private/ZenAsyncCbPackageReceiver.cpp
--- a/Engine/Source/Developer/Zen/Private/ZenAsyncCbPackageReceiver.cpp
+++ b/Engine/Source/Developer/Zen/Private/ZenAsyncCbPackageReceiver.cpp
@@ -52,7 +52,7 @@ IHttpReceiver* FAsyncCbPackageReceiver::OnComplete(IHttpResponse& LocalResponse)
 		FOnComplete LocalOnComplete = MoveTemp(OnCompleteCallback);
 
 		// Calling LocalOnComplete may result in "this" being deleted, so no further access can happen to anything on "this"
-		LocalOnComplete(this, Response, Package);
+		LocalOnComplete(this);
 	}
 	return nullptr;
 }
