Comment clarification

diff --git a/Engine/Source/Developer/IoStoreUtilities/Private/ZenStoreHttpClient.cpp b/Engine/Source/Developer/IoStoreUtilities/Private/ZenStoreHttpClient.cpp
--- a/Engine/Source/Developer/IoStoreUtilities/Private/ZenStoreHttpClient.cpp
+++ b/Engine/Source/Developer/IoStoreUtilities/Private/ZenStoreHttpClient.cpp
@@ -393,7 +393,7 @@ void FZenStoreHttpClient::RpcAsync(FAnsiStringView Uri, FCbObject RpcRequestObje
 	Request->AddAcceptType(EHttpMediaType::CbPackage);
 	Request->SetContentType(EHttpMediaType::CbObject);
 	Request->SetBody(RpcRequestObject.GetBuffer().MakeOwned());
-	// PackageReceiver will delete itself in the completion callback
+	// PackageReceiver will delete itself in an async task launched from the completion callback
 	Zen::FAsyncCbPackageReceiver* PackageReceiver = new Zen::FAsyncCbPackageReceiver(MoveTemp(Request), ZenService.GetInstance(),
 		[ResponseCallback = MoveTemp(ResponseCallback)]
 		(Zen::FAsyncCbPackageReceiver* Receiver) mutable
