When incrementally cooking, ensure that files on the filesystem that were part of the cooked output but have been deleted are no longer included in the zen filesystem manifest that is written to to the "EndCook" op.

diff --git a/Engine/Source/Developer/IoStoreUtilities/Private/ZenStoreWriter.cpp b/Engine/Source/Developer/IoStoreUtilities/Private/ZenStoreWriter.cpp
--- a/Engine/Source/Developer/IoStoreUtilities/Private/ZenStoreWriter.cpp
+++ b/Engine/Source/Developer/IoStoreUtilities/Private/ZenStoreWriter.cpp
@@ -580,10 +580,14 @@ void FZenStoreWriter::Initialize(const FCookInfo& Info)
 
 					if (!ServerPath.IsEmpty())
 					{
-						FIoChunkId FileChunkId;
-						FileChunkId.Set(FileId.GetView());
+						const FString AbsPath = ZenFileSystemManifest->ServerRootPath() / ServerPath;
+						if (FPaths::FileExists(AbsPath))
+						{
+							FIoChunkId FileChunkId;
+							FileChunkId.Set(FileId.GetView());
 
-						ZenFileSystemManifest->AddManifestEntry(FileChunkId, MoveTemp(ServerPath), MoveTemp(ClientPath));
+							ZenFileSystemManifest->AddManifestEntry(FileChunkId, MoveTemp(ServerPath), MoveTemp(ClientPath));
+						}
 					}
 				}
 
