UnrealPak: Make the error that a required cooked chunk fails to load from ZenServer non-fatal.
Give a better diagnostic about  the missing chunk.
This allows testing to still occur for tests that do not require the file with the missing chunk (by changing the exit code to 0 or re-running UAT stage with "-pak -skippak").

#rnx

diff --git a/Engine/Source/Developer/IoStoreUtilities/Private/IoStoreLooseFiles.cpp b/Engine/Source/Developer/IoStoreUtilities/Private/IoStoreLooseFiles.cpp
--- a/Engine/Source/Developer/IoStoreUtilities/Private/IoStoreLooseFiles.cpp
+++ b/Engine/Source/Developer/IoStoreUtilities/Private/IoStoreLooseFiles.cpp
@@ -115,6 +115,11 @@ public:
 				return SourceBuffer.DataSize();
 			}
 
+			virtual const TCHAR* DebugNameOfRepository() const override
+			{
+				return TEXT("FLooseFilesIoStoreWriter::Append");
+			}
+
 			FIoBuffer SourceBuffer;
 			uint64 OrderHint;
 		};
diff --git a/Engine/Source/Developer/IoStoreUtilities/Private/IoStoreUtilities.cpp b/Engine/Source/Developer/IoStoreUtilities/Private/IoStoreUtilities.cpp
--- a/Engine/Source/Developer/IoStoreUtilities/Private/IoStoreUtilities.cpp
+++ b/Engine/Source/Developer/IoStoreUtilities/Private/IoStoreUtilities.cpp
@@ -3396,11 +3396,12 @@ private:
 
 		virtual const FIoBuffer* GetSourceBuffer() override
 		{
-			return &SourceBuffer;
+			return bSourceBufferValid ? &SourceBuffer : nullptr;
 		}
 		
 		virtual void FreeSourceBuffer() override
 		{
+			bSourceBufferValid = false;
 			if (SourceBuffer.DataSize() > 0)
 			{
 				SourceBuffer = FIoBuffer();
@@ -3422,10 +3423,11 @@ private:
 			, FileRegions(TargetFile.FileRegions)
 			, SourceBufferSize(TargetFile.SourceSize) { }
 
-		void OnSourceBufferLoaded(bool bTriggerCompletionEvent)
+		void OnSourceBufferLoaded(bool bSuccess, bool bTriggerCompletionEvent)
 		{
 			TRACE_COUNTER_DECREMENT(IoStoreSourceReadsInflight);
 			TRACE_COUNTER_INCREMENT(IoStoreSourceReadsDone);
+			bSourceBufferValid = bSuccess;
 			QueueEntry->ReleaseRef(Manager);
 			if (bTriggerCompletionEvent)
 			{
@@ -3444,6 +3446,7 @@ private:
 		TOptional<UE::Tasks::FTaskEvent> CompletionEvent;
 		FIoBuffer SourceBuffer;
 		FQueueEntry* QueueEntry = nullptr;
+		bool bSourceBufferValid = false;
 	};
 
 	class FInMemoryWriteRequest
@@ -3458,7 +3461,12 @@ private:
 			Manager.MemorySourceReads[(int8)TargetFile.ChunkId.GetChunkType()].IncrementExchange();
 			SourceBuffer = TargetFile.SourceBuffer.GetValue();
 			Manager.MemorySourceBytes[(int8)TargetFile.ChunkId.GetChunkType()] += SourceBuffer.DataSize();
-			OnSourceBufferLoaded(/*bTriggerCompletionEvent*/ true);
+			OnSourceBufferLoaded(/*bSuccess*/ true, /*bTriggerCompletionEvent*/ true);
+		}
+
+		virtual const TCHAR* DebugNameOfRepository() const override
+		{
+			return TEXT("InMemory");
 		}
 	};
 
@@ -3495,7 +3503,7 @@ private:
 					check(Package->OptimizedOptionalSegmentPackage);
 					SourceBuffer = Manager.PackageStoreOptimizer.CreatePackageBuffer(Package->OptimizedOptionalSegmentPackage, SourceBuffer);
 				}
-				OnSourceBufferLoaded(/*bTriggerCompletionEvent*/ true);
+				OnSourceBufferLoaded(/*bSuccess*/ true, /*bTriggerCompletionEvent*/ true);
 			};
 
 			QueueEntry->ReadRequest.Reset(
@@ -3503,6 +3511,11 @@ private:
 			QueueEntry->ReleaseRef(Manager);
 		}
 
+		virtual const TCHAR* DebugNameOfRepository() const override
+		{
+			return TEXT("LooseFile");
+		}
+
 	private:
 		FLegacyCookedPackage* Package;
 	};
@@ -3521,13 +3534,22 @@ private:
 				TargetFile.ChunkId,
 				[this](TIoStatusOr<FIoBuffer> Status)
 				{
-					SourceBuffer = Status.ConsumeValueOrDie();
-					Manager.ZenSourceBytes[(int8)TargetFile.ChunkId.GetChunkType()] += SourceBuffer.DataSize();
-					OnSourceBufferLoaded(/*bTriggerCompletionEvent*/ false);
+					const bool bSuccess = Status.IsOk();
+					if (bSuccess)
+					{
+						SourceBuffer = Status.ConsumeValueOrDie();
+						Manager.ZenSourceBytes[(int8)TargetFile.ChunkId.GetChunkType()] += SourceBuffer.DataSize();
+					}
+					OnSourceBufferLoaded(bSuccess, /*bTriggerCompletionEvent*/ false);
 				});
 			CompletionEvent.GetValue().AddPrerequisites(ReadTask);
 			CompletionEvent.GetValue().Trigger();
 		}
+
+		virtual const TCHAR* DebugNameOfRepository() const override
+		{
+			return TEXT("ZenServer");
+		}
 	};
 
 	struct FQueueEntry
@@ -6045,6 +6067,14 @@ int32 CreateTarget(const FIoStoreArguments& Arguments, const FIoStoreWriterSetti
 		FTaskGraphInterface::Get().WaitUntilTaskCompletes(WriteCsvFileTask);
 	}
 
+	if (uint32 ErrorCount = IoStoreWriterContext->GetErrors())
+	{
+		// Change this to return 0 or re-run UAT stage run with "-pak -skippak" to
+		// force produce a local broken build that will crash upon loading any of the reported failed assets.
+		UE_LOG(LogIoStore, Display, TEXT("Failure - %d error(s)"), ErrorCount);
+		return -2;
+	}
+	UE_LOG(LogIoStore, Display, TEXT("Success"));
 	return 0;
 }
 
@@ -10203,7 +10233,7 @@ int32 CreateIoStoreContainerFiles(const TCHAR* CmdLine)
 
 	if (!ParseOrderFileArguments(Arguments))
 	{
-		return false;
+		return -1;
 	}
 
 	FString ScriptObjectsFile;
diff --git a/Engine/Source/Developer/IoStoreUtilities/Private/IoStoreWriter.cpp b/Engine/Source/Developer/IoStoreUtilities/Private/IoStoreWriter.cpp
--- a/Engine/Source/Developer/IoStoreUtilities/Private/IoStoreWriter.cpp
+++ b/Engine/Source/Developer/IoStoreUtilities/Private/IoStoreWriter.cpp
@@ -551,6 +551,16 @@ public:
 		TRACE_COUNTER_SUBTRACT(IoStoreCompressionMemoryUsed, CompressionBufferSize);
 	}
 
+	void ReportError()
+	{
+		ErrorCount.IncrementExchange();
+	}
+
+	uint32 GetErrors() const
+	{
+		return ErrorCount.Load(EMemoryOrder::Relaxed);
+	}
+
 private:
 	UE::DerivedData::FCacheKey MakeDDCKey(FIoStoreWriteQueueEntry* Entry) const;
 
@@ -587,6 +597,7 @@ private:
 	TAtomic<uint64> CompressionDDCMissCount{ 0 };
 	TAtomic<uint64> CompressionDDCPutErrorCount{ 0 };
 	TAtomic<uint64> ScheduledCompressionMemory{ 0 };
+	TAtomic<int32> ErrorCount{ 0 };
 	FCriticalSection AvailableCompressionBuffersCritical;
 	TArray<FIoBuffer*> AvailableCompressionBuffers;
 	uint64 MaxCompressionBufferMemory = 0;
@@ -627,6 +638,11 @@ FIoStoreWriterContext::FProgress FIoStoreWriterContext::GetProgress() const
 	return Impl->GetProgress();
 }
 
+uint32 FIoStoreWriterContext::GetErrors() const
+{
+	return Impl->GetErrors();
+}
+
 class FIoStoreTocBuilder
 {
 public:
@@ -911,6 +927,17 @@ public:
 		{
 			TRACE_CPUPROFILER_EVENT_SCOPE(HashChunk);
 			const FIoBuffer* SourceBuffer = Entry->Request->GetSourceBuffer();
+			if (!SourceBuffer)
+			{
+				UE_LOG(LogIoStore, Error, TEXT("HashChunk: Failed to load Chunk %s from %s for file %s."),
+					*LexToString(Entry->ChunkId), Entry->Request->DebugNameOfRepository(), *Entry->Options.FileName);
+				Entry->ChunkHash.Reset();
+				Entry->bUseDDCForCompression = false;
+				WriterContext->HashedChunksCount.IncrementExchange();
+				WriterContext->ReportError();
+				return;
+			}
+
 			FIoHash ChunkHash = FIoHash::HashBuffer(SourceBuffer->Data(), SourceBuffer->DataSize());
 
 			if (!Entry->ChunkHash.IsZero() && Entry->ChunkHash != ChunkHash)
@@ -987,6 +1014,11 @@ public:
 				return SourceBuffer.DataSize();
 			}
 
+			virtual const TCHAR* DebugNameOfRepository() const override
+			{
+				return TEXT("FIoStoreWriter::Append");
+			}
+
 			FIoBuffer SourceBuffer;
 			uint64 OrderHint;
 		};
@@ -1745,6 +1777,18 @@ private:
 		}
 
 		const FIoBuffer* SourceBuffer = Entry->Request->GetSourceBuffer();
+		if (!SourceBuffer)
+		{
+			UE_LOG(LogIoStore, Error, TEXT("CompressChunk: Failed to load Chunk %s from %s for file %s. Runtime will receive 0 bytes when loading the file."),
+				*LexToString(Entry->ChunkId), Entry->Request->DebugNameOfRepository(), *Entry->Options.FileName);
+			WriterContext->ReportError();
+			Entry->bStoreCompressedDataInDDC = false;
+			Entry->UncompressedSize.Emplace(0);
+			Entry->NumChunkBlocks = 0;
+			Entry->FinishCompressionBarrier.Trigger();
+			return;
+		}
+
 		Entry->UncompressedSize.Emplace(SourceBuffer->DataSize());
 		Entry->NumChunkBlocks = CalculateNumChunkBlocks(Entry->UncompressedSize.GetValue());
 
diff --git a/Engine/Source/Developer/IoStoreUtilities/Private/IoStoreWriter.h b/Engine/Source/Developer/IoStoreUtilities/Private/IoStoreWriter.h
--- a/Engine/Source/Developer/IoStoreUtilities/Private/IoStoreWriter.h
+++ b/Engine/Source/Developer/IoStoreUtilities/Private/IoStoreWriter.h
@@ -95,6 +95,7 @@ public:
 	TSharedPtr<class IIoStoreWriter> CreateContainer(const TCHAR* InContainerPath, const FIoContainerSettings& InContainerSettings);
 	void Flush();
 	FProgress GetProgress() const;
+	uint32 GetErrors() const;
 
 private:
 	FIoStoreWriterContextImpl* Impl;
@@ -119,6 +120,9 @@ public:
 	virtual void FreeSourceBuffer() = 0;
 
 	virtual uint64 GetSourceBufferSizeEstimate() = 0;
+
+	// Return an identifier for the repository that is providing the source data, e.g. loosefile, zenserver.
+	virtual const TCHAR* DebugNameOfRepository() const = 0;
 };
 
 class IIoStoreWriterReferenceChunkDatabase
