Add extra logging when zenstorehttpclient fails to fetch a chunk from zenserver

#rnx

diff --git a/Engine/Source/Developer/IoStoreUtilities/Private/ZenStoreHttpClient.cpp b/Engine/Source/Developer/IoStoreUtilities/Private/ZenStoreHttpClient.cpp
--- a/Engine/Source/Developer/IoStoreUtilities/Private/ZenStoreHttpClient.cpp
+++ b/Engine/Source/Developer/IoStoreUtilities/Private/ZenStoreHttpClient.cpp
@@ -867,7 +867,7 @@ TIoStatusOr<FIoBuffer> FZenStoreHttpClient::ReadOpLogUri(FStringBuilderBase& Chu
 				uint64 BlockSize = 0;
 				ECompressedBufferCompressor Compressor;
 				ECompressedBufferCompressionLevel CompressionLevel;
-				if  (!Compressed.TryGetCompressParameters(Compressor, CompressionLevel, BlockSize))
+				if (!Compressed.TryGetCompressParameters(Compressor, CompressionLevel, BlockSize))
 				{
 					return FIoStatus(EIoErrorCode::CompressionError);
 				}
@@ -891,7 +891,20 @@ TIoStatusOr<FIoBuffer> FZenStoreHttpClient::ReadOpLogUri(FStringBuilderBase& Chu
 			return FIoBuffer(FIoBuffer::Clone, GetBuffer.GetData(), GetBuffer.Num());
 		}
 	}
-	return FIoStatus(EIoErrorCode::NotFound);
+	else
+	{
+		int ResponseCode = Request->GetResponseCode();
+		if (ResponseCode == 404)
+		{
+			UE_LOG(LogZenStore, Log, TEXT("Request '%s' failed with %s (%d). \"%s\""), *ChunkUri, *FString(UE::Zen::FZenHttpRequest::LexResponseCodeToString(ResponseCode)), ResponseCode, *Request->GetResponseAsString());
+			return FIoStatus(EIoErrorCode::NotFound);
+		}
+		else
+		{
+			UE_LOG(LogZenStore, Log, TEXT("Request '%s' failed. Error: %s. Status: %s (%d). \"%s\""), *ChunkUri, *FString(Request->GetError()), *FString(UE::Zen::FZenHttpRequest::LexResponseCodeToString(ResponseCode)), ResponseCode, *Request->GetResponseAsString());
+		}
+	}
+	return FIoStatus(EIoErrorCode::ReadError);
 }
 
 TFuture<TIoStatusOr<FCbObject>> FZenStoreHttpClient::GetProjectInfo()
