Cooker: ZenFileSystemManifest: Refactor FZenFileSystemManifest::Generator for maintainability: replace lambda functions with member functions and a new helper class. This will support robust implementation of upcoming changes to the set of files in the manifest when incrementally cooking DLC.
No behavior change expected.
#rnx

diff --git a/Engine/Source/Developer/IoStoreUtilities/Private/ZenFileSystemManifest.cpp b/Engine/Source/Developer/IoStoreUtilities/Private/ZenFileSystemManifest.cpp
--- a/Engine/Source/Developer/IoStoreUtilities/Private/ZenFileSystemManifest.cpp
+++ b/Engine/Source/Developer/IoStoreUtilities/Private/ZenFileSystemManifest.cpp
@@ -17,6 +17,183 @@
 
 DEFINE_LOG_CATEGORY_STATIC(LogZenFileSystemManifest, Display, All);
 
+namespace UE::ZenFileSystemManifest
+{
+
+/** Reports whether files should pass a filter based on extension include/exclude and directory exclude properties. */
+class FFileFilter
+{
+public:
+	/** Mark that directories with the given leafname are excluded. Returns *this for chained function calls. */
+	FFileFilter& ExcludeDirectoryLeafName(const TCHAR* LeafName);
+	/**
+	 * Mark that filenames with the given extension (which must be provided without a leading dot) are excluded.
+	 * Returns *this for chained function calls.
+	 */
+	FFileFilter& ExcludeExtension(const TCHAR* NoDotExtension);
+	/**
+	 * Mark that filenames with the given extension  (which must be provided without a leading dot) are included, if
+	 * not excluded. Overruled by ExcludeExtension.
+	 * If no extensions are explicitly included then all non-excluded extensions are included.
+	 * Returns *this for chained function calls.
+	 */
+	FFileFilter& IncludeExtension(const TCHAR* NoDotExtension);
+
+	/** Report whether a directory with the given LeafName is kept by the filter. */
+	bool IsDirectoryLeafNameKept(FStringView LeafName) const;
+	/** Report whether the given extension (which must be provided without a leading dot) is kept by the filter. */
+	bool IsFileExtensionKept(FStringView NoDotExtension);
+
+private:
+	TArray<FString> DirectoryExclusionFilter;
+	TArray<FString> ExtensionExclusionFilter;
+	TArray<FString> ExtensionInclusionFilter;
+};
+
+/**
+ * A visitor used by FManifestGenerator::AddFilesFromDirectory to receive the results from
+ * IPlatformFile::IterateDirectory.
+ */
+struct FAddFilesFromDirectoryVisitor
+{
+public:
+	FAddFilesFromDirectoryVisitor(FManifestGenerator& InGenerator, const FString& InClientRoot,
+		const FString& InLocalRoot, bool bInIncludeSubDirs, FFileFilter* InAdditionalFilter);
+
+	/**
+	 * Callback passed to IPlatformFile::IterateDirectory. Filters and adds files
+	 * and subdirectories. Kept files are added to Generator.Manifest. Kept subdirectories are added to
+	 * this->DirectoryVisitQueue. Return value is whether IterateDirectory should continue.
+	 */
+	bool VisitorFunction(const TCHAR* InnerFileNameOrDirectory, bool bIsDirectory);
+
+public:
+	FManifestGenerator& Generator;
+	TArray<FString> DirectoryVisitQueue;
+	const FString& ClientRoot;
+	const FString& LocalRoot;
+	FFileFilter* AdditionalFilter = nullptr;
+	bool bIncludeSubDirs = false;
+};
+
+/**
+ * Helper class for FManifestGenerator::Generate. Has functions that support collecting files from directories in the
+ * uncooked folders used by or the cooked folders created by the cook. Specifications of which files to collect from
+ * which kinds of folders are driven by the top-level code in Generate.
+ */
+class FManifestGenerator
+{
+public:
+	FManifestGenerator(FZenFileSystemManifest& InManifest);
+
+	/** FPaths::RootDir, the root of the UnrealEngine tree, parent directory of EngineDir. */
+	const FString& GetUnrealEngineRootDir() const;
+	/** FPaths::EngineDir, the root of assets and source code shared across all Unreal projects. */
+	const FString& GetEngineDir() const;
+	/**
+	 * FPaths::ProjectDir, the root of the project's project-specific assets and source code. Contains the .uproject
+	 * file. Not necessarily the parent directory of all of the project's assets; some assets can be present in plugin
+	 * folders outside of the ProjectDir.
+	 */
+	const FString& GetProjectDir() const;
+	/**
+	 * FPaths::ProjectName, the name of the project, which we expect to find as a direct subdirectory of the cooked
+	 * output's root directory.
+	 */
+	const FString& GetProjectName() const;
+	/** <CookDirectory>/Engine */
+	FString GetCookedEngineDir() const;
+	/** <CookDirectory>/<ProjectName> */
+	FString GetCookedProjectDir() const;
+	/** <CookDirectory>/<ProjectName>/Metadata */
+	FString GetCookedMetadataDir() const;
+	/** TargetPlatform's PlatformInfo's UBTPlatformString. */
+	const FString& GetUBTPlatformString() const;
+	/** Read/Write the BaseFilter the Generator uses for adding files from directories. */
+	FFileFilter& GetBaseFilter();
+
+	/**
+	 * Add the given file to the manifest, stored relative to the given PathMappingPair of ClientRoot,LocalRoot. A
+	 * PathMappingPair represents the two possible formats for the same root directory (the format for loading the path
+	 * on the UnrealGame client and the format for loading the path on the cooker's machine) and the relative path
+	 * to a file under the mapping pair is the same for each root. LocalRoot should be a parent directory of
+	 * LocalFilePath; if it is not, the stored path falls back to absolute path on the cooker's machine and will not
+	 * be loadable on the client.
+	 */
+	void AddSingleFile(const FString& ClientRoot, const FString& LocalRoot, FStringView LocalFilePath);
+	/** If file exists at LocalRoot\RelPathFromRoot, AddSingleFile on it. Return whether it was added. */
+	bool TryAddSingleRelpath(const FString& ClientRoot, const FString& LocalRoot, const FString& RelPathFromRoot);
+
+	/**
+	 * Add to the manifest all (optionally recursive) files under LocalAddRoot, if they pass the Generator's basefilter
+	 * and the given Additional Filter. ClientRoot,LocalRoot are a PathMappingPair, with LocalRoot being a parent of
+	 * LocalAddRoot, @see AddSingleFile. If LocalAddRoot is not provided it is set equal to LocalRoot.
+	 */
+	void AddFilesFromDirectory(const FString& ClientRoot, const FString& LocalRoot, bool bIncludeSubdirs,
+		FFileFilter* AdditionalFilter = nullptr);
+	void AddFilesFromDirectory(const FString& ClientRoot, const FString& LocalRoot, const FString& LocalAddRoot,
+		bool bIncludeSubdirs, FFileFilter* AdditionalFilter = nullptr);
+	/**
+	 * If the given path is a subpath of one of the known client roots (EngineDir, ProjectDir), call
+	 * AddFilesFromDirectory on it and return true, else return false.
+	 */
+	bool TryAddFilesFromFlexDirectory(const FString& FullOrProjectRelativePath, bool bIncludeSubDirs,
+		FFileFilter* AdditionalFilter = nullptr);
+
+	/**
+	 * A collector AddFiles function. It looks in all GetExtensionDirectories() folders for a direct subdirectory
+	 * with leafname equal to ExtensionSubDir, and calls AddFilesFromDirectory on each of those ExtensionSubDir
+	 * directories, with bRecursive==true and with the given filter.
+	 */
+	void AddFilesFromExtensionDirectories(const TCHAR* ExtensionSubDir, FFileFilter* AdditionalFilter = nullptr);
+
+	/**
+	 * Return all the plugins that could be used by the runtime the manifest will be used in - plugins are filtered by
+	 * the current project and platform.
+	 */
+	TArray<TSharedRef<IPlugin>> GetApplicablePlugins();
+	/**
+	 * Return the ClientRoot,LocalRoot PathMappingPair that is one of the manifest's recognized root directories
+	 * and is the parent directory of the plugin's uncooked directory.
+	 * These roots allow the construction of the relative path used to store the plugin's files on the ZenServer.
+	 */
+	void GetPluginClientAndLocalRoots(TSharedRef<IPlugin>& Plugin, FString& OutClientRoot, FString& OutLocalRoot);
+	/**
+	 * For a given plugin, return full paths to all of the (possibly-platform-specific) ExtensionBaseDirs that are
+	 * present in the plugin and are applicable to the current project and platform. These directories need the same
+	 * set of relative paths added to the manifest that are added for the plugin's base dir.
+	 */
+	TArray<FString> GetApplicablePluginExtensionBaseDirs(TSharedRef<IPlugin>& Plugin);
+
+	/** Return the subfolder used for localization assets based on the project's Internationalization settings. */
+	FString GetInternationalizationPresetPath();
+
+	/** Return bSSLCertificatesWillStage from network settings for the cook's targetplatform, false by default. */
+	bool IsSSLCertificatesWillStage();
+
+private:
+	void PopulatePlatformDirectoryNames();
+
+	static void GetExtensionDirs(TArray<FString>& OutExtensionDirs, const TCHAR* BaseDir, const TCHAR* SubDir,
+		const TArray<FString>& PlatformDirectoryNames);
+
+private:
+	FZenFileSystemManifest& Manifest;
+	IPlatformFile& PlatformFile;
+	FString UnrealEngineRootDir;
+	FString EngineDir;
+	FString ProjectDir;
+	FString ProjectName;
+	FFileFilter BaseFilter;
+	TArray<FString> PlatformDirectoryNames;
+	TSet<FString> PlatformDirectoryNameSet;
+	FString UBTPlatformString;
+
+	friend FAddFilesFromDirectoryVisitor;
+};
+
+} // namespace UE::ZenFileSystemManifest
+
 const FZenFileSystemManifestEntry FZenFileSystemManifest::InvalidEntry = FZenFileSystemManifestEntry();
 
 FZenFileSystemManifest::FZenFileSystemManifest(const ITargetPlatform& InTargetPlatform, FString InCookDirectory)
@@ -32,293 +209,46 @@ FZenFileSystemManifest::FZenFileSystemManifest(const ITargetPlatform& InTargetPl
 #endif
 }
 
-void FZenFileSystemManifest::GetExtensionDirs(TArray<FString>& OutExtensionDirs, const TCHAR* BaseDir, const TCHAR* SubDir, const TArray<FString>& PlatformDirectoryNames)
-{
-	auto AddIfDirectoryExists = [&OutExtensionDirs](FString&& Dir)
-	{
-		if (FPaths::DirectoryExists(Dir))
-		{
-			OutExtensionDirs.Emplace(MoveTemp(Dir));
-		}
-	};
-	
-	AddIfDirectoryExists(FPaths::Combine(BaseDir, SubDir));
-
-	FString PlatformExtensionBaseDir = FPaths::Combine(BaseDir, TEXT("Platforms"));
-	for (const FString& PlatformDirectoryName : PlatformDirectoryNames)
-	{
-		AddIfDirectoryExists(FPaths::Combine(PlatformExtensionBaseDir, PlatformDirectoryName, SubDir));
-	}
-
-	FString RestrictedBaseDir = FPaths::Combine(BaseDir, TEXT("Restricted"));
-	IFileManager::Get().IterateDirectory(*RestrictedBaseDir, [&OutExtensionDirs, SubDir, &PlatformDirectoryNames](const TCHAR* FilenameOrDirectory, bool bIsDirectory) -> bool
-		{
-			if (bIsDirectory)
-			{
-				GetExtensionDirs(OutExtensionDirs, FilenameOrDirectory, SubDir, PlatformDirectoryNames);
-			}
-			return true;
-		});
-}
-
 int32 FZenFileSystemManifest::Generate()
 {
-	TRACE_CPUPROFILER_EVENT_SCOPE(GenerateStorageServerFileSystemManifest);
+	using namespace UE::ZenFileSystemManifest;
 
-	class FFileFilter
-	{
-	public:
-		FFileFilter& ExcludeDirectory(const TCHAR* Name)
-		{
-			DirectoryExclusionFilter.Emplace(Name);
-			return *this;
-		}
-
-		FFileFilter& ExcludeExtension(const TCHAR* Extension)
-		{
-			ExtensionExclusionFilter.Emplace(Extension);
-			return *this;
-		}
-
-		FFileFilter& IncludeExtension(const TCHAR* Extension)
-		{
-			ExtensionInclusionFilter.Emplace(Extension);
-			return *this;
-		}
-
-		bool FilterDirectory(FStringView Name) const
-		{
-			for (const FString& ExcludedDirectory : DirectoryExclusionFilter)
-			{
-				if (Name == ExcludedDirectory)
-				{
-					return false;
-				}
-			}
-			return true;
-		}
-
-		bool FilterFile(FStringView Extension)
-		{
-			if (!ExtensionExclusionFilter.IsEmpty())
-			{
-				for (const FString& ExcludedExtension : ExtensionExclusionFilter)
-				{
-					if (Extension == ExcludedExtension)
-					{
-						return false;
-					}
-				}
-			}
-			if (!ExtensionInclusionFilter.IsEmpty())
-			{
-				for (const FString& IncludedExtension : ExtensionInclusionFilter)
-				{
-					if (Extension == IncludedExtension)
-					{
-						return true;
-					}
-				}
-				return false;
-			}
-			return true;
-		}
-
-	private:
-		TArray<FString> DirectoryExclusionFilter;
-		TArray<FString> ExtensionExclusionFilter;
-		TArray<FString> ExtensionInclusionFilter;
-	};
-
-	IPlatformFile& PlatformFile = FPlatformFileManager::Get().GetPlatformFile();
-
-	FString RootDir = FPaths::RootDir();
-	FString EngineDir = FPaths::EngineDir();
-	FPaths::NormalizeDirectoryName(EngineDir);
-	FString ProjectDir = FPaths::ProjectDir();
-	FPaths::NormalizeDirectoryName(ProjectDir);
-
-	FFileFilter BaseFilter = FFileFilter()
-		.ExcludeDirectory(TEXT("Binaries"))
-		.ExcludeDirectory(TEXT("Intermediate"))
-		.ExcludeDirectory(TEXT("Saved"))
-		.ExcludeDirectory(TEXT("Source"));
-	auto AddFilesFromDirectory =
-		[this, &PlatformFile, &RootDir, &BaseFilter]
-		(const FString& ClientDirectory, const FString& LocalDirectory, bool bIncludeSubdirs, FFileFilter* AdditionalFilter = nullptr)
-		{
-			//TRACE_BOOKMARK(TEXT("AddFilesFromDirectory: %s"), *ClientDirectory);
-			FString ServerRelativeDirectory = LocalDirectory;
-			FPaths::MakePathRelativeTo(ServerRelativeDirectory, *RootDir);
-			ServerRelativeDirectory = TEXT("/") + ServerRelativeDirectory;
-
-			TArray<FString> DirectoriesToVisit;
-			auto VisitorFunc =
-				[this, &DirectoriesToVisit, &RootDir, &ClientDirectory, &LocalDirectory, &ServerRelativeDirectory, bIncludeSubdirs, &BaseFilter, AdditionalFilter]
-				(const TCHAR* InnerFileNameOrDirectory, bool bIsDirectory)
-				{
-					if (bIsDirectory)
-					{
-						if (!bIncludeSubdirs)
-						{
-							return true;
-						}
-						FStringView DirectoryName = FPathViews::GetPathLeaf(InnerFileNameOrDirectory);
-						if (!BaseFilter.FilterDirectory(DirectoryName))
-						{
-							return true;
-						}
-						if (AdditionalFilter && !AdditionalFilter->FilterDirectory(DirectoryName))
-						{
-							return true;
-						}
-						DirectoriesToVisit.Add(InnerFileNameOrDirectory);
-						return true;
-					}
-					FStringView Extension = FPathViews::GetExtension(InnerFileNameOrDirectory);
-					if (!BaseFilter.FilterFile(Extension))
-					{
-						return true;
-					}
-					if (AdditionalFilter && !AdditionalFilter->FilterFile(Extension))
-					{
-						return true;
-					}
-					//TRACE_CPUPROFILER_EVENT_SCOPE(AddManifestEntry);
-					FStringView RelativePath = InnerFileNameOrDirectory;
-					RelativePath.RightChopInline(LocalDirectory.Len() + 1);
-					FString ClientPath = FPaths::Combine(ClientDirectory, RelativePath.GetData());
-					const FIoChunkId FileChunkId = CreateExternalFileChunkId(ClientPath);
-					AddManifestEntry(
-						FileChunkId,
-						FPaths::Combine(ServerRelativeDirectory, RelativePath.GetData()),
-						MoveTemp(ClientPath));
-					return true;
-				};
-
-			DirectoriesToVisit.Push(LocalDirectory);
-			while (!DirectoriesToVisit.IsEmpty())
-			{
-				PlatformFile.IterateDirectory(*DirectoriesToVisit.Pop(EAllowShrinking::No), VisitorFunc);
-			}
-		};
+	TRACE_CPUPROFILER_EVENT_SCOPE(GenerateStorageServerFileSystemManifest);
 
 	const UProjectPackagingSettings* const PackagingSettings = GetDefault<UProjectPackagingSettings>();
+	const int32 PreviousEntryCount = NumEntries();
 
-	TArray<FString> PlatformDirectoryNames;
-	const FDataDrivenPlatformInfo& PlatformInfo = FDataDrivenPlatformInfoRegistry::GetPlatformInfo(TargetPlatform.IniPlatformName());
-	PlatformDirectoryNames.Reserve(PlatformInfo.IniParentChain.Num() + PlatformInfo.AdditionalRestrictedFolders.Num() + 1);
-	PlatformDirectoryNames.Add(TargetPlatform.IniPlatformName());
-	for (const FString& PlatformName : PlatformInfo.AdditionalRestrictedFolders)
-	{
-		PlatformDirectoryNames.AddUnique(PlatformName);
-	}
-	for (const FString& PlatformName : PlatformInfo.IniParentChain)
-	{
-		PlatformDirectoryNames.AddUnique(PlatformName);
-	}
-	auto AddFilesFromExtensionDirectories = [&AddFilesFromDirectory, &EngineDir, &ProjectDir, &PlatformDirectoryNames](const TCHAR* ExtensionSubDir, FFileFilter* AdditionalFilter = nullptr)
-		{
-			TArray<FString> ExtensionDirs;
-			GetExtensionDirs(ExtensionDirs, *EngineDir, ExtensionSubDir, PlatformDirectoryNames);
-			for (const FString& Dir : ExtensionDirs)
-			{
-				AddFilesFromDirectory(Dir.Replace(*EngineDir, TEXT("/{engine}")), Dir, true, AdditionalFilter);
-			}
-			ExtensionDirs.Reset();
-			GetExtensionDirs(ExtensionDirs, *ProjectDir, ExtensionSubDir, PlatformDirectoryNames);
-			for (const FString& Dir : ExtensionDirs)
-			{
-				AddFilesFromDirectory(Dir.Replace(*ProjectDir, TEXT("/{project}")), Dir, true, AdditionalFilter);
-			}
-		};
+	FManifestGenerator Generator(*this);
 
-	const int32 PreviousEntryCount = NumEntries();
+	Generator.GetBaseFilter()
+		.ExcludeDirectoryLeafName(TEXT("Binaries"))
+		.ExcludeDirectoryLeafName(TEXT("Intermediate"))
+		.ExcludeDirectoryLeafName(TEXT("Saved"))
+		.ExcludeDirectoryLeafName(TEXT("Source"));
 
 	FFileFilter CookedFilter = FFileFilter()
-		.ExcludeDirectory(TEXT("Metadata"))
+		.ExcludeDirectoryLeafName(TEXT("Metadata"))
 		.ExcludeExtension(TEXT("uasset"))
 		.ExcludeExtension(TEXT("ubulk"))
 		.ExcludeExtension(TEXT("uexp"))
 		.ExcludeExtension(TEXT("umap"))
 		.ExcludeExtension(TEXT("uregs"));
-	AddFilesFromDirectory(TEXT("/{engine}"), FPaths::Combine(CookDirectory, TEXT("Engine")), true, &CookedFilter);
-	AddFilesFromDirectory(TEXT("/{project}"), FPaths::Combine(CookDirectory, FApp::GetProjectName()), true, &CookedFilter);
+	Generator.AddFilesFromDirectory(TEXT("/{engine}"), Generator.GetCookedEngineDir(), true, &CookedFilter);
+	Generator.AddFilesFromDirectory(TEXT("/{project}"), Generator.GetCookedProjectDir(), true, &CookedFilter);
 
 	FFileFilter CookedMetadataFilter = FFileFilter()
-		.ExcludeDirectory(TEXT("ShaderLibrarySource"))
+		.ExcludeDirectoryLeafName(TEXT("ShaderLibrarySource"))
 		.ExcludeExtension(TEXT("manifest"));
-	AddFilesFromDirectory(TEXT("/{project}/Metadata"), FPaths::Combine(CookDirectory, FApp::GetProjectName(), "Metadata"), true, &CookedMetadataFilter);
+	Generator.AddFilesFromDirectory(TEXT("/{project}/Metadata"), Generator.GetCookedMetadataDir(), true,
+		&CookedMetadataFilter);
 
 	FFileFilter ProjectSourceFilter = FFileFilter()
 		.IncludeExtension(TEXT("uproject"));
-	AddFilesFromDirectory(TEXT("/{project}"), ProjectDir, false, &ProjectSourceFilter);
+	Generator.AddFilesFromDirectory(TEXT("/{project}"), Generator.GetProjectDir(), false, &ProjectSourceFilter);
 	
 	FFileFilter ConfigFilter = FFileFilter()
 		.IncludeExtension(TEXT("ini"));
-	AddFilesFromExtensionDirectories(TEXT("Config"), &ConfigFilter);
-
-	auto AddFromPluginPath = [this](const FString& EngineDir, const FString& ClientDirectory, const FString& SourcePath, const TCHAR* Path) {
-		//TRACE_CPUPROFILER_EVENT_SCOPE(AddManifestEntry);
-
-		FString ServerRelativePath = Path;
-		FPaths::MakePathRelativeTo(ServerRelativePath, *EngineDir);
-		ServerRelativePath = TEXT("/") + ServerRelativePath;
-
-		FString ClientPath = Path;
-		FPaths::MakePathRelativeTo(ClientPath, *(SourcePath / TEXT("")));
-		ClientPath = ClientDirectory / ClientPath;
-
-		const FIoChunkId FileChunkId = CreateExternalFileChunkId(ClientPath);
-		AddManifestEntry(
-			FileChunkId,
-			ServerRelativePath,
-			MoveTemp(ClientPath));
-	};
-
-	auto AddFromPluginDir = [this, &PlatformFile, &BaseFilter, &AddFromPluginPath](const FString& EngineDir, const FString& ClientDirectory, const FString& SourcePath, const FString& DirectoryPath, bool bIncludeSubdirs, FFileFilter* AdditionalFilter = nullptr) {
-		TArray<FString> DirectoriesToVisit;
-		auto VisitorFunc =
-			[this, &AddFromPluginPath, &EngineDir, &ClientDirectory, &SourcePath, &DirectoriesToVisit, &BaseFilter, bIncludeSubdirs, AdditionalFilter]//, &RootDir, &ClientDirectory, &LocalDirectory, &ServerRelativeDirectory, &BaseFilter, AdditionalFilter]
-			(const TCHAR* InnerFileNameOrDirectory, bool bIsDirectory)
-		{
-			if (bIsDirectory)
-			{
-				if (!bIncludeSubdirs)
-				{
-					return true;
-				}
-				FStringView DirectoryName = FPathViews::GetPathLeaf(InnerFileNameOrDirectory);
-				if (!BaseFilter.FilterDirectory(DirectoryName))
-				{
-					return true;
-				}
-				if (AdditionalFilter && !AdditionalFilter->FilterDirectory(DirectoryName))
-				{
-					return true;
-				}
-				DirectoriesToVisit.Add(InnerFileNameOrDirectory);
-				return true;
-			}
-			FStringView Extension = FPathViews::GetExtension(InnerFileNameOrDirectory);
-			if (!BaseFilter.FilterFile(Extension))
-			{
-				return true;
-			}
-			if (AdditionalFilter && !AdditionalFilter->FilterFile(Extension))
-			{
-				return true;
-			}
-			AddFromPluginPath(EngineDir, ClientDirectory, SourcePath, InnerFileNameOrDirectory);
-			return true;
-		};
-
-		DirectoriesToVisit.Push(DirectoryPath);
-		while (!DirectoriesToVisit.IsEmpty())
-		{
-			PlatformFile.IterateDirectory(*DirectoriesToVisit.Pop(EAllowShrinking::No), VisitorFunc);
-		}
-	};
+	Generator.AddFilesFromExtensionDirectories(TEXT("Config"), &ConfigFilter);
 
 	FFileFilter LocalizationFilter = FFileFilter()
 		.IncludeExtension(TEXT("locmeta"))
@@ -327,152 +257,61 @@ int32 FZenFileSystemManifest::Generate()
 	FFileFilter PluginFilter = FFileFilter()
 		.IncludeExtension(TEXT("uplugin"));
 
-	const bool FilterDisabledPlugins = false;
-	FString PluginTargetPlatformString = PlatformInfo.UBTPlatformString;
-	TSet<FString> PlatformDirectoryNameSet;
-	PlatformDirectoryNameSet.Append(PlatformDirectoryNames);
-	IPluginManager& PluginManager = IPluginManager::Get();
-	TArray<TSharedRef<IPlugin>> DiscoveredPlugins = PluginManager.GetDiscoveredPlugins();
-	for (TSharedRef<IPlugin>& Plugin : DiscoveredPlugins)
+	for (TSharedRef<IPlugin>& Plugin : Generator.GetApplicablePlugins())
 	{
-		FString ProjectName = Plugin->GetName();
-		if (FilterDisabledPlugins && !Plugin->IsEnabled())
-		{
-			UE_LOG(LogZenFileSystemManifest, Verbose, TEXT("Plugin '%s' disabled, skipping"), *ProjectName);
-			continue;
-		}
-		const FPluginDescriptor& Descriptor = Plugin->GetDescriptor();
-		if (!Descriptor.SupportsTargetPlatform(PluginTargetPlatformString))
-		{
-			UE_LOG(LogZenFileSystemManifest, Verbose, TEXT("Plugin '%s' not supported on platform '%s', skipping"), *ProjectName, *TargetPlatform.PlatformName());
-			for (const auto& SupportedTargetPlatform : Descriptor.SupportedTargetPlatforms)
-			{
-				UE_LOG(LogZenFileSystemManifest, Verbose, TEXT("       '%s' supports platform '%s'"), *ProjectName, *SupportedTargetPlatform);
-			}
-			continue;
-		}
+		const FString& PluginName = Plugin->GetName();
 		FString BaseDir = Plugin->GetBaseDir();
-		FString ProjectFile = Plugin->GetDescriptorFileName();
+		const FString& UPluginFile = Plugin->GetDescriptorFileName();
 		FString ContentDir = Plugin->GetContentDir();
 		FString LocalizationDir = ContentDir / TEXT("Localization");
 		FString ConfigDir = BaseDir / TEXT("Config");
-		UE_LOG(LogZenFileSystemManifest, Verbose, TEXT("Plugin '%s': BaseDir: '%s'"), *ProjectName, *BaseDir);
+		UE_LOG(LogZenFileSystemManifest, Verbose, TEXT("Plugin '%s': BaseDir: '%s'"), *PluginName, *BaseDir);
 
-		FString ClientDirectory;
-		FString SourcePath;
-		switch (Plugin->GetLoadedFrom())
-		{
-		case EPluginLoadedFrom::Engine:
-			ClientDirectory = TEXT("/{engine}");
-			SourcePath = EngineDir;
-			break;
-		case EPluginLoadedFrom::Project:
-			ClientDirectory = TEXT("/{project}");
-			SourcePath = ProjectDir;
-			break;
-		}
-		AddFromPluginPath(EngineDir, ClientDirectory, SourcePath, *ProjectFile);
-		AddFromPluginDir(EngineDir, ClientDirectory, SourcePath, LocalizationDir, true, &LocalizationFilter);
-		AddFromPluginDir(EngineDir, ClientDirectory, SourcePath, ConfigDir, true, &ConfigFilter);
+		FString ClientRoot;
+		FString LocalRoot;
+		Generator.GetPluginClientAndLocalRoots(Plugin, ClientRoot, LocalRoot);
+
+		Generator.AddSingleFile(ClientRoot, LocalRoot, UPluginFile);
+		Generator.AddFilesFromDirectory(ClientRoot, LocalRoot, LocalizationDir, true, &LocalizationFilter);
+		Generator.AddFilesFromDirectory(ClientRoot, LocalRoot, ConfigDir, true, &ConfigFilter);
 
 		// Next add any valid plugin extension directories of this plugin.
-		TArray<FString> ExtensionBaseDirs = Plugin->GetExtensionBaseDirs();
-		for (const FString& ExtensionBaseDir : ExtensionBaseDirs)
+		for (const FString& ExtensionBaseDir : Generator.GetApplicablePluginExtensionBaseDirs(Plugin))
 		{
-			// Scan the extension path for "Platforms/X" and include this extension if it is not platform specific at all,
-			// or if X is found and it is a valid target platform
-			bool bFoundPlatformsComponent = false;
-			bool bDone = false;
-			bool bIncludeExtension = true;
-			FPathViews::IterateComponents(
-				ExtensionBaseDir,
-				[&bFoundPlatformsComponent, &bDone, &bIncludeExtension, &PlatformDirectoryNameSet](FStringView CurrentPathComponent)
-				{
-					if (!bFoundPlatformsComponent)
-					{
-						if (CurrentPathComponent == TEXTVIEW("Platforms"))
-						{
-							bFoundPlatformsComponent = true;
-						}
-					}
-					else if (!bDone)
-					{
-						const bool bIsValidPlatform = PlatformDirectoryNameSet.Contains(FString(CurrentPathComponent));
-						bIncludeExtension = bIsValidPlatform;
-						bDone = true;
-					}
-					else
-					{
-						// Do nothing.
-					}
-				}
-			);
+			FString ExtensionLocalizationDir = ExtensionBaseDir / TEXT("Content") / TEXT("Localization");
+			FString ExtensionConfigDir = ExtensionBaseDir / TEXT("Config");
+			UE_LOG(LogZenFileSystemManifest, Verbose, TEXT("Plugin '%s': ExtensionBaseDir: '%s'"), *PluginName, *ExtensionBaseDir);
 
-			if (bIncludeExtension)
-			{
-				FString ExtensionLocalizationDir = ExtensionBaseDir / TEXT("Content") / TEXT("Localization");
-				FString ExtensionConfigDir = ExtensionBaseDir / TEXT("Config");
-				UE_LOG(LogZenFileSystemManifest, Verbose, TEXT("Plugin '%s': ExtensionBaseDir: '%s'"), *ProjectName, *ExtensionBaseDir);
-
-				AddFromPluginDir(EngineDir, ClientDirectory, SourcePath, ExtensionBaseDir, false, &PluginFilter);
-				AddFromPluginDir(EngineDir, ClientDirectory, SourcePath, ExtensionLocalizationDir, true, &LocalizationFilter);
-				AddFromPluginDir(EngineDir, ClientDirectory, SourcePath, ExtensionConfigDir, true, &ConfigFilter);
-			}
+			Generator.AddFilesFromDirectory(ClientRoot, LocalRoot, ExtensionBaseDir, false, &PluginFilter);
+			Generator.AddFilesFromDirectory(ClientRoot, LocalRoot, ExtensionLocalizationDir, true, &LocalizationFilter);
+			Generator.AddFilesFromDirectory(ClientRoot, LocalRoot, ExtensionConfigDir, true, &ConfigFilter);
 		}
 	}
 
-	FString InternationalizationPresetAsString = UEnum::GetValueAsString(PackagingSettings->InternationalizationPreset);
-	const TCHAR* InternationalizationPresetPath = FCString::Strrchr(*InternationalizationPresetAsString, ':');
-	if (InternationalizationPresetPath)
-	{
-		++InternationalizationPresetPath;
-	}
-	else
-	{
-		UE_LOG(LogZenFileSystemManifest, Warning, TEXT("Failed reading internationalization preset setting, defaulting to English"));
-		InternationalizationPresetPath = TEXT("English");
-	}
+	FString InternationalizationPresetPath = Generator.GetInternationalizationPresetPath();
 	const TCHAR* ICUDataVersion = TEXT("icudt64l"); // TODO: Could this go into datadriven platform info? But it's basically always this.
-	AddFilesFromDirectory(*FPaths::Combine(TEXT("/{engine}"), TEXT("Content"), TEXT("Internationalization"), ICUDataVersion), FPaths::Combine(EngineDir, TEXT("Content"), TEXT("Internationalization"), InternationalizationPresetPath, ICUDataVersion), true);
+	Generator.AddFilesFromDirectory(
+		FPaths::Combine(TEXT("/{engine}"),
+			TEXT("Content"), TEXT("Internationalization"), ICUDataVersion),
+		FPaths::Combine(Generator.GetEngineDir(),
+			TEXT("Content"), TEXT("Internationalization"), InternationalizationPresetPath, ICUDataVersion),
+		true /* bRecursive */);
 	
-	AddFilesFromExtensionDirectories(TEXT("Content/Localization"), &LocalizationFilter);
+	Generator.AddFilesFromExtensionDirectories(TEXT("Content/Localization"), &LocalizationFilter);
 
-	bool bSSLCertificatesWillStage = false;
-	FConfigCacheIni* TargetPlatformConfig = TargetPlatform.GetConfigSystem();
-	if (TargetPlatformConfig)
-	{
-		GConfig->GetBool(TEXT("/Script/Engine.NetworkSettings"), TEXT("n.VerifyPeer"), bSSLCertificatesWillStage, GEngineIni);
-	}
-	if (bSSLCertificatesWillStage)
+	if (Generator.IsSSLCertificatesWillStage())
 	{
-		FString ProjectCertFile = FPaths::Combine(ProjectDir, TEXT("Content"), TEXT("Certificates"), TEXT("cacert.pem"));
-		if (FPaths::FileExists(ProjectCertFile))
+		if (!Generator.TryAddSingleRelpath(TEXT("/{project}"), Generator.GetProjectDir(),
+			FPaths::Combine(TEXT("Content"), TEXT("Certificates"), TEXT("cacert.pem"))))
 		{
-			const TCHAR* ClientProjectCertFile = TEXT("/{project}/Content/Certificates/cacert.pem");
-			const FIoChunkId FileChunkId = CreateExternalFileChunkId(ClientProjectCertFile);
-			FPaths::MakePathRelativeTo(ProjectCertFile, *FPaths::RootDir());
-			AddManifestEntry(
-				FileChunkId,
-				ProjectCertFile,
-				ClientProjectCertFile);
-		}
-		else
-		{
-			FString EngineCertFile = FPaths::Combine(EngineDir, TEXT("Content"), TEXT("Certificates"), TEXT("ThirdParty"), TEXT("cacert.pem"));
-			if (FPaths::FileExists(EngineCertFile))
-			{
-				const TCHAR* ClientEngineCertFile = TEXT("/{engine}/Content/Certificates/ThirdParty/cacert.pem");
-				const FIoChunkId FileChunkId = CreateExternalFileChunkId(ClientEngineCertFile);
-				FPaths::MakePathRelativeTo(EngineCertFile, *FPaths::RootDir());
-				AddManifestEntry(
-					FileChunkId,
-					EngineCertFile,
-					ClientEngineCertFile);
-			}
+			Generator.TryAddSingleRelpath(TEXT("/{engine}"), Generator.GetEngineDir(),
+				FPaths::Combine(TEXT("Content"), TEXT("Certificates"), TEXT("ThirdParty"), TEXT("cacert.pem")));
 		}
+
 		FFileFilter CertificateFilter = FFileFilter()
 			.IncludeExtension(TEXT("pem"));
-		AddFilesFromDirectory(TEXT("/{project}/Certificates"), FPaths::Combine(ProjectDir, TEXT("Certificates")), true, &CertificateFilter);
+		Generator.AddFilesFromDirectory(TEXT("/{project}/Certificates"),
+			FPaths::Combine(Generator.GetProjectDir(), TEXT("Certificates")), true, &CertificateFilter);
 	}
 
 	FFileFilter ContentFilter = FFileFilter()
@@ -480,52 +319,45 @@ int32 FZenFileSystemManifest::Generate()
 		.ExcludeExtension(TEXT("ubulk"))
 		.ExcludeExtension(TEXT("uexp"))
 		.ExcludeExtension(TEXT("umap"));
-	AddFilesFromDirectory(TEXT("/{engine}/Content/Slate"), FPaths::Combine(EngineDir, TEXT("Content"), TEXT("Slate")), true, &ContentFilter);
-	AddFilesFromDirectory(TEXT("/{project}/Content/Slate"), FPaths::Combine(ProjectDir, TEXT("Content"), TEXT("Slate")), true, &ContentFilter);
-	AddFilesFromDirectory(TEXT("/{engine}/Content/Movies"), FPaths::Combine(EngineDir, TEXT("Content"), TEXT("Movies")), true, &ContentFilter);
-	AddFilesFromDirectory(TEXT("/{project}/Content/Movies"), FPaths::Combine(ProjectDir, TEXT("Content"), TEXT("Movies")), true, &ContentFilter);
+	Generator.AddFilesFromDirectory(TEXT("/{engine}/Content/Slate"),
+		FPaths::Combine(Generator.GetEngineDir(), TEXT("Content"), TEXT("Slate")),
+		true, &ContentFilter);
+	Generator.AddFilesFromDirectory(TEXT("/{project}/Content/Slate"),
+		FPaths::Combine(Generator.GetProjectDir(), TEXT("Content"), TEXT("Slate")),
+		true, &ContentFilter);
+	Generator.AddFilesFromDirectory(TEXT("/{engine}/Content/Movies"),
+		FPaths::Combine(Generator.GetEngineDir(), TEXT("Content"), TEXT("Movies")),
+		true, &ContentFilter);
+	Generator.AddFilesFromDirectory(TEXT("/{project}/Content/Movies"),
+		FPaths::Combine(Generator.GetProjectDir(), TEXT("Content"), TEXT("Movies")),
+		true, &ContentFilter);
 	
 	FFileFilter OoodleDictionaryFilter = FFileFilter()
 		.IncludeExtension(TEXT("udic"));
-	AddFilesFromDirectory(TEXT("/{project}/Content/Oodle"), FPaths::Combine(ProjectDir, TEXT("Content"), TEXT("Oodle")), false, &OoodleDictionaryFilter);
+	Generator.AddFilesFromDirectory(TEXT("/{project}/Content/Oodle"),
+		FPaths::Combine(Generator.GetProjectDir(), TEXT("Content"), TEXT("Oodle")),
+		false, &OoodleDictionaryFilter);
 
 	FFileFilter ShaderCacheFilter = FFileFilter()
 		.IncludeExtension(TEXT("ushadercache"))
 		.IncludeExtension(TEXT("upipelinecache"));
-	AddFilesFromDirectory(TEXT("/{project}/Content"), FPaths::Combine(ProjectDir, TEXT("Content")), false, &ShaderCacheFilter);
-	AddFilesFromDirectory(FPaths::Combine(TEXT("/{project}"), TEXT("Content"), TEXT("PipelineCaches"), TargetPlatform.IniPlatformName()), FPaths::Combine(ProjectDir, TEXT("Content"), TEXT("PipelineCaches"), TargetPlatform.IniPlatformName()), false, &ShaderCacheFilter);
+	Generator.AddFilesFromDirectory(TEXT("/{project}/Content"),
+		FPaths::Combine(Generator.GetProjectDir(), TEXT("Content")),
+		false, &ShaderCacheFilter);
+	Generator.AddFilesFromDirectory(FPaths::Combine(TEXT("/{project}"), TEXT("Content"), TEXT("PipelineCaches"), TargetPlatform.IniPlatformName()),
+		FPaths::Combine(Generator.GetProjectDir(), TEXT("Content"), TEXT("PipelineCaches"), TargetPlatform.IniPlatformName()),
+		false, &ShaderCacheFilter);
 
-	auto AddAdditionalFilesFromConfig = [&AddFilesFromDirectory , &ContentFilter, &ProjectDir, &EngineDir](const FString& RelativeDirToStage)
-	{
-		FString AbsoluteDirToStage = FPaths::ConvertRelativePathToFull(FPaths::Combine(ProjectDir, TEXT("Content"), RelativeDirToStage));
-		FPaths::NormalizeDirectoryName(AbsoluteDirToStage);
-		FString AbsoluteEngineDir = FPaths::ConvertRelativePathToFull(EngineDir);
-		FString AbsoluteProjectDir = FPaths::ConvertRelativePathToFull(ProjectDir);
-		FStringView RelativeToKnownRootView;
-		if (FPathViews::TryMakeChildPathRelativeTo(AbsoluteDirToStage, AbsoluteProjectDir, RelativeToKnownRootView))
-		{
-			AddFilesFromDirectory(FPaths::Combine(TEXT("/{project}"), RelativeToKnownRootView.GetData()), FPaths::Combine(ProjectDir, RelativeToKnownRootView.GetData()), true, &ContentFilter);
-		}
-		else if (FPathViews::TryMakeChildPathRelativeTo(AbsoluteDirToStage, AbsoluteEngineDir, RelativeToKnownRootView))
-		{
-			AddFilesFromDirectory(FPaths::Combine(TEXT("/{engine}"), RelativeToKnownRootView.GetData()), FPaths::Combine(EngineDir, RelativeToKnownRootView.GetData()), true, &ContentFilter);
-		}
-		else
-		{
-			UE_LOG(LogZenFileSystemManifest, Warning, TEXT("Ignoring additional folder to stage that is not relative to the engine or project directory: %s"), *RelativeDirToStage);
-		}
-	};
 	for (const FDirectoryPath& AdditionalFolderToStage : PackagingSettings->DirectoriesToAlwaysStageAsUFS)
 	{
-		AddAdditionalFilesFromConfig(AdditionalFolderToStage.Path);
+		Generator.TryAddFilesFromFlexDirectory(AdditionalFolderToStage.Path, true, &ContentFilter);
 	}
 	for (const FDirectoryPath& AdditionalFolderToStage : PackagingSettings->DirectoriesToAlwaysStageAsUFSServer)
 	{
-		AddAdditionalFilesFromConfig(AdditionalFolderToStage.Path);
+		Generator.TryAddFilesFromFlexDirectory(AdditionalFolderToStage.Path, true, &ContentFilter);
 	}
 
 	const int32 CurrentEntryCount = NumEntries();
-	
 	return CurrentEntryCount - PreviousEntryCount;
 }
 
@@ -643,3 +475,456 @@ bool FZenFileSystemManifest::Save(const TCHAR* Filename)
 
 	return FFileHelper::SaveStringArrayToFile(CsvLines, Filename);
 }
+
+
+namespace UE::ZenFileSystemManifest
+{
+
+FManifestGenerator::FManifestGenerator(FZenFileSystemManifest& InManifest)
+	: Manifest(InManifest)
+	, PlatformFile(FPlatformFileManager::Get().GetPlatformFile())
+{
+	UnrealEngineRootDir = FPaths::RootDir();
+	EngineDir = FPaths::EngineDir();
+	FPaths::NormalizeDirectoryName(EngineDir);
+	ProjectDir = FPaths::ProjectDir();
+	FPaths::NormalizeDirectoryName(ProjectDir);
+	ProjectName = FApp::GetProjectName();
+
+	PopulatePlatformDirectoryNames();
+}
+
+const FString& FManifestGenerator::GetUnrealEngineRootDir() const
+{
+	return UnrealEngineRootDir;
+}
+
+const FString& FManifestGenerator::GetEngineDir() const
+{
+	return EngineDir;
+}
+
+const FString& FManifestGenerator::GetProjectDir() const
+{
+	return ProjectDir;
+}
+
+const FString& FManifestGenerator::GetProjectName() const
+{
+	return ProjectName;
+}
+
+FString FManifestGenerator::GetCookedEngineDir() const
+{
+	return FPaths::Combine(Manifest.CookDirectory, TEXT("Engine"));
+}
+
+FString FManifestGenerator::GetCookedProjectDir() const
+{
+	return FPaths::Combine(Manifest.CookDirectory, GetProjectName());
+}
+
+FString FManifestGenerator::GetCookedMetadataDir() const
+{
+	return FPaths::Combine(Manifest.CookDirectory, GetProjectName(), TEXT("Metadata"));
+}
+
+const FString& FManifestGenerator::GetUBTPlatformString() const
+{
+	return UBTPlatformString;
+}
+
+FFileFilter& FManifestGenerator::GetBaseFilter()
+{
+	return BaseFilter;
+}
+
+void FManifestGenerator::AddSingleFile(const FString& ClientRoot, const FString& LocalRoot, FStringView LocalFilePath)
+{
+	//TRACE_CPUPROFILER_EVENT_SCOPE(AddManifestEntry);
+	FString ClientPath(LocalFilePath);
+	// MakePathRelative calls GetPath (aka the parent path) on InRelativeTo argument before using it to look for the
+	// common prefix. Since we want our root paths to be used directly, rather than their parent path, add a
+	// terminating / to our root paths when passing them into MakePathRelativeTo; a terminating slash causes GetPath
+	// to return the path before the terminating slash.
+	FPaths::MakePathRelativeTo(ClientPath, *(LocalRoot / TEXT("")));
+	ClientPath = FPaths::Combine(ClientRoot, ClientPath);
+
+	FString ServerRelativePath(LocalFilePath);
+	FPaths::MakePathRelativeTo(ServerRelativePath, *(UnrealEngineRootDir / TEXT("")));
+	ServerRelativePath = FString(TEXT("/")) + ServerRelativePath;
+
+	const FIoChunkId FileChunkId = CreateExternalFileChunkId(ClientPath);
+	Manifest.AddManifestEntry(
+		FileChunkId,
+		MoveTemp(ServerRelativePath),
+		MoveTemp(ClientPath));
+}
+
+bool FManifestGenerator::TryAddSingleRelpath(const FString& ClientRoot, const FString& LocalRoot, const FString& RelPathFromRoot)
+{
+	FString LocalFilePath = FPaths::Combine(LocalRoot, RelPathFromRoot);
+	if (!FPaths::FileExists(LocalFilePath))
+	{
+		return false;
+	}
+	AddSingleFile(ClientRoot, LocalRoot, LocalFilePath);
+	return true;
+}
+
+void FManifestGenerator::AddFilesFromDirectory(const FString& ClientRoot, const FString& LocalRoot,
+	bool bIncludeSubDirs, FFileFilter* AdditionalFilter)
+{
+	AddFilesFromDirectory(ClientRoot, LocalRoot, LocalRoot, bIncludeSubDirs, AdditionalFilter);
+}
+
+void FManifestGenerator::AddFilesFromDirectory(const FString& ClientRoot, const FString& LocalRoot,
+	const FString& LocalAddRoot, bool bIncludeSubDirs, FFileFilter* AdditionalFilter)
+{
+	FAddFilesFromDirectoryVisitor Visitor(*this, ClientRoot, LocalRoot, bIncludeSubDirs, AdditionalFilter);
+
+	Visitor.DirectoryVisitQueue.Push(LocalAddRoot);
+	while (!Visitor.DirectoryVisitQueue.IsEmpty())
+	{
+		FString DirectoryToVisit = Visitor.DirectoryVisitQueue.Pop(EAllowShrinking::No);
+		PlatformFile.IterateDirectory(*DirectoryToVisit,
+			[&Visitor](const TCHAR* InnerFileNameOrDirectory, bool bIsDirectory) -> bool
+			{
+				return Visitor.VisitorFunction(InnerFileNameOrDirectory, bIsDirectory);
+			});
+	}
+};
+
+bool FManifestGenerator::TryAddFilesFromFlexDirectory(const FString& FullOrProjectRelativePath,
+	bool bRecursive, FFileFilter* AdditionalFilter)
+{
+	FString AbsoluteDirToStage = FPaths::ConvertRelativePathToFull(FPaths::Combine(ProjectDir, TEXT("Content"), FullOrProjectRelativePath));
+	FPaths::NormalizeDirectoryName(AbsoluteDirToStage);
+	FString AbsoluteEngineDir = FPaths::ConvertRelativePathToFull(EngineDir);
+	FString AbsoluteProjectDir = FPaths::ConvertRelativePathToFull(ProjectDir);
+	FStringView RelativeToKnownRootView;
+	if (FPathViews::TryMakeChildPathRelativeTo(AbsoluteDirToStage, AbsoluteProjectDir, RelativeToKnownRootView))
+	{
+		AddFilesFromDirectory(FPaths::Combine(TEXT("/{project}"), RelativeToKnownRootView.GetData()),
+			FPaths::Combine(ProjectDir, RelativeToKnownRootView.GetData()),
+			bRecursive, AdditionalFilter);
+		return true;
+	}
+	else if (FPathViews::TryMakeChildPathRelativeTo(AbsoluteDirToStage, AbsoluteEngineDir, RelativeToKnownRootView))
+	{
+		AddFilesFromDirectory(FPaths::Combine(TEXT("/{engine}"), RelativeToKnownRootView.GetData()),
+			FPaths::Combine(EngineDir, RelativeToKnownRootView.GetData()),
+			bRecursive, AdditionalFilter);
+		return true;
+	}
+	else
+	{
+		UE_LOG(LogZenFileSystemManifest, Warning,
+			TEXT("Ignoring additional folder to stage that is not relative to the engine or project directory: %s"),
+			*FullOrProjectRelativePath);
+		return false;
+	}
+}
+
+void FManifestGenerator::AddFilesFromExtensionDirectories(const TCHAR* ExtensionSubDir, FFileFilter* AdditionalFilter)
+{
+	TArray<FString> ExtensionDirs;
+	GetExtensionDirs(ExtensionDirs, *EngineDir, ExtensionSubDir, PlatformDirectoryNames);
+	for (const FString& LocalDir : ExtensionDirs)
+	{
+		if (!LocalDir.StartsWith(EngineDir))
+		{
+			UE_LOG(LogZenFileSystemManifest, Warning,
+				TEXT("GetExtensionDirs on root directory '%s' unexpectedly returned folder '%s' that is not a subdirectory of the root directory."),
+				*EngineDir, *LocalDir);
+			continue;
+		}
+		FString ClientDir = FPaths::Combine(TEXT("/{engine}"), LocalDir.RightChop(EngineDir.Len()));
+		AddFilesFromDirectory(ClientDir, LocalDir, true, AdditionalFilter);
+	}
+	ExtensionDirs.Reset();
+	GetExtensionDirs(ExtensionDirs, *ProjectDir, ExtensionSubDir, PlatformDirectoryNames);
+	for (const FString& LocalDir : ExtensionDirs)
+	{
+		if (!LocalDir.StartsWith(ProjectDir))
+		{
+			UE_LOG(LogZenFileSystemManifest, Warning,
+				TEXT("GetExtensionDirs on root directory '%s' unexpectedly returned folder '%s' that is not a subdirectory of the root directory."),
+				*ProjectDir, *LocalDir);
+			continue;
+		}
+		FString ClientDir = FPaths::Combine(TEXT("/{project}"), LocalDir.RightChop(ProjectDir.Len()));
+		AddFilesFromDirectory(ClientDir, LocalDir, true, AdditionalFilter);
+	}
+};
+
+void FManifestGenerator::GetExtensionDirs(TArray<FString>& OutExtensionDirs, const TCHAR* BaseDir, const TCHAR* SubDir,
+	const TArray<FString>& PlatformDirectoryNames)
+{
+	auto AddIfDirectoryExists = [&OutExtensionDirs](FString&& Dir)
+	{
+		if (FPaths::DirectoryExists(Dir))
+		{
+			OutExtensionDirs.Emplace(MoveTemp(Dir));
+		}
+	};
+
+	AddIfDirectoryExists(FPaths::Combine(BaseDir, SubDir));
+
+	FString PlatformExtensionBaseDir = FPaths::Combine(BaseDir, TEXT("Platforms"));
+	for (const FString& PlatformDirectoryName : PlatformDirectoryNames)
+	{
+		AddIfDirectoryExists(FPaths::Combine(PlatformExtensionBaseDir, PlatformDirectoryName, SubDir));
+	}
+
+	FString RestrictedBaseDir = FPaths::Combine(BaseDir, TEXT("Restricted"));
+	IFileManager::Get().IterateDirectory(*RestrictedBaseDir,
+		[&OutExtensionDirs, SubDir, &PlatformDirectoryNames]
+		(const TCHAR* FilenameOrDirectory, bool bIsDirectory) -> bool
+		{
+			if (bIsDirectory)
+			{
+				GetExtensionDirs(OutExtensionDirs, FilenameOrDirectory, SubDir, PlatformDirectoryNames);
+			}
+			return true;
+		});
+}
+
+TArray<TSharedRef<IPlugin>> FManifestGenerator::GetApplicablePlugins()
+{
+	constexpr bool bSkipDisabledPlugins = false;
+
+	IPluginManager& PluginManager = IPluginManager::Get();
+	TArray<TSharedRef<IPlugin>> DiscoveredPlugins = PluginManager.GetDiscoveredPlugins();
+
+	DiscoveredPlugins.RemoveAll([this](TSharedRef<IPlugin>& Plugin)
+		{
+			const FString& PluginName = Plugin->GetName();
+			if (bSkipDisabledPlugins && !Plugin->IsEnabled())
+			{
+				UE_LOG(LogZenFileSystemManifest, Verbose, TEXT("Plugin '%s' disabled, skipping"), *PluginName);
+				return true;
+			}
+			const FPluginDescriptor& Descriptor = Plugin->GetDescriptor();
+			if (!Descriptor.SupportsTargetPlatform(UBTPlatformString))
+			{
+				UE_LOG(LogZenFileSystemManifest, Verbose, TEXT("Plugin '%s' not supported on platform '%s', skipping"),
+					*PluginName, *Manifest.TargetPlatform.PlatformName());
+				for (const FString& SupportedTargetPlatform : Descriptor.SupportedTargetPlatforms)
+				{
+					UE_LOG(LogZenFileSystemManifest, Verbose, TEXT("       '%s' supports platform '%s'"), *PluginName, *SupportedTargetPlatform);
+				}
+				return true;
+			}
+
+			return false;
+		});
+	return DiscoveredPlugins;
+}
+
+void FManifestGenerator::GetPluginClientAndLocalRoots(TSharedRef<IPlugin>& Plugin,
+	FString& OutClientRoot, FString& OutLocalRoot)
+{
+	switch (Plugin->GetLoadedFrom())
+	{
+	case EPluginLoadedFrom::Engine:
+		OutClientRoot = TEXT("/{engine}");
+		OutLocalRoot = EngineDir;
+		break;
+	case EPluginLoadedFrom::Project:
+		OutClientRoot = TEXT("/{project}");
+		OutLocalRoot = ProjectDir;
+		break;
+	default:
+		checkNoEntry();
+		break;
+	}
+}
+
+TArray<FString> FManifestGenerator::GetApplicablePluginExtensionBaseDirs(TSharedRef<IPlugin>& Plugin)
+{
+	TArray<FString> ExtensionBaseDirs = Plugin->GetExtensionBaseDirs();
+	ExtensionBaseDirs.RemoveAll([this](const FString& ExtensionBaseDir)
+		{
+			// Scan the extension path for "Platforms/X" and include this extension if it is not platform specific at all,
+			// or if X is found and it is a valid target platform
+			bool bFoundPlatformsComponent = false;
+			bool bDone = false;
+			bool bIncludeExtension = true;
+			FPathViews::IterateComponents(
+				ExtensionBaseDir,
+				[this, &bFoundPlatformsComponent, &bDone, &bIncludeExtension](FStringView CurrentPathComponent)
+				{
+					if (!bFoundPlatformsComponent)
+					{
+						if (CurrentPathComponent == TEXTVIEW("Platforms"))
+						{
+							bFoundPlatformsComponent = true;
+						}
+					}
+					else if (!bDone)
+					{
+						uint32 CurrentPathComponentHash = GetTypeHash(CurrentPathComponent);
+						const bool bIsValidPlatform = PlatformDirectoryNameSet.ContainsByHash(CurrentPathComponentHash, CurrentPathComponent);
+						bIncludeExtension = bIsValidPlatform;
+						bDone = true;
+					}
+					else
+					{
+						// Do nothing.
+					}
+				});
+			return !bIncludeExtension;
+		});
+	return ExtensionBaseDirs;
+}
+
+FString FManifestGenerator::GetInternationalizationPresetPath()
+{
+	const UProjectPackagingSettings* const PackagingSettings = GetDefault<UProjectPackagingSettings>();
+
+	FString InternationalizationPresetAsString = UEnum::GetValueAsString(PackagingSettings->InternationalizationPreset);
+	const TCHAR* InternationalizationPresetPath = FCString::Strrchr(*InternationalizationPresetAsString, ':');
+	if (InternationalizationPresetPath)
+	{
+		++InternationalizationPresetPath;
+	}
+	else
+	{
+		UE_LOG(LogZenFileSystemManifest, Warning, TEXT("Failed reading internationalization preset setting, defaulting to English"));
+		InternationalizationPresetPath = TEXT("English");
+	}
+	return FString(InternationalizationPresetPath);
+}
+
+bool FManifestGenerator::IsSSLCertificatesWillStage()
+{
+	bool bSSLCertificatesWillStage = false;
+	FConfigCacheIni* TargetPlatformConfig = Manifest.TargetPlatform.GetConfigSystem();
+	if (TargetPlatformConfig)
+	{
+		GConfig->GetBool(TEXT("/Script/Engine.NetworkSettings"), TEXT("n.VerifyPeer"), bSSLCertificatesWillStage, GEngineIni);
+	}
+	return bSSLCertificatesWillStage;
+}
+
+void FManifestGenerator::PopulatePlatformDirectoryNames()
+{
+	FString IniPlatformName = Manifest.TargetPlatform.IniPlatformName();
+	const FDataDrivenPlatformInfo& PlatformInfo = FDataDrivenPlatformInfoRegistry::GetPlatformInfo(IniPlatformName);
+	PlatformDirectoryNameSet.Reserve(PlatformInfo.IniParentChain.Num()
+		+ PlatformInfo.AdditionalRestrictedFolders.Num() + 1);
+
+	PlatformDirectoryNameSet.Add(IniPlatformName);
+	for (const FString& PlatformName : PlatformInfo.AdditionalRestrictedFolders)
+	{
+		PlatformDirectoryNameSet.Add(PlatformName);
+	}
+	for (const FString& PlatformName : PlatformInfo.IniParentChain)
+	{
+		PlatformDirectoryNameSet.Add(PlatformName);
+	}
+
+	PlatformDirectoryNames = PlatformDirectoryNameSet.Array();
+
+	UBTPlatformString = PlatformInfo.UBTPlatformString;
+}
+
+FFileFilter& FFileFilter::ExcludeDirectoryLeafName(const TCHAR* LeafName)
+{
+	DirectoryExclusionFilter.Emplace(LeafName);
+	return *this;
+}
+
+FFileFilter& FFileFilter::ExcludeExtension(const TCHAR* Extension)
+{
+	ExtensionExclusionFilter.Emplace(Extension);
+	return *this;
+}
+
+FFileFilter& FFileFilter::IncludeExtension(const TCHAR* Extension)
+{
+	ExtensionInclusionFilter.Emplace(Extension);
+	return *this;
+}
+
+bool FFileFilter::IsDirectoryLeafNameKept(FStringView LeafName) const
+{
+	for (const FString& ExcludedDirectory : DirectoryExclusionFilter)
+	{
+		if (LeafName == ExcludedDirectory)
+		{
+			return false;
+		}
+	}
+	return true;
+}
+
+bool FFileFilter::IsFileExtensionKept(FStringView Extension)
+{
+	if (!ExtensionExclusionFilter.IsEmpty())
+	{
+		for (const FString& ExcludedExtension : ExtensionExclusionFilter)
+		{
+			if (Extension == ExcludedExtension)
+			{
+				return false;
+			}
+		}
+	}
+	if (!ExtensionInclusionFilter.IsEmpty())
+	{
+		for (const FString& IncludedExtension : ExtensionInclusionFilter)
+		{
+			if (Extension == IncludedExtension)
+			{
+				return true;
+			}
+		}
+		return false;
+	}
+	return true;
+}
+
+FAddFilesFromDirectoryVisitor::FAddFilesFromDirectoryVisitor(
+	FManifestGenerator& InGenerator, const FString& InClientRoot, const FString& InLocalRoot, bool bInIncludeSubDirs,
+	FFileFilter* InAdditionalFilter)
+	: Generator(InGenerator), ClientRoot(InClientRoot), LocalRoot(InLocalRoot)
+	, AdditionalFilter(InAdditionalFilter), bIncludeSubDirs(bInIncludeSubDirs)
+{
+}
+
+bool FAddFilesFromDirectoryVisitor::VisitorFunction(const TCHAR* InnerFileNameOrDirectory, bool bIsDirectory)
+{
+	if (bIsDirectory)
+	{
+		if (!bIncludeSubDirs)
+		{
+			return true;
+		}
+		FStringView DirectoryName = FPathViews::GetPathLeaf(InnerFileNameOrDirectory);
+		if (!Generator.BaseFilter.IsDirectoryLeafNameKept(DirectoryName))
+		{
+			return true;
+		}
+		if (AdditionalFilter && !AdditionalFilter->IsDirectoryLeafNameKept(DirectoryName))
+		{
+			return true;
+		}
+		DirectoryVisitQueue.Add(InnerFileNameOrDirectory);
+		return true;
+	}
+	FStringView Extension = FPathViews::GetExtension(InnerFileNameOrDirectory);
+	if (!Generator.BaseFilter.IsFileExtensionKept(Extension))
+	{
+		return true;
+	}
+	if (AdditionalFilter && !AdditionalFilter->IsFileExtensionKept(Extension))
+	{
+		return true;
+	}
+	Generator.AddSingleFile(ClientRoot, LocalRoot, InnerFileNameOrDirectory);
+	return true;
+}
+
+} // namespace UE::ZenFileSystemManifest
diff --git a/Engine/Source/Developer/IoStoreUtilities/Private/ZenFileSystemManifest.h b/Engine/Source/Developer/IoStoreUtilities/Private/ZenFileSystemManifest.h
--- a/Engine/Source/Developer/IoStoreUtilities/Private/ZenFileSystemManifest.h
+++ b/Engine/Source/Developer/IoStoreUtilities/Private/ZenFileSystemManifest.h
@@ -8,6 +8,8 @@
 
 class ITargetPlatform;
 class FArchive;
+namespace UE::ZenFileSystemManifest { class FFileFilter; }
+namespace UE::ZenFileSystemManifest { class FManifestGenerator; }
 
 struct FZenFileSystemManifestEntry
 {
@@ -61,8 +63,6 @@ public:
 	}
 
 private:
-	static void GetExtensionDirs(TArray<FString>& OutExtensionDirs, const TCHAR* BaseDir, const TCHAR* SubDir, const TArray<FString>& PlatformDirectoryNames);
-
 	const ITargetPlatform& TargetPlatform;
 	FString CookDirectory;
 	FString ServerRoot;
@@ -74,4 +74,7 @@ private:
 #endif
 
 	static const FZenFileSystemManifestEntry InvalidEntry;
+
+	friend UE::ZenFileSystemManifest::FFileFilter;
+	friend UE::ZenFileSystemManifest::FManifestGenerator;
 };
