Initialize StorageServerPlatformFile after ZenStoreWriter has potentially deleted the oplog, to avoid stale file manifests in StorageServerPlatformFile referring to files in Zen which have been deleted.

[RN] none

diff --git a/Engine/Source/Developer/IoStoreUtilities/Private/ZenCookArtifactReader.cpp b/Engine/Source/Developer/IoStoreUtilities/Private/ZenCookArtifactReader.cpp
--- a/Engine/Source/Developer/IoStoreUtilities/Private/ZenCookArtifactReader.cpp
+++ b/Engine/Source/Developer/IoStoreUtilities/Private/ZenCookArtifactReader.cpp
@@ -17,12 +17,8 @@ FZenCookArtifactReader::FZenCookArtifactReader(
 #if !UE_BUILD_SHIPPING
 	: ScopeZenService(MakeUnique<UE::Zen::FScopeZenService>())
 	, ZenRootPath(InputPath)
-	, StorageServerPlatformFile(IStorageServerClientModule::Get().TryCreateCustomPlatformFile(*InputPath, &FPlatformFileManager::Get().GetPlatformFile()))
+	, StorageServerPlatformFile(nullptr)
 {
-	if (StorageServerPlatformFile)
-	{
-		StorageServerPlatformFile->SetLowerLevel(nullptr);
-	}
 }
 #else
 {
@@ -33,6 +29,17 @@ FZenCookArtifactReader::~FZenCookArtifactReader()
 {
 }
 
+void FZenCookArtifactReader::Initialize()
+{
+#if !UE_BUILD_SHIPPING
+	StorageServerPlatformFile.Reset(IStorageServerClientModule::Get().TryCreateCustomPlatformFile(*ZenRootPath, &FPlatformFileManager::Get().GetPlatformFile()));
+	if (StorageServerPlatformFile)
+	{
+		StorageServerPlatformFile->SetLowerLevel(nullptr);
+	}
+#endif // !UE_BUILD_SHIPPING
+}
+
 bool FZenCookArtifactReader::FileExists(const TCHAR* Filename)
 {
 #if !UE_BUILD_SHIPPING
diff --git a/Engine/Source/Developer/IoStoreUtilities/Public/CookArtifactReader.h b/Engine/Source/Developer/IoStoreUtilities/Public/CookArtifactReader.h
--- a/Engine/Source/Developer/IoStoreUtilities/Public/CookArtifactReader.h
+++ b/Engine/Source/Developer/IoStoreUtilities/Public/CookArtifactReader.h
@@ -12,6 +12,8 @@ class ICookArtifactReader
 {
 public:
 	virtual ~ICookArtifactReader() = default;
+
+	virtual void Initialize() {}
 	
 	virtual bool FileExists(const TCHAR* Filename) = 0;
 	virtual int64 FileSize(const TCHAR* Filename) = 0;
diff --git a/Engine/Source/Developer/IoStoreUtilities/Public/LayeredCookArtifactReader.h b/Engine/Source/Developer/IoStoreUtilities/Public/LayeredCookArtifactReader.h
--- a/Engine/Source/Developer/IoStoreUtilities/Public/LayeredCookArtifactReader.h
+++ b/Engine/Source/Developer/IoStoreUtilities/Public/LayeredCookArtifactReader.h
@@ -12,6 +12,8 @@ class FLayeredCookArtifactReader
 public:
 	FLayeredCookArtifactReader() = default;
 
+	void Initialize() override { for (TSharedRef<ICookArtifactReader>& Layer : Layers) { Layer->Initialize(); } }
+
 	IOSTOREUTILITIES_API void AddLayer(TSharedRef<ICookArtifactReader> InLayer);
 	IOSTOREUTILITIES_API bool RemoveLayer(TSharedRef<ICookArtifactReader> InLayer);
 	IOSTOREUTILITIES_API void EmptyLayers();
diff --git a/Engine/Source/Developer/IoStoreUtilities/Public/ZenCookArtifactReader.h b/Engine/Source/Developer/IoStoreUtilities/Public/ZenCookArtifactReader.h
--- a/Engine/Source/Developer/IoStoreUtilities/Public/ZenCookArtifactReader.h
+++ b/Engine/Source/Developer/IoStoreUtilities/Public/ZenCookArtifactReader.h
@@ -18,6 +18,7 @@ public:
 												const FString& MetadataDirectoryPath, 
 												const ITargetPlatform* TargetPlatform);
 	IOSTOREUTILITIES_API virtual ~FZenCookArtifactReader();
+	IOSTOREUTILITIES_API void Initialize() override;
 
 	IOSTOREUTILITIES_API bool FileExists(const TCHAR* Filename) override;
 	IOSTOREUTILITIES_API int64 FileSize(const TCHAR* Filename) override;
diff --git a/Engine/Source/Editor/UnrealEd/Private/CookOnTheFlyServer.cpp b/Engine/Source/Editor/UnrealEd/Private/CookOnTheFlyServer.cpp
--- a/Engine/Source/Editor/UnrealEd/Private/CookOnTheFlyServer.cpp
+++ b/Engine/Source/Editor/UnrealEd/Private/CookOnTheFlyServer.cpp
@@ -10417,6 +10417,9 @@ void UCookOnTheFlyServer::BeginCookSandbox(FBeginCookContext& BeginContext)
 			// Refresh PackageWriterCapabilities because they can change during Initialize
 			SavePackageContext.PackageWriterCapabilities = PackageWriter.GetCookCapabilities();
 
+			ICookArtifactReader& Reader = *SavePackageContext.ArtifactReader;
+			Reader.Initialize();
+
 			if (!PlatformContext.bWorkerOnSharedSandbox)
 			{
 				check(!IsCookWorkerMode());
