DLCCook: ZenServer: When collecting manifest files for a DLC cook, fix collecting the incorrect MetaData directory. And for better performance, skip collecting files from engine, game, or other plugins, unless the DLC cook cooked packages in them. Skipping those files is also a useful way to avoid warnings when cooking DLC from a cooked basegame editor, with the editor's source files stored in Pak files rather than on disk.
#rnx

diff --git a/Engine/Source/Developer/IoStoreUtilities/Private/IoStoreUtilities.cpp b/Engine/Source/Developer/IoStoreUtilities/Private/IoStoreUtilities.cpp
--- a/Engine/Source/Developer/IoStoreUtilities/Private/IoStoreUtilities.cpp
+++ b/Engine/Source/Developer/IoStoreUtilities/Private/IoStoreUtilities.cpp
@@ -9007,7 +9007,7 @@ int32 GenerateZenFileSystemManifest(ITargetPlatform* TargetPlatform)
 	const FString ResolvedMetadataPath = MetadataPathSandbox.Replace(TEXT("[Platform]"), *PlatformString);
 
 	FZenFileSystemManifest ZenFileSystemManifest(*TargetPlatform, ResolvedRootPath);
-	ZenFileSystemManifest.Generate();
+	ZenFileSystemManifest.Generate(ResolvedMetadataPath);
 	ZenFileSystemManifest.Save(*FPaths::Combine(ResolvedMetadataPath, TEXT("zenfs.manifest")));
 	return 0;
 }
diff --git a/Engine/Source/Developer/IoStoreUtilities/Private/ZenFileSystemManifest.cpp b/Engine/Source/Developer/IoStoreUtilities/Private/ZenFileSystemManifest.cpp
--- a/Engine/Source/Developer/IoStoreUtilities/Private/ZenFileSystemManifest.cpp
+++ b/Engine/Source/Developer/IoStoreUtilities/Private/ZenFileSystemManifest.cpp
@@ -84,7 +84,8 @@ public:
 class FManifestGenerator
 {
 public:
-	FManifestGenerator(FZenFileSystemManifest& InManifest);
+	FManifestGenerator(FZenFileSystemManifest& InManifest,
+		const TOptional<ICookedPackageWriter::FReferencedPluginsInfo>& InReferencedPlugins);
 
 	/** FPaths::RootDir, the root of the UnrealEngine tree, parent directory of EngineDir. */
 	const FString& GetUnrealEngineRootDir() const;
@@ -171,6 +172,12 @@ public:
 	/** Return bSSLCertificatesWillStage from network settings for the cook's targetplatform, false by default. */
 	bool IsSSLCertificatesWillStage();
 
+	/** Return true if the manifest should gather files from the Engine. This can be false during DLC cooks. */
+	bool IsIncludeEngine();
+
+	/** Return true if the manifest should gather files from the Project. This can be false during DLC cooks. */
+	bool IsIncludeProject();
+
 private:
 	void PopulatePlatformDirectoryNames();
 
@@ -179,6 +186,7 @@ private:
 
 private:
 	FZenFileSystemManifest& Manifest;
+	const TOptional<ICookedPackageWriter::FReferencedPluginsInfo>& ReferencedPlugins;
 	IPlatformFile& PlatformFile;
 	FString UnrealEngineRootDir;
 	FString EngineDir;
@@ -209,7 +217,8 @@ FZenFileSystemManifest::FZenFileSystemManifest(const ITargetPlatform& InTargetPl
 #endif
 }
 
-int32 FZenFileSystemManifest::Generate()
+int32 FZenFileSystemManifest::Generate(const FString& MetadataDirectoryPath,
+	const TOptional<ICookedPackageWriter::FReferencedPluginsInfo>& ReferencedPlugins, bool bMultiProductBuild)
 {
 	using namespace UE::ZenFileSystemManifest;
 
@@ -218,7 +227,7 @@ int32 FZenFileSystemManifest::Generate()
 	const UProjectPackagingSettings* const PackagingSettings = GetDefault<UProjectPackagingSettings>();
 	const int32 PreviousEntryCount = NumEntries();
 
-	FManifestGenerator Generator(*this);
+	FManifestGenerator Generator(*this, ReferencedPlugins);
 
 	Generator.GetBaseFilter()
 		.ExcludeDirectoryLeafName(TEXT("Binaries"))
@@ -239,12 +248,28 @@ int32 FZenFileSystemManifest::Generate()
 	FFileFilter CookedMetadataFilter = FFileFilter()
 		.ExcludeDirectoryLeafName(TEXT("ShaderLibrarySource"))
 		.ExcludeExtension(TEXT("manifest"));
-	Generator.AddFilesFromDirectory(TEXT("/{project}/Metadata"), Generator.GetCookedMetadataDir(), true,
-		&CookedMetadataFilter);
+	if (!bMultiProductBuild)
+	{
+		Generator.AddFilesFromDirectory(TEXT("/{project}"), Generator.GetCookedProjectDir(), MetadataDirectoryPath,
+			true, &CookedMetadataFilter);
+	}
+	else
+	{
+		// Temporary workaround: Do not collect the Plugin Metadata directory for MultiProductBuild; collect the
+		// project's MetaData directory (which does not exist but we look for it anyway to match legacy beahvior).
+		// We skip collecting the Plugin Metadatadirectory because those files are later moved by the
+		// MultiProductBuild scripts into the project's metadata directory, so collecting them would add references
+		// to files that will eventually not exist.
+		Generator.AddFilesFromDirectory(TEXT("/{project}/Metadata"), Generator.GetCookedMetadataDir(), true,
+			&CookedMetadataFilter);
+	}
 
-	FFileFilter ProjectSourceFilter = FFileFilter()
-		.IncludeExtension(TEXT("uproject"));
-	Generator.AddFilesFromDirectory(TEXT("/{project}"), Generator.GetProjectDir(), false, &ProjectSourceFilter);
+	if (Generator.IsIncludeProject())
+	{
+		FFileFilter ProjectSourceFilter = FFileFilter()
+			.IncludeExtension(TEXT("uproject"));
+		Generator.AddFilesFromDirectory(TEXT("/{project}"), Generator.GetProjectDir(), false, &ProjectSourceFilter);
+	}
 	
 	FFileFilter ConfigFilter = FFileFilter()
 		.IncludeExtension(TEXT("ini"));
@@ -288,30 +313,35 @@ int32 FZenFileSystemManifest::Generate()
 		}
 	}
 
-	FString InternationalizationPresetPath = Generator.GetInternationalizationPresetPath();
-	const TCHAR* ICUDataVersion = TEXT("icudt64l"); // TODO: Could this go into datadriven platform info? But it's basically always this.
-	Generator.AddFilesFromDirectory(
-		FPaths::Combine(TEXT("/{engine}"),
-			TEXT("Content"), TEXT("Internationalization"), ICUDataVersion),
-		FPaths::Combine(Generator.GetEngineDir(),
-			TEXT("Content"), TEXT("Internationalization"), InternationalizationPresetPath, ICUDataVersion),
-		true /* bRecursive */);
-	
+	if (Generator.IsIncludeEngine())
+	{
+		FString InternationalizationPresetPath = Generator.GetInternationalizationPresetPath();
+		const TCHAR* ICUDataVersion = TEXT("icudt64l"); // TODO: Could this go into datadriven platform info? But it's basically always this.
+		Generator.AddFilesFromDirectory(
+			FPaths::Combine(TEXT("/{engine}"),
+				TEXT("Content"), TEXT("Internationalization"), ICUDataVersion),
+			FPaths::Combine(Generator.GetEngineDir(),
+				TEXT("Content"), TEXT("Internationalization"), InternationalizationPresetPath, ICUDataVersion),
+			true /* bRecursive */);
+	}
 	Generator.AddFilesFromExtensionDirectories(TEXT("Content/Localization"), &LocalizationFilter);
 
-	if (Generator.IsSSLCertificatesWillStage())
+	if (Generator.IsIncludeEngine() || Generator.IsIncludeProject())
 	{
-		if (!Generator.TryAddSingleRelpath(TEXT("/{project}"), Generator.GetProjectDir(),
-			FPaths::Combine(TEXT("Content"), TEXT("Certificates"), TEXT("cacert.pem"))))
+		if (Generator.IsSSLCertificatesWillStage())
 		{
-			Generator.TryAddSingleRelpath(TEXT("/{engine}"), Generator.GetEngineDir(),
-				FPaths::Combine(TEXT("Content"), TEXT("Certificates"), TEXT("ThirdParty"), TEXT("cacert.pem")));
-		}
+			if (!Generator.TryAddSingleRelpath(TEXT("/{project}"), Generator.GetProjectDir(),
+				FPaths::Combine(TEXT("Content"), TEXT("Certificates"), TEXT("cacert.pem"))))
+			{
+				Generator.TryAddSingleRelpath(TEXT("/{engine}"), Generator.GetEngineDir(),
+					FPaths::Combine(TEXT("Content"), TEXT("Certificates"), TEXT("ThirdParty"), TEXT("cacert.pem")));
+			}
 
-		FFileFilter CertificateFilter = FFileFilter()
-			.IncludeExtension(TEXT("pem"));
-		Generator.AddFilesFromDirectory(TEXT("/{project}/Certificates"),
-			FPaths::Combine(Generator.GetProjectDir(), TEXT("Certificates")), true, &CertificateFilter);
+			FFileFilter CertificateFilter = FFileFilter()
+				.IncludeExtension(TEXT("pem"));
+			Generator.AddFilesFromDirectory(TEXT("/{project}/Certificates"),
+				FPaths::Combine(Generator.GetProjectDir(), TEXT("Certificates")), true, &CertificateFilter);
+		}
 	}
 
 	FFileFilter ContentFilter = FFileFilter()
@@ -319,42 +349,57 @@ int32 FZenFileSystemManifest::Generate()
 		.ExcludeExtension(TEXT("ubulk"))
 		.ExcludeExtension(TEXT("uexp"))
 		.ExcludeExtension(TEXT("umap"));
-	Generator.AddFilesFromDirectory(TEXT("/{engine}/Content/Slate"),
-		FPaths::Combine(Generator.GetEngineDir(), TEXT("Content"), TEXT("Slate")),
-		true, &ContentFilter);
-	Generator.AddFilesFromDirectory(TEXT("/{project}/Content/Slate"),
-		FPaths::Combine(Generator.GetProjectDir(), TEXT("Content"), TEXT("Slate")),
-		true, &ContentFilter);
-	Generator.AddFilesFromDirectory(TEXT("/{engine}/Content/Movies"),
-		FPaths::Combine(Generator.GetEngineDir(), TEXT("Content"), TEXT("Movies")),
-		true, &ContentFilter);
-	Generator.AddFilesFromDirectory(TEXT("/{project}/Content/Movies"),
-		FPaths::Combine(Generator.GetProjectDir(), TEXT("Content"), TEXT("Movies")),
-		true, &ContentFilter);
+	if (Generator.IsIncludeEngine())
+	{
+		Generator.AddFilesFromDirectory(TEXT("/{engine}/Content/Slate"),
+			FPaths::Combine(Generator.GetEngineDir(), TEXT("Content"), TEXT("Slate")),
+			true, &ContentFilter);
+		Generator.AddFilesFromDirectory(TEXT("/{engine}/Content/Movies"),
+			FPaths::Combine(Generator.GetEngineDir(), TEXT("Content"), TEXT("Movies")),
+			true, &ContentFilter);
+	}
+	if (Generator.IsIncludeProject())
+	{
+		Generator.AddFilesFromDirectory(TEXT("/{project}/Content/Slate"),
+			FPaths::Combine(Generator.GetProjectDir(), TEXT("Content"), TEXT("Slate")),
+			true, &ContentFilter);
+		Generator.AddFilesFromDirectory(TEXT("/{project}/Content/Movies"),
+			FPaths::Combine(Generator.GetProjectDir(), TEXT("Content"), TEXT("Movies")),
+			true, &ContentFilter);
+	}
 	
-	FFileFilter OoodleDictionaryFilter = FFileFilter()
-		.IncludeExtension(TEXT("udic"));
-	Generator.AddFilesFromDirectory(TEXT("/{project}/Content/Oodle"),
-		FPaths::Combine(Generator.GetProjectDir(), TEXT("Content"), TEXT("Oodle")),
-		false, &OoodleDictionaryFilter);
+	if (Generator.IsIncludeProject())
+	{
+		FFileFilter OoodleDictionaryFilter = FFileFilter()
+			.IncludeExtension(TEXT("udic"));
+		Generator.AddFilesFromDirectory(TEXT("/{project}/Content/Oodle"),
+			FPaths::Combine(Generator.GetProjectDir(), TEXT("Content"), TEXT("Oodle")),
+			false, &OoodleDictionaryFilter);
+	}
 
 	FFileFilter ShaderCacheFilter = FFileFilter()
 		.IncludeExtension(TEXT("ushadercache"))
 		.IncludeExtension(TEXT("upipelinecache"));
-	Generator.AddFilesFromDirectory(TEXT("/{project}/Content"),
-		FPaths::Combine(Generator.GetProjectDir(), TEXT("Content")),
-		false, &ShaderCacheFilter);
-	Generator.AddFilesFromDirectory(FPaths::Combine(TEXT("/{project}"), TEXT("Content"), TEXT("PipelineCaches"), TargetPlatform.IniPlatformName()),
-		FPaths::Combine(Generator.GetProjectDir(), TEXT("Content"), TEXT("PipelineCaches"), TargetPlatform.IniPlatformName()),
-		false, &ShaderCacheFilter);
-
-	for (const FDirectoryPath& AdditionalFolderToStage : PackagingSettings->DirectoriesToAlwaysStageAsUFS)
+	if (Generator.IsIncludeProject())
 	{
-		Generator.TryAddFilesFromFlexDirectory(AdditionalFolderToStage.Path, true, &ContentFilter);
+		Generator.AddFilesFromDirectory(TEXT("/{project}/Content"),
+			FPaths::Combine(Generator.GetProjectDir(), TEXT("Content")),
+			false, &ShaderCacheFilter);
+		Generator.AddFilesFromDirectory(FPaths::Combine(TEXT("/{project}"), TEXT("Content"), TEXT("PipelineCaches"), TargetPlatform.IniPlatformName()),
+			FPaths::Combine(Generator.GetProjectDir(), TEXT("Content"), TEXT("PipelineCaches"), TargetPlatform.IniPlatformName()),
+			false, &ShaderCacheFilter);
 	}
-	for (const FDirectoryPath& AdditionalFolderToStage : PackagingSettings->DirectoriesToAlwaysStageAsUFSServer)
+
+	if (Generator.IsIncludeProject())
 	{
-		Generator.TryAddFilesFromFlexDirectory(AdditionalFolderToStage.Path, true, &ContentFilter);
+		for (const FDirectoryPath& AdditionalFolderToStage : PackagingSettings->DirectoriesToAlwaysStageAsUFS)
+		{
+			Generator.TryAddFilesFromFlexDirectory(AdditionalFolderToStage.Path, true, &ContentFilter);
+		}
+		for (const FDirectoryPath& AdditionalFolderToStage : PackagingSettings->DirectoriesToAlwaysStageAsUFSServer)
+		{
+			Generator.TryAddFilesFromFlexDirectory(AdditionalFolderToStage.Path, true, &ContentFilter);
+		}
 	}
 
 	const int32 CurrentEntryCount = NumEntries();
@@ -480,8 +525,10 @@ bool FZenFileSystemManifest::Save(const TCHAR* Filename)
 namespace UE::ZenFileSystemManifest
 {
 
-FManifestGenerator::FManifestGenerator(FZenFileSystemManifest& InManifest)
+FManifestGenerator::FManifestGenerator(FZenFileSystemManifest& InManifest,
+	const TOptional<ICookedPackageWriter::FReferencedPluginsInfo>& InReferencedPlugins)
 	: Manifest(InManifest)
+	, ReferencedPlugins(InReferencedPlugins)
 	, PlatformFile(FPlatformFileManager::Get().GetPlatformFile())
 {
 	UnrealEngineRootDir = FPaths::RootDir();
@@ -629,32 +676,38 @@ bool FManifestGenerator::TryAddFilesFromFlexDirectory(const FString& FullOrProje
 void FManifestGenerator::AddFilesFromExtensionDirectories(const TCHAR* ExtensionSubDir, FFileFilter* AdditionalFilter)
 {
 	TArray<FString> ExtensionDirs;
-	GetExtensionDirs(ExtensionDirs, *EngineDir, ExtensionSubDir, PlatformDirectoryNames);
-	for (const FString& LocalDir : ExtensionDirs)
+	if (IsIncludeEngine())
 	{
-		if (!LocalDir.StartsWith(EngineDir))
+		GetExtensionDirs(ExtensionDirs, *EngineDir, ExtensionSubDir, PlatformDirectoryNames);
+		for (const FString& LocalDir : ExtensionDirs)
 		{
-			UE_LOG(LogZenFileSystemManifest, Warning,
-				TEXT("GetExtensionDirs on root directory '%s' unexpectedly returned folder '%s' that is not a subdirectory of the root directory."),
-				*EngineDir, *LocalDir);
-			continue;
+			if (!LocalDir.StartsWith(EngineDir))
+			{
+				UE_LOG(LogZenFileSystemManifest, Warning,
+					TEXT("GetExtensionDirs on root directory '%s' unexpectedly returned folder '%s' that is not a subdirectory of the root directory."),
+					*EngineDir, *LocalDir);
+				continue;
+			}
+			FString ClientDir = FPaths::Combine(TEXT("/{engine}"), LocalDir.RightChop(EngineDir.Len()));
+			AddFilesFromDirectory(ClientDir, LocalDir, true, AdditionalFilter);
 		}
-		FString ClientDir = FPaths::Combine(TEXT("/{engine}"), LocalDir.RightChop(EngineDir.Len()));
-		AddFilesFromDirectory(ClientDir, LocalDir, true, AdditionalFilter);
 	}
-	ExtensionDirs.Reset();
-	GetExtensionDirs(ExtensionDirs, *ProjectDir, ExtensionSubDir, PlatformDirectoryNames);
-	for (const FString& LocalDir : ExtensionDirs)
+	if (IsIncludeProject())
 	{
-		if (!LocalDir.StartsWith(ProjectDir))
+		ExtensionDirs.Reset();
+		GetExtensionDirs(ExtensionDirs, *ProjectDir, ExtensionSubDir, PlatformDirectoryNames);
+		for (const FString& LocalDir : ExtensionDirs)
 		{
-			UE_LOG(LogZenFileSystemManifest, Warning,
-				TEXT("GetExtensionDirs on root directory '%s' unexpectedly returned folder '%s' that is not a subdirectory of the root directory."),
-				*ProjectDir, *LocalDir);
-			continue;
+			if (!LocalDir.StartsWith(ProjectDir))
+			{
+				UE_LOG(LogZenFileSystemManifest, Warning,
+					TEXT("GetExtensionDirs on root directory '%s' unexpectedly returned folder '%s' that is not a subdirectory of the root directory."),
+					*ProjectDir, *LocalDir);
+				continue;
+			}
+			FString ClientDir = FPaths::Combine(TEXT("/{project}"), LocalDir.RightChop(ProjectDir.Len()));
+			AddFilesFromDirectory(ClientDir, LocalDir, true, AdditionalFilter);
 		}
-		FString ClientDir = FPaths::Combine(TEXT("/{project}"), LocalDir.RightChop(ProjectDir.Len()));
-		AddFilesFromDirectory(ClientDir, LocalDir, true, AdditionalFilter);
 	}
 };
 
@@ -702,21 +755,58 @@ TArray<TSharedRef<IPlugin>> FManifestGenerator::GetApplicablePlugins()
 			const FString& PluginName = Plugin->GetName();
 			if (bSkipDisabledPlugins && !Plugin->IsEnabled())
 			{
-				UE_LOG(LogZenFileSystemManifest, Verbose, TEXT("Plugin '%s' disabled, skipping"), *PluginName);
+				UE_LOG(LogZenFileSystemManifest, Verbose, TEXT("Skipping plugin '%s' because it is disabled."), *PluginName);
 				return true;
 			}
 			const FPluginDescriptor& Descriptor = Plugin->GetDescriptor();
 			if (!Descriptor.SupportsTargetPlatform(UBTPlatformString))
 			{
-				UE_LOG(LogZenFileSystemManifest, Verbose, TEXT("Plugin '%s' not supported on platform '%s', skipping"),
+				UE_LOG(LogZenFileSystemManifest, Verbose, TEXT("Skipping plugin '%s' because it is not supported on platform '%s'."),
 					*PluginName, *Manifest.TargetPlatform.PlatformName());
 				for (const FString& SupportedTargetPlatform : Descriptor.SupportedTargetPlatforms)
 				{
-					UE_LOG(LogZenFileSystemManifest, Verbose, TEXT("       '%s' supports platform '%s'"), *PluginName, *SupportedTargetPlatform);
+					UE_LOG(LogZenFileSystemManifest, Verbose, TEXT("       '%s' supports platform '%s'"),
+						*PluginName, *SupportedTargetPlatform);
 				}
 				return true;
 			}
 
+			if (ReferencedPlugins.IsSet())
+			{
+				// Skip adding files to the manifest that are under plugins that were not cooked in the current BuildLayer.
+
+				if (!Plugin->CanContainContent())
+				{
+					// Code-only plugins would not be found by looking for cooked packages, so do not filter them out based
+					// on their presence in the ReferencedPlugins. Instead, add all code-only plugins from game or engine that
+					// are enabled on the target, but only if we are cooking game or engine.
+					if (Plugin->GetLoadedFrom() == EPluginLoadedFrom::Engine && !ReferencedPlugins->bReferencesEngine)
+					{
+						UE_LOG(LogZenFileSystemManifest, Verbose,
+							TEXT("Skipping plugin '%s' because it is a code-only Engine plugin, and we did not cook /Engine."),
+							*PluginName);
+						return true;
+					}
+					if (Plugin->GetLoadedFrom() == EPluginLoadedFrom::Project && !ReferencedPlugins->bReferencesGame)
+					{
+						UE_LOG(LogZenFileSystemManifest, Verbose,
+							TEXT("Skipping plugin '%s' because it is a code-only Project plugin, and we did not cook /Game."),
+							*PluginName);
+						return true;
+					}
+				}
+				else
+				{
+					if (!ReferencedPlugins->ReferencedPlugins.Contains(Plugin->GetName()))
+					{
+						UE_LOG(LogZenFileSystemManifest, Verbose,
+							TEXT("Skipping plugin '%s' because we did not cook any content from it."),
+							*PluginName);
+						return true;
+					}
+				}
+			}
+
 			return false;
 		});
 	return DiscoveredPlugins;
@@ -808,6 +898,25 @@ bool FManifestGenerator::IsSSLCertificatesWillStage()
 	return bSSLCertificatesWillStage;
 }
 
+bool FManifestGenerator::IsIncludeEngine()
+{
+	if (!ReferencedPlugins.IsSet())
+	{
+		return true;
+	}
+
+	return ReferencedPlugins->bReferencesEngine;
+}
+
+bool FManifestGenerator::IsIncludeProject()
+{
+	if (!ReferencedPlugins.IsSet())
+	{
+		return true;
+	}
+	return ReferencedPlugins->bReferencesGame;
+}
+
 void FManifestGenerator::PopulatePlatformDirectoryNames()
 {
 	FString IniPlatformName = Manifest.TargetPlatform.IniPlatformName();
diff --git a/Engine/Source/Developer/IoStoreUtilities/Private/ZenFileSystemManifest.h b/Engine/Source/Developer/IoStoreUtilities/Private/ZenFileSystemManifest.h
--- a/Engine/Source/Developer/IoStoreUtilities/Private/ZenFileSystemManifest.h
+++ b/Engine/Source/Developer/IoStoreUtilities/Private/ZenFileSystemManifest.h
@@ -9,6 +9,7 @@
 #include "CoreTypes.h"
 #include "IO/IoDispatcher.h"
 #include "Misc/Optional.h"
+#include "Serialization/PackageWriter.h"
 
 class ITargetPlatform;
 class FArchive;
@@ -32,7 +33,10 @@ public:
 	 * present in the manifest. Files in the manifest do not contain any payload data, so no update is necessary
 	 * for Entries that already exist. Returns the number of new entries.
 	 */
-	int32 Generate();
+	int32 Generate(const FString& MetadataDirectoryPath,
+		const TOptional<ICookedPackageWriter::FReferencedPluginsInfo>& ReferencedPlugins
+			= TOptional<ICookedPackageWriter::FReferencedPluginsInfo>(),
+		bool bMultiProductBuild = false);
 
 	const FZenFileSystemManifestEntry& CreateManifestEntry(const FString& Filename);
 	
diff --git a/Engine/Source/Developer/IoStoreUtilities/Private/ZenStoreWriter.cpp b/Engine/Source/Developer/IoStoreUtilities/Private/ZenStoreWriter.cpp
--- a/Engine/Source/Developer/IoStoreUtilities/Private/ZenStoreWriter.cpp
+++ b/Engine/Source/Developer/IoStoreUtilities/Private/ZenStoreWriter.cpp
@@ -628,7 +628,7 @@ void FZenStoreWriter::BeginCook(const FCookInfo& Info)
 
 	if (CookMode == ICookedPackageWriter::FCookInfo::CookOnTheFlyMode)
 	{
-		ZenFileSystemManifest->Generate();
+		ZenFileSystemManifest->Generate(MetadataDirectoryPath, Info.ReferencedPlugins, Info.bMultiProductBuild);
 		TIoStatusOr<uint64> Status = HttpClient->AppendOp(CreateProjectMetaDataOpPackage("CookOnTheFly"));
 		UE_CLOG(!Status.IsOk(), LogZenStoreWriter, Fatal, TEXT("Failed to append OpLog"));
 	}
@@ -662,7 +662,7 @@ void FZenStoreWriter::EndCook(const FCookInfo& Info)
 
 	if (!Info.bWorkerOnSharedSandbox)
 	{
-		ZenFileSystemManifest->Generate();
+		ZenFileSystemManifest->Generate(MetadataDirectoryPath, Info.ReferencedPlugins, Info.bMultiProductBuild);
 		if (FCbPackage ReferencedSetPackage = CreateReferencedSetOpPackage())
 		{
 			TIoStatusOr<uint64> Status = HttpClient->AppendOp(MoveTemp(ReferencedSetPackage));
diff --git a/Engine/Source/Editor/UnrealEd/Classes/CookOnTheSide/CookOnTheFlyServer.h b/Engine/Source/Editor/UnrealEd/Classes/CookOnTheSide/CookOnTheFlyServer.h
--- a/Engine/Source/Editor/UnrealEd/Classes/CookOnTheSide/CookOnTheFlyServer.h
+++ b/Engine/Source/Editor/UnrealEd/Classes/CookOnTheSide/CookOnTheFlyServer.h
@@ -39,7 +39,6 @@ class FReferenceCollector;
 class FSavePackageContext;
 class IAssetRegistry;
 class ICookArtifactReader;
-class ICookedPackageWriter;
 class IPlugin;
 class ITargetPlatform;
 enum class EPackageWriterResult : uint8;
@@ -1217,8 +1216,13 @@ private:
 	void SetNeverCookPackageConfigSettings(FBeginCookContext& BeginContext, UE::Cook::FBeginCookConfigSettings& Settings);
 
 	/** Finalize the package store */
-	void FinalizePackageStore();
-
+	void FinalizePackageStore(const ITargetPlatform* TargetPlatform, const TSet<FName>& CookedPackageNames);
+	/**
+	 * Calculate which Plugins (and also /Game and /Editor) had packages referenced in the current session,
+	 * this is important to know for DLC cooks.
+	 */
+	void CalculateReferencedPlugins(const ITargetPlatform* TargetPlatform,
+		const TSet<FName>& CookedPackageNames, ICookedPackageWriter::FReferencedPluginsInfo& OutReferencedPluginsInfo);
 	/** Empties SavePackageContexts and deletes the contents */
 	void ClearPackageStoreContexts();
 
diff --git a/Engine/Source/Editor/UnrealEd/Private/CookOnTheFlyServer.cpp b/Engine/Source/Editor/UnrealEd/Private/CookOnTheFlyServer.cpp
--- a/Engine/Source/Editor/UnrealEd/Private/CookOnTheFlyServer.cpp
+++ b/Engine/Source/Editor/UnrealEd/Private/CookOnTheFlyServer.cpp
@@ -9360,8 +9360,17 @@ void UCookOnTheFlyServer::CookByTheBookFinishedInternal()
 
 	UE_LOG(LogCookStatus, Display, TEXT("Collecting cook information for all platforms"));
 
+	struct FEndCookPlatformData
+	{
+		TSet<FName> CookedPackageNames;
+	};
+	TArray<FEndCookPlatformData> EndCookPlatformDatas;
+	EndCookPlatformDatas.SetNum(PlatformManager->GetSessionPlatforms().Num());
+
+	int32 PlatformIndex = 0;
 	for (const ITargetPlatform* TargetPlatform : PlatformManager->GetSessionPlatforms())
 	{
+		FEndCookPlatformData& EndCookPlatformData = EndCookPlatformDatas[PlatformIndex++];
 		if (FindOrCreateSaveContext(TargetPlatform).PackageWriterCapabilities.bReadOnly)
 		{
 			continue;
@@ -9374,7 +9383,7 @@ void UCookOnTheFlyServer::CookByTheBookFinishedInternal()
 
 		FString PlatformNameString = TargetPlatform->PlatformName();
 
-		TSet<FName> CookedPackageNames;
+		TSet<FName>& CookedPackageNames = EndCookPlatformData.CookedPackageNames;
 		TSet<FName> IgnorePackageNames;
 		{
 			UE_SCOPED_HIERARCHICAL_COOKTIMER(CalculateReferencedSet);
@@ -9586,7 +9595,16 @@ void UCookOnTheFlyServer::CookByTheBookFinishedInternal()
 		GenerateCachedEditorThumbnails();
 	}
 
-	FinalizePackageStore();
+	{
+		UE_SCOPED_HIERARCHICAL_COOKTIMER(FinalizePackageStore);
+		UE_LOG(LogCook, Display, TEXT("Finalize package store(s)..."));
+		PlatformIndex = 0;
+		for (const ITargetPlatform* TargetPlatform : PlatformManager->GetSessionPlatforms())
+		{
+			FEndCookPlatformData& EndCookPlatformData = EndCookPlatformDatas[PlatformIndex++];
+			FinalizePackageStore(TargetPlatform, EndCookPlatformData.CookedPackageNames);
+		}
+	}
 
 	UE_LOG(LogCookStatus, Display, TEXT("FinishedInternal: Done"));
 }
@@ -10413,6 +10431,7 @@ void UCookOnTheFlyServer::BeginCookSandbox(FBeginCookContext& BeginContext)
 			CookInfo.bFullBuild = PlatformContext.bFullBuild;
 			CookInfo.bLegacyIterativeSharedBuild = PlatformContext.bLegacyIterativeSharedBuild;
 			CookInfo.bWorkerOnSharedSandbox = PlatformContext.bWorkerOnSharedSandbox;
+			CookInfo.bMultiProductBuild = bCookingDLCBaseGamePlugin;
 			PackageWriter.Initialize(CookInfo);
 			// Refresh PackageWriterCapabilities because they can change during Initialize
 			SavePackageContext.PackageWriterCapabilities = PackageWriter.GetCookCapabilities();
@@ -10665,21 +10684,80 @@ void UCookOnTheFlyServer::DeleteOutputForPackage(FName PackageName, const ITarge
 	FindOrCreatePackageWriter(TargetPlatform).RemoveCookedPackages(TArrayView<FName>(&PackageName, 1));
 }
 
-void UCookOnTheFlyServer::FinalizePackageStore()
+void UCookOnTheFlyServer::FinalizePackageStore(const ITargetPlatform* TargetPlatform, const TSet<FName>& CookedPackageNames)
 {
-	UE_SCOPED_HIERARCHICAL_COOKTIMER(FinalizePackageStore);
+	using namespace UE::Cook;
 
-	UE_LOG(LogCook, Display, TEXT("Finalize package store(s)..."));
-	for (const ITargetPlatform* TargetPlatform : PlatformManager->GetSessionPlatforms())
+	FPlatformData* PlatformData = PlatformManager->GetPlatformData(TargetPlatform);
+	ICookedPackageWriter::FCookInfo CookInfo;
+	CookInfo.CookMode = IsDirectorCookOnTheFly() ? ICookedPackageWriter::FCookInfo::CookOnTheFlyMode : ICookedPackageWriter::FCookInfo::CookByTheBookMode;
+	CookInfo.bFullBuild = PlatformData->bFullBuild;
+	CookInfo.bLegacyIterativeSharedBuild = PlatformData->bLegacyIterativeSharedBuild;
+	CookInfo.bWorkerOnSharedSandbox = PlatformData->bWorkerOnSharedSandbox;
+	CookInfo.bMultiProductBuild = bCookingDLCBaseGamePlugin;
+
+	if (GetCookingDLC() == ECookingDLC::Yes && !CookInfo.bWorkerOnSharedSandbox)
 	{
-		UE::Cook::FPlatformData* PlatformData = PlatformManager->GetPlatformData(TargetPlatform);
-		ICookedPackageWriter::FCookInfo CookInfo;
-		CookInfo.CookMode = IsDirectorCookOnTheFly() ? ICookedPackageWriter::FCookInfo::CookOnTheFlyMode : ICookedPackageWriter::FCookInfo::CookByTheBookMode;
-		CookInfo.bFullBuild = PlatformData->bFullBuild;
-		CookInfo.bLegacyIterativeSharedBuild = PlatformData->bLegacyIterativeSharedBuild;
-		CookInfo.bWorkerOnSharedSandbox = PlatformData->bWorkerOnSharedSandbox;
+		CalculateReferencedPlugins(TargetPlatform, CookedPackageNames, CookInfo.ReferencedPlugins.Emplace());
+	}
+
+	FindOrCreatePackageWriter(TargetPlatform).EndCook(CookInfo);
+}
 
-		FindOrCreatePackageWriter(TargetPlatform).EndCook(CookInfo);
+void UCookOnTheFlyServer::CalculateReferencedPlugins(const ITargetPlatform* TargetPlatform,
+	const TSet<FName>& CookedPackageNames, ICookedPackageWriter::FReferencedPluginsInfo& OutReferencedPluginsInfo)
+{
+	using namespace UE::Cook;
+
+	OutReferencedPluginsInfo.bReferencesEngine = false;
+	OutReferencedPluginsInfo.bReferencesGame = false;
+	OutReferencedPluginsInfo.ReferencedPlugins.Reset();
+
+	IPluginManager& PluginManager = IPluginManager::Get();
+
+	FString PluginNameStr;
+	PluginNameStr.Reserve(256);
+	for (FName PackageName : CookedPackageNames)
+	{
+		TStringBuilder<FName::StringBufferSize> PackageNameStr(InPlace, PackageName);
+		if (PackageNameStr.Len() == 0 || (*PackageNameStr)[0] != '/')
+		{
+			continue;
+		}
+		FStringView MountName(PackageNameStr);
+		MountName.RightChopInline(1);
+		int32 SecondSlashIndex;
+		if (!MountName.FindChar('/', SecondSlashIndex))
+		{
+			continue;
+		}
+		MountName.LeftInline(SecondSlashIndex);
+
+		if (MountName == TEXT("Engine"))
+		{
+			OutReferencedPluginsInfo.bReferencesEngine = true;
+			continue;
+		}
+		if (MountName == TEXT("Game"))
+		{
+			OutReferencedPluginsInfo.bReferencesGame = true;
+			continue;
+		}
+		TSharedPtr<IPlugin> Plugin = PluginManager.FindPlugin(MountName);
+		if (Plugin)
+		{
+			PluginNameStr.Reset();
+			PluginNameStr.Append(MountName);
+			OutReferencedPluginsInfo.ReferencedPlugins.Add(PluginNameStr);
+		}
+	}
+
+	if (GetCookingDLC() == ECookingDLC::No)
+	{
+		// For basegame cooks, mark that they reference Engine and Game even if they
+		// contain no packages, so that they stage the non-package Engine and Game files.
+		OutReferencedPluginsInfo.bReferencesEngine = true;
+		OutReferencedPluginsInfo.bReferencesGame = true;
 	}
 }
 
@@ -11116,7 +11194,14 @@ void UCookOnTheFlyServer::CookAsCookWorkerFinished()
 	}
 
 	ShutdownShaderLibraryCookerAndCompilers(GetProjectShaderLibraryName());
-	FinalizePackageStore();
+	{
+		UE_SCOPED_HIERARCHICAL_COOKTIMER(FinalizePackageStore);
+		UE_LOG(LogCook, Display, TEXT("Finalize package store(s)..."));
+		for (const ITargetPlatform* TargetPlatform : PlatformManager->GetSessionPlatforms())
+		{
+			FinalizePackageStore(TargetPlatform, TSet<FName>());
+		}
+	}
 
 	if (IsDirectorCookOnTheFly())
 	{
@@ -11727,6 +11812,7 @@ void UCookOnTheFlyServer::BeginCookPackageWriters(FBeginCookContext& BeginContex
 		CookInfo.bFullBuild = Context.bFullBuild;
 		CookInfo.bLegacyIterativeSharedBuild = Context.bLegacyIterativeSharedBuild;
 		CookInfo.bWorkerOnSharedSandbox = Context.bWorkerOnSharedSandbox;
+		CookInfo.bMultiProductBuild = bCookingDLCBaseGamePlugin;
 
 		FindOrCreatePackageWriter(Context.TargetPlatform).BeginCook(CookInfo);
 	}
diff --git a/Engine/Source/Runtime/Core/Public/Serialization/PackageWriter.h b/Engine/Source/Runtime/Core/Public/Serialization/PackageWriter.h
--- a/Engine/Source/Runtime/Core/Public/Serialization/PackageWriter.h
+++ b/Engine/Source/Runtime/Core/Public/Serialization/PackageWriter.h
@@ -3,6 +3,7 @@
 #pragma once
 
 #include "Async/Future.h"
+#include "Containers/Set.h"
 #include "Containers/StringView.h"
 #include "IO/IoDispatcher.h"
 #include "IO/IoHash.h"
@@ -10,6 +11,7 @@
 #include "Misc/AssertionMacros.h"
 #include "Misc/DateTime.h"
 #include "Misc/EnumClassFlags.h"
+#include "Misc/Optional.h"
 #include "Misc/SecureHash.h"
 #include "Serialization/CompactBinary.h"
 #include "Templates/RefCounting.h"
@@ -372,6 +374,13 @@ public:
 		return this;
 	}
 
+	struct FReferencedPluginsInfo
+	{
+		TSet<FString> ReferencedPlugins;
+		bool bReferencesGame = false;
+		bool bReferencesEngine = false;
+	};
+
 	struct FCookInfo
 	{
 		enum ECookMode
@@ -380,12 +389,14 @@ public:
 			CookOnTheFlyMode
 		};
 
+		TOptional<FReferencedPluginsInfo> ReferencedPlugins;
 		ECookMode CookMode = CookByTheBookMode;
 		bool bFullBuild = true;
 		UE_DEPRECATED(5.6, "Use bLegacyIterativeSharedBuild")
 		bool bIterateSharedBuild = false;
 		bool bLegacyIterativeSharedBuild = false;
 		bool bWorkerOnSharedSandbox = false;
+		bool bMultiProductBuild = false;
 	};
 
 	/** Delete outdated cooked data, etc.
