DLCIncrementalCook: Remove the workaround for MultiProductBuild that reports the incorrect cooked metadata folder to the ZenFileSystemManifest for the DLC cooks of MPB. We think that using the correct folder and therefore adding the cook's DLC metadata files into the ZenFileSystemManifest will not cause a problem.
#rnx

diff --git a/Engine/Source/Developer/IoStoreUtilities/Private/ZenFileSystemManifest.cpp b/Engine/Source/Developer/IoStoreUtilities/Private/ZenFileSystemManifest.cpp
--- a/Engine/Source/Developer/IoStoreUtilities/Private/ZenFileSystemManifest.cpp
+++ b/Engine/Source/Developer/IoStoreUtilities/Private/ZenFileSystemManifest.cpp
@@ -218,7 +218,7 @@ FZenFileSystemManifest::FZenFileSystemManifest(const ITargetPlatform& InTargetPl
 }
 
 int32 FZenFileSystemManifest::Generate(const FString& MetadataDirectoryPath,
-	const TOptional<ICookedPackageWriter::FReferencedPluginsInfo>& ReferencedPlugins, bool bMultiProductBuild)
+	const TOptional<ICookedPackageWriter::FReferencedPluginsInfo>& ReferencedPlugins)
 {
 	using namespace UE::ZenFileSystemManifest;
 
@@ -248,21 +248,8 @@ int32 FZenFileSystemManifest::Generate(const FString& MetadataDirectoryPath,
 	FFileFilter CookedMetadataFilter = FFileFilter()
 		.ExcludeDirectoryLeafName(TEXT("ShaderLibrarySource"))
 		.ExcludeExtension(TEXT("manifest"));
-	if (!bMultiProductBuild)
-	{
-		Generator.AddFilesFromDirectory(TEXT("/{project}"), Generator.GetCookedProjectDir(), MetadataDirectoryPath,
-			true, &CookedMetadataFilter);
-	}
-	else
-	{
-		// Temporary workaround: Do not collect the Plugin Metadata directory for MultiProductBuild; collect the
-		// project's MetaData directory (which does not exist but we look for it anyway to match legacy beahvior).
-		// We skip collecting the Plugin Metadatadirectory because those files are later moved by the
-		// MultiProductBuild scripts into the project's metadata directory, so collecting them would add references
-		// to files that will eventually not exist.
-		Generator.AddFilesFromDirectory(TEXT("/{project}/Metadata"), Generator.GetCookedMetadataDir(), true,
-			&CookedMetadataFilter);
-	}
+	Generator.AddFilesFromDirectory(TEXT("/{project}"), Generator.GetCookedProjectDir(), MetadataDirectoryPath,
+		true, &CookedMetadataFilter);
 
 	if (Generator.IsIncludeProject())
 	{
diff --git a/Engine/Source/Developer/IoStoreUtilities/Private/ZenFileSystemManifest.h b/Engine/Source/Developer/IoStoreUtilities/Private/ZenFileSystemManifest.h
--- a/Engine/Source/Developer/IoStoreUtilities/Private/ZenFileSystemManifest.h
+++ b/Engine/Source/Developer/IoStoreUtilities/Private/ZenFileSystemManifest.h
@@ -35,8 +35,7 @@ public:
 	 */
 	int32 Generate(const FString& MetadataDirectoryPath,
 		const TOptional<ICookedPackageWriter::FReferencedPluginsInfo>& ReferencedPlugins
-			= TOptional<ICookedPackageWriter::FReferencedPluginsInfo>(),
-		bool bMultiProductBuild = false);
+			= TOptional<ICookedPackageWriter::FReferencedPluginsInfo>());
 
 	const FZenFileSystemManifestEntry& CreateManifestEntry(const FString& Filename);
 	
diff --git a/Engine/Source/Developer/IoStoreUtilities/Private/ZenStoreWriter.cpp b/Engine/Source/Developer/IoStoreUtilities/Private/ZenStoreWriter.cpp
--- a/Engine/Source/Developer/IoStoreUtilities/Private/ZenStoreWriter.cpp
+++ b/Engine/Source/Developer/IoStoreUtilities/Private/ZenStoreWriter.cpp
@@ -628,7 +628,7 @@ void FZenStoreWriter::BeginCook(const FCookInfo& Info)
 
 	if (CookMode == ICookedPackageWriter::FCookInfo::CookOnTheFlyMode)
 	{
-		ZenFileSystemManifest->Generate(MetadataDirectoryPath, Info.ReferencedPlugins, Info.bMultiProductBuild);
+		ZenFileSystemManifest->Generate(MetadataDirectoryPath, Info.ReferencedPlugins);
 		TIoStatusOr<uint64> Status = HttpClient->AppendOp(CreateProjectMetaDataOpPackage("CookOnTheFly"));
 		UE_CLOG(!Status.IsOk(), LogZenStoreWriter, Fatal, TEXT("Failed to append OpLog"));
 	}
@@ -662,7 +662,7 @@ void FZenStoreWriter::EndCook(const FCookInfo& Info)
 
 	if (!Info.bWorkerOnSharedSandbox)
 	{
-		ZenFileSystemManifest->Generate(MetadataDirectoryPath, Info.ReferencedPlugins, Info.bMultiProductBuild);
+		ZenFileSystemManifest->Generate(MetadataDirectoryPath, Info.ReferencedPlugins);
 		if (FCbPackage ReferencedSetPackage = CreateReferencedSetOpPackage())
 		{
 			TIoStatusOr<uint64> Status = HttpClient->AppendOp(MoveTemp(ReferencedSetPackage));
diff --git a/Engine/Source/Editor/UnrealEd/Private/CookOnTheFlyServer.cpp b/Engine/Source/Editor/UnrealEd/Private/CookOnTheFlyServer.cpp
--- a/Engine/Source/Editor/UnrealEd/Private/CookOnTheFlyServer.cpp
+++ b/Engine/Source/Editor/UnrealEd/Private/CookOnTheFlyServer.cpp
@@ -10431,7 +10431,6 @@ void UCookOnTheFlyServer::BeginCookSandbox(FBeginCookContext& BeginContext)
 			CookInfo.bFullBuild = PlatformContext.bFullBuild;
 			CookInfo.bLegacyIterativeSharedBuild = PlatformContext.bLegacyIterativeSharedBuild;
 			CookInfo.bWorkerOnSharedSandbox = PlatformContext.bWorkerOnSharedSandbox;
-			CookInfo.bMultiProductBuild = bCookingDLCBaseGamePlugin;
 			PackageWriter.Initialize(CookInfo);
 			// Refresh PackageWriterCapabilities because they can change during Initialize
 			SavePackageContext.PackageWriterCapabilities = PackageWriter.GetCookCapabilities();
@@ -10694,7 +10693,6 @@ void UCookOnTheFlyServer::FinalizePackageStore(const ITargetPlatform* TargetPlat
 	CookInfo.bFullBuild = PlatformData->bFullBuild;
 	CookInfo.bLegacyIterativeSharedBuild = PlatformData->bLegacyIterativeSharedBuild;
 	CookInfo.bWorkerOnSharedSandbox = PlatformData->bWorkerOnSharedSandbox;
-	CookInfo.bMultiProductBuild = bCookingDLCBaseGamePlugin;
 
 	if (GetCookingDLC() == ECookingDLC::Yes && !CookInfo.bWorkerOnSharedSandbox)
 	{
@@ -11812,7 +11810,6 @@ void UCookOnTheFlyServer::BeginCookPackageWriters(FBeginCookContext& BeginContex
 		CookInfo.bFullBuild = Context.bFullBuild;
 		CookInfo.bLegacyIterativeSharedBuild = Context.bLegacyIterativeSharedBuild;
 		CookInfo.bWorkerOnSharedSandbox = Context.bWorkerOnSharedSandbox;
-		CookInfo.bMultiProductBuild = bCookingDLCBaseGamePlugin;
 
 		FindOrCreatePackageWriter(Context.TargetPlatform).BeginCook(CookInfo);
 	}
diff --git a/Engine/Source/Runtime/Core/Public/Serialization/PackageWriter.h b/Engine/Source/Runtime/Core/Public/Serialization/PackageWriter.h
--- a/Engine/Source/Runtime/Core/Public/Serialization/PackageWriter.h
+++ b/Engine/Source/Runtime/Core/Public/Serialization/PackageWriter.h
@@ -396,7 +396,6 @@ public:
 		bool bIterateSharedBuild = false;
 		bool bLegacyIterativeSharedBuild = false;
 		bool bWorkerOnSharedSandbox = false;
-		bool bMultiProductBuild = false;
 	};
 
 	/** Delete outdated cooked data, etc.
