Initialize StorageServerPlatformFile before attempting to read cook settings from Zen, and re-initialize the StorageServerPlatformFile TOC if the oplog has been deleted.

[RN] none

#was-backout 47045172

diff --git a/Engine/Source/Developer/IoStoreUtilities/Private/ZenCookArtifactReader.cpp b/Engine/Source/Developer/IoStoreUtilities/Private/ZenCookArtifactReader.cpp
--- a/Engine/Source/Developer/IoStoreUtilities/Private/ZenCookArtifactReader.cpp
+++ b/Engine/Source/Developer/IoStoreUtilities/Private/ZenCookArtifactReader.cpp
@@ -17,8 +17,12 @@ FZenCookArtifactReader::FZenCookArtifactReader(
 #if !UE_BUILD_SHIPPING
 	: ScopeZenService(MakeUnique<UE::Zen::FScopeZenService>())
 	, ZenRootPath(InputPath)
-	, StorageServerPlatformFile(nullptr)
+	, StorageServerPlatformFile(IStorageServerClientModule::Get().TryCreateCustomPlatformFile(*InputPath, &FPlatformFileManager::Get().GetPlatformFile()))
 {
+	if (StorageServerPlatformFile)
+	{
+		StorageServerPlatformFile->SetLowerLevel(nullptr);
+	}
 }
 #else
 {
@@ -29,15 +33,14 @@ FZenCookArtifactReader::~FZenCookArtifactReader()
 {
 }
 
-void FZenCookArtifactReader::Initialize()
+void FZenCookArtifactReader::Initialize(bool bCleanBuild)
 {
 #if !UE_BUILD_SHIPPING
-	StorageServerPlatformFile.Reset(IStorageServerClientModule::Get().TryCreateCustomPlatformFile(*ZenRootPath, &FPlatformFileManager::Get().GetPlatformFile()));
-	if (StorageServerPlatformFile)
+	if (bCleanBuild && StorageServerPlatformFile)
 	{
-		StorageServerPlatformFile->SetLowerLevel(nullptr);
+		StorageServerPlatformFile->UpdateFileList();
 	}
-#endif // !UE_BUILD_SHIPPING
+#endif
 }
 
 bool FZenCookArtifactReader::FileExists(const TCHAR* Filename)
diff --git a/Engine/Source/Developer/IoStoreUtilities/Public/CookArtifactReader.h b/Engine/Source/Developer/IoStoreUtilities/Public/CookArtifactReader.h
--- a/Engine/Source/Developer/IoStoreUtilities/Public/CookArtifactReader.h
+++ b/Engine/Source/Developer/IoStoreUtilities/Public/CookArtifactReader.h
@@ -13,7 +13,7 @@ class ICookArtifactReader
 public:
 	virtual ~ICookArtifactReader() = default;
 
-	virtual void Initialize() {}
+	virtual void Initialize(bool bCleanBuild) {}
 	
 	virtual bool FileExists(const TCHAR* Filename) = 0;
 	virtual int64 FileSize(const TCHAR* Filename) = 0;
diff --git a/Engine/Source/Developer/IoStoreUtilities/Public/LayeredCookArtifactReader.h b/Engine/Source/Developer/IoStoreUtilities/Public/LayeredCookArtifactReader.h
--- a/Engine/Source/Developer/IoStoreUtilities/Public/LayeredCookArtifactReader.h
+++ b/Engine/Source/Developer/IoStoreUtilities/Public/LayeredCookArtifactReader.h
@@ -12,7 +12,7 @@ class FLayeredCookArtifactReader
 public:
 	FLayeredCookArtifactReader() = default;
 
-	void Initialize() override { for (TSharedRef<ICookArtifactReader>& Layer : Layers) { Layer->Initialize(); } }
+	void Initialize(bool bCleanBuild) override { for (TSharedRef<ICookArtifactReader>& Layer : Layers) { Layer->Initialize(bCleanBuild); } }
 
 	IOSTOREUTILITIES_API void AddLayer(TSharedRef<ICookArtifactReader> InLayer);
 	IOSTOREUTILITIES_API bool RemoveLayer(TSharedRef<ICookArtifactReader> InLayer);
diff --git a/Engine/Source/Developer/IoStoreUtilities/Public/ZenCookArtifactReader.h b/Engine/Source/Developer/IoStoreUtilities/Public/ZenCookArtifactReader.h
--- a/Engine/Source/Developer/IoStoreUtilities/Public/ZenCookArtifactReader.h
+++ b/Engine/Source/Developer/IoStoreUtilities/Public/ZenCookArtifactReader.h
@@ -18,7 +18,7 @@ public:
 												const FString& MetadataDirectoryPath, 
 												const ITargetPlatform* TargetPlatform);
 	IOSTOREUTILITIES_API virtual ~FZenCookArtifactReader();
-	IOSTOREUTILITIES_API void Initialize() override;
+	IOSTOREUTILITIES_API void Initialize(bool bCleanBuild) override;
 
 	IOSTOREUTILITIES_API bool FileExists(const TCHAR* Filename) override;
 	IOSTOREUTILITIES_API int64 FileSize(const TCHAR* Filename) override;
diff --git a/Engine/Source/Editor/UnrealEd/Private/CookOnTheFlyServer.cpp b/Engine/Source/Editor/UnrealEd/Private/CookOnTheFlyServer.cpp
--- a/Engine/Source/Editor/UnrealEd/Private/CookOnTheFlyServer.cpp
+++ b/Engine/Source/Editor/UnrealEd/Private/CookOnTheFlyServer.cpp
@@ -10463,7 +10463,8 @@ void UCookOnTheFlyServer::BeginCookSandbox(FBeginCookContext& BeginContext)
 			SavePackageContext.PackageWriterCapabilities = PackageWriter.GetCookCapabilities();
 
 			ICookArtifactReader& Reader = *SavePackageContext.ArtifactReader;
-			Reader.Initialize();
+			bool bCleanBuild = CookInfo.bFullBuild && !CookInfo.bWorkerOnSharedSandbox;
+			Reader.Initialize(bCleanBuild);
 
 			if (!PlatformContext.bWorkerOnSharedSandbox)
 			{
diff --git a/Engine/Source/Runtime/StorageServerClient/Internal/IStorageServerPlatformFile.h b/Engine/Source/Runtime/StorageServerClient/Internal/IStorageServerPlatformFile.h
--- a/Engine/Source/Runtime/StorageServerClient/Internal/IStorageServerPlatformFile.h
+++ b/Engine/Source/Runtime/StorageServerClient/Internal/IStorageServerPlatformFile.h
@@ -20,4 +20,5 @@ public:
 	virtual FStringView GetHostAddr() const = 0;
 
 	virtual void GetAndResetConnectionStats(FConnectionStats& OutStats) = 0;
+	virtual void UpdateFileList() = 0;
 };
diff --git a/Engine/Source/Runtime/StorageServerClient/Private/StorageServerPlatformFile.cpp b/Engine/Source/Runtime/StorageServerClient/Private/StorageServerPlatformFile.cpp
--- a/Engine/Source/Runtime/StorageServerClient/Private/StorageServerPlatformFile.cpp
+++ b/Engine/Source/Runtime/StorageServerClient/Private/StorageServerPlatformFile.cpp
@@ -122,6 +122,21 @@ FStorageServerFileSystemTOC::FDirectory* FStorageServerFileSystemTOC::AddDirecto
 	return Directory;
 }
 
+void FStorageServerFileSystemTOC::Clear()
+{
+	FWriteScopeLock _(TocLock);
+	for (auto& KV : Directories)
+	{
+		delete KV.Value;
+	}
+
+	Root.Directories.Empty();
+	Root.Files.Empty();
+	Files.Empty();
+	Directories.Empty();
+	FilePathToIndexMap.Empty();
+}
+
 void FStorageServerFileSystemTOC::AddFile(const FIoChunkId& FileChunkId, FStringView PathView, int64 RawSize)
 {
 	FWriteScopeLock _(TocLock);
@@ -1411,6 +1426,21 @@ void FStorageServerPlatformFile::InitializeConnection()
 	}
 }
 
+void FStorageServerPlatformFile::UpdateFileList()
+{
+	if (!Connection)
+	{
+		return;
+	}
+
+	ServerToc.Clear();
+	if (!SendGetFileListMessage())
+	{
+		FStringView HostAddr = Connection->GetHostAddr();
+		UE_LOG(LogStorageServerPlatformFile, Fatal, TEXT("Failed to get file list from Zen at '%.*s'"), HostAddr.Len(), HostAddr.GetData());
+	}
+}
+
 bool FStorageServerPlatformFile::FileExists(const TCHAR* Filename)
 {
 	TStringBuilder<1024> StorageServerFilename;
diff --git a/Engine/Source/Runtime/StorageServerClient/Private/StorageServerPlatformFile.h b/Engine/Source/Runtime/StorageServerClient/Private/StorageServerPlatformFile.h
--- a/Engine/Source/Runtime/StorageServerClient/Private/StorageServerPlatformFile.h
+++ b/Engine/Source/Runtime/StorageServerClient/Private/StorageServerPlatformFile.h
@@ -30,6 +30,7 @@ class FStorageServerFileSystemTOC
 public:
 	~FStorageServerFileSystemTOC();
 	void AddFile(const FIoChunkId& FileChunkId, FStringView Path, int64 RawSize);
+	void Clear();
 	bool FileExists(const FString& Path);
 	bool DirectoryExists(const FString& Path);
 	const FIoChunkId* GetFileChunkId(const FString& Path);
@@ -125,6 +126,7 @@ public:
 	{
 		CustomProjectStorePath = InProjectStorePath;
 	}
+	void UpdateFileList() override;
 private:
 	friend class FStorageServerFileHandle;
 	template<typename ParentVisitorClass, typename DataType> friend class FUniqueDirectoryStatVisitor;
