UnrealPak: Fix incorrect tracking of UsedBufferMemory after change 44847596 that allows the process to continue on failed source file reads.
UsedBufferMemory was increased when a read started but never decremented on ZenServer read failures which could theoretically throttle subsequent requests and risk starvation (in practice unlikely since this would require either an extreme amount of requests failing or a few very huge ones).

#rnx

diff --git a/Engine/Source/Developer/IoStoreUtilities/Private/IoStoreUtilities.cpp b/Engine/Source/Developer/IoStoreUtilities/Private/IoStoreUtilities.cpp
--- a/Engine/Source/Developer/IoStoreUtilities/Private/IoStoreUtilities.cpp
+++ b/Engine/Source/Developer/IoStoreUtilities/Private/IoStoreUtilities.cpp
@@ -3392,6 +3392,7 @@ private:
 		{
 			CompletionEvent.Emplace(InCompletionEvent);
 			Manager.ScheduleLoad(this);
+			bLoadScheduled = true;
 		}
 
 		virtual const FIoBuffer* GetSourceBuffer() override
@@ -3401,9 +3402,10 @@ private:
 		
 		virtual void FreeSourceBuffer() override
 		{
-			bSourceBufferValid = false;
-			if (SourceBuffer.DataSize() > 0)
+			if (bLoadScheduled)
 			{
+				bLoadScheduled = false;
+				bSourceBufferValid = false;
 				SourceBuffer = FIoBuffer();
 				Manager.OnBufferMemoryFreed(SourceBufferSize);
 			}
@@ -3446,6 +3448,7 @@ private:
 		TOptional<UE::Tasks::FTaskEvent> CompletionEvent;
 		FIoBuffer SourceBuffer;
 		FQueueEntry* QueueEntry = nullptr;
+		bool bLoadScheduled = false;
 		bool bSourceBufferValid = false;
 	};
 
