Improved resiliance to zenserver crashes in UE
Improved log diagnostics when connection to zenserver fails and recovery is triggered

#rnx

diff --git a/Engine/Source/Developer/DerivedDataCache/Private/ZenCacheStore.cpp b/Engine/Source/Developer/DerivedDataCache/Private/ZenCacheStore.cpp
--- a/Engine/Source/Developer/DerivedDataCache/Private/ZenCacheStore.cpp
+++ b/Engine/Source/Developer/DerivedDataCache/Private/ZenCacheStore.cpp
@@ -1564,17 +1564,16 @@ private:
 		return &BaseReceiver;
 	}
 
-	bool ShouldRecoverAndRetry(IHttpResponse& LocalResponse)
+	bool ShouldRetry(IHttpResponse& LocalResponse)
 	{
-		if (!ZenServiceInstance.IsServiceRunningLocally())
-		{
-			return false;
-		}
-
 		if ((LocalResponse.GetErrorCode() == EHttpErrorCode::Connect) ||
 			(LocalResponse.GetErrorCode() == EHttpErrorCode::TlsConnect) ||
 			(LocalResponse.GetErrorCode() == EHttpErrorCode::TimedOut))
 		{
+			if (ZenServiceInstance.IsServiceRunningLocally())
+			{
+				return ZenServiceInstance.TryRecovery(Zen::FZenServiceInstance::ERecoveryMode::Soft);
+			}
 			return true;
 		}
 
@@ -1585,7 +1584,7 @@ private:
 	{
 		Owner->End(this, [Self = this, &LocalResponse]
 		{
-			if (Self->ShouldRecoverAndRetry(LocalResponse) && Self->ZenServiceInstance.TryRecovery())
+			if (Self->ShouldRetry(LocalResponse))
 			{
 				new FAsyncHealthReceiver(MoveTemp(Self->Request), Self->Owner, Self->ZenServiceInstance, MoveTemp(Self->OnHealthComplete));
 				return;
@@ -1639,7 +1638,8 @@ public:
 				InnerOnComplete(Self->Response, Self->Package, Self->GetPayloadAsString());
 			});
 		},
-		/*MaxAttempts*/ 3)
+		/*MaxAttempts*/ 3,
+		/*UseSoftRecovery*/ true)
 		, Owner(InOwner)
 	{
 	}
@@ -2364,4 +2364,4 @@ ILegacyCacheStore* CreateZenCacheStore(const TCHAR* NodeName, const TCHAR* Confi
 
 } // UE::DerivedData
 
-#endif // UE_WITH_ZEN
\ No newline at end of file
+#endif // UE_WITH_ZEN
diff --git a/Engine/Source/Developer/IoStoreUtilities/Internal/ZenStoreHttpClient.h b/Engine/Source/Developer/IoStoreUtilities/Internal/ZenStoreHttpClient.h
--- a/Engine/Source/Developer/IoStoreUtilities/Internal/ZenStoreHttpClient.h
+++ b/Engine/Source/Developer/IoStoreUtilities/Internal/ZenStoreHttpClient.h
@@ -86,14 +86,15 @@ public:
 	static UE_API const UTF8CHAR* FindAttachmentId(FUtf8StringView AttachmentText);
 
 private:
-	UE_API TIoStatusOr<FIoBuffer> ReadOpLogUri(FStringBuilderBase& ChunkUri, uint64 Offset = 0, uint64 Size = ~0ull);
-	UE_API bool Download(Zen::FZenHttpRequest& Request, FStringView Uri, TArray64<uint8>* Buffer, Zen::EContentType AcceptType);
-	UE_API bool Post(Zen::FZenHttpRequest& Request, FStringView Uri, FCbObjectView Obj);
-	UE_API bool Post(Zen::FZenHttpRequest& Request, FStringView Uri, FMemoryView Payload);
-	UE_API void Rpc(FAnsiStringView Uri, FCbObject RpcRequestObject, FCbPackage& OutResponse);
-	UE_API void RpcAsync(FAnsiStringView Uri, FCbObject RpcRequestObject, TUniqueFunction<void(const FCbPackage& Response)>&& ResponseCallback);
-	UE_API bool Delete(Zen::FZenHttpRequest& Request, FStringView Uri);
-	UE_API bool ShouldRecoverAndRetry(Zen::FZenHttpRequest& Request);
+	TIoStatusOr<FIoBuffer> ReadOpLogUri(FStringBuilderBase& ChunkUri, uint64 Offset = 0, uint64 Size = ~0ull);
+	bool Download(Zen::FZenHttpRequest& Request, FStringView Uri, TArray64<uint8>* Buffer, Zen::EContentType AcceptType);
+	bool Post(Zen::FZenHttpRequest& Request, FStringView Uri, FCbObjectView Obj);
+	bool Post(Zen::FZenHttpRequest& Request, FStringView Uri, FMemoryView Payload);
+	void Rpc(FAnsiStringView Uri, FCbObject RpcRequestObject, FCbPackage& OutResponse);
+	void RpcAsync(FAnsiStringView Uri, FCbObject RpcRequestObject, TUniqueFunction<void(const FCbPackage& Response)>&& ResponseCallback);
+	bool Delete(Zen::FZenHttpRequest& Request, FStringView Uri);
+	bool ShouldRetry(Zen::FZenHttpRequest& Request);
+	bool TryRecover();
 	void SetRequestUri(class IHttpRequest& Request, FAnsiStringView Path);
 
 	template<bool bAsync, typename T>
diff --git a/Engine/Source/Developer/IoStoreUtilities/Private/ZenStoreHttpClient.cpp b/Engine/Source/Developer/IoStoreUtilities/Private/ZenStoreHttpClient.cpp
--- a/Engine/Source/Developer/IoStoreUtilities/Private/ZenStoreHttpClient.cpp
+++ b/Engine/Source/Developer/IoStoreUtilities/Private/ZenStoreHttpClient.cpp
@@ -147,20 +147,29 @@ FZenStoreHttpClient::TryCreateProject(FStringView InProjectId,
 
 			if (!Post(*Request, ProjectUri, ProjInfo.Save().AsObject()))
 			{
-				UE_LOG(LogZenStore, Display, TEXT("Zen project '%s' creation FAILED"), *FString(InProjectId));
+				UE_LOG(LogZenStore, Warning, TEXT("Zen project '%s' creation FAILED: %s"), *FString(InProjectId), *Request->GetDetailedResultString(ProjectUri));
 				bConnectionSucceeded = false;
 
 				// TODO: how to recover / handle this?
 			}
-			else if (Request->GetResponseCode() == 201)
-			{
-				UE_LOG(LogZenStore, Display, TEXT("Zen project '%s' created"), *FString(InProjectId));
-				bConnectionSucceeded = true;
-			}
 			else
 			{
-				UE_LOG(LogZenStore, Warning, TEXT("Zen project '%s' creation returned success but not HTTP 201"), *FString(InProjectId));
-				bConnectionSucceeded = true;
+				const int ResponseCode = Request->GetResponseCode();
+				if (ResponseCode == 201)
+				{
+					UE_LOG(LogZenStore, Display, TEXT("Zen project '%s' created"), *FString(InProjectId));
+					bConnectionSucceeded = true;
+				}
+				else if (200 <= ResponseCode && ResponseCode < 300)
+				{
+					UE_LOG(LogZenStore, Warning, TEXT("Zen project '%s' creation returned success but unexpected result: %s"), *FString(InProjectId), *Request->GetDetailedResultString(ProjectUri));
+					bConnectionSucceeded = true;
+				}
+				else
+				{
+					UE_LOG(LogZenStore, Warning, TEXT("Zen project '%s' creation FAILED: %s"), *FString(InProjectId), *Request->GetDetailedResultString(ProjectUri));
+					bConnectionSucceeded = false;
+				}
 			}
 		}
 	}
@@ -183,10 +192,14 @@ bool FZenStoreHttpClient::TryDeleteOplog(FStringView InProjectId, FStringView In
 	{
 		return false;
 	}
-	UE::Zen::FZenScopedRequestPtr Request(RequestPool.Get());
+	UE::Zen::FZenScopedRequestPtr Request(RequestPool.Get(), false /* bLogErrors */);
 
 	UE_LOG(LogZenStore, Display, TEXT("Deleting oplog '%s/%s' if it exists"), *FString(InProjectId), *FString(InOplogId));
 	bool Success = Delete(*Request, OplogPath);
+	if (!Success && Request->GetResponseCode() != 404)
+	{
+		UE_LOG(LogZenStore, Warning, TEXT("Deleting oplog FAILED: %s"), *Request->GetDetailedResultString(OplogPath));
+	}
 	return Success;
 }
 
@@ -197,7 +210,7 @@ bool FZenStoreHttpClient::TryCreateOplog(FStringView InProjectId, FStringView In
 		return false;
 	}
 
-	UE::Zen::FZenScopedRequestPtr Request(RequestPool.Get());
+	UE::Zen::FZenScopedRequestPtr Request(RequestPool.Get(), false /* bLogErrors */);
 
 	TArray64<uint8> GetBuffer;
 	FCbObjectView OplogInfo;
@@ -220,18 +233,29 @@ bool FZenStoreHttpClient::TryCreateOplog(FStringView InProjectId, FStringView In
 
 		if (!Post(*Request, OplogPath, OplogCreateInfo.AsObjectView()))
 		{
-			UE_LOG(LogZenStore, Error, TEXT("Zen oplog '%s/%s' creation FAILED"), *FString(InProjectId), *FString(InOplogId));
+			UE_LOG(LogZenStore, Warning, TEXT("Zen oplog '%s/%s' creation FAILED: %s"), *FString(InProjectId), *FString(InOplogId), *Request->GetDetailedResultString(OplogPath));
 			// Demote the connection status back to not connected
 			bConnectionSucceeded = false;
 			return false;
 		}
-		else if (Request->GetResponseCode() == 201)
-		{
-			UE_LOG(LogZenStore, Display, TEXT("Zen oplog '%s/%s' created"), *FString(InProjectId), *FString(InOplogId));
-		}
 		else
 		{
-			UE_LOG(LogZenStore, Warning, TEXT("Zen oplog '%s/%s' creation returned success but not HTTP 201"), *FString(InProjectId), *FString(InOplogId));
+			const int ResponseCode = Request->GetResponseCode();
+			if (ResponseCode == 201)
+			{
+				UE_LOG(LogZenStore, Display, TEXT("Zen oplog '%s/%s' created"), *FString(InProjectId), *FString(InOplogId));
+			}
+			else if (200 <= ResponseCode && ResponseCode < 300)
+			{
+				UE_LOG(LogZenStore, Warning, TEXT("Zen oplog '%s/%s' creation returned success but unexpected result: %s"), *FString(InProjectId), *FString(InOplogId), *Request->GetDetailedResultString(OplogPath));
+			}
+			else
+			{
+				UE_LOG(LogZenStore, Warning, TEXT("Zen oplog '%s/%s' creation FAILED: %s"), *FString(InProjectId), *FString(InOplogId), *Request->GetDetailedResultString(OplogPath));
+				// Demote the connection status back to not connected
+				bConnectionSucceeded = false;
+				return false;
+			}
 		}
 
 		// Issue another GET to retrieve information
@@ -242,6 +266,12 @@ bool FZenStoreHttpClient::TryCreateOplog(FStringView InProjectId, FStringView In
 		{
 			OplogInfo = FCbObjectView(GetBuffer.GetData());
 		}
+		else
+		{
+			UE_LOG(LogZenStore, Warning, TEXT("Zen oplog '%s/%s' get FAILED: %s"), *FString(InProjectId), *FString(InOplogId), *Request->GetDetailedResultString(OplogPath));
+			bConnectionSucceeded = false;
+			return false;
+		}
 	}
 
 	TempDirPath = FUTF8ToTCHAR(OplogInfo["tempdir"].AsString());
@@ -252,21 +282,22 @@ void FZenStoreHttpClient::InitializeReadOnly(FStringView InProjectId, FStringVie
 {
 	TRACE_CPUPROFILER_EVENT_SCOPE(ZenStoreHttp_InitializeReadOnly);
 
-	UE_LOG(LogZenStore, Display, TEXT("Establishing oplog %s / %s"), *FString(InProjectId), *FString(InOplogId));
+	UE_LOG(LogZenStore, Display, TEXT("Opening oplog %s / %s"), *FString(InProjectId), *FString(InOplogId));
 
 	// Establish project
 
+	ProjectPath = WriteToString<128>("/prj/", InProjectId);
 	{
-		UE::Zen::FZenScopedRequestPtr Request(RequestPool.Get());
+		UE::Zen::FZenScopedRequestPtr Request(RequestPool.Get(), false /* bLogErrors */);
 
 		TArray64<uint8> GetBuffer;
 
 		// TODO: how to handle failure here? This is probably the most likely point of failure
 		// if the service is not up or not responding
 
-		if (!Download(*Request, WriteToString<128>("/prj/", InProjectId), &GetBuffer, Zen::EContentType::CbObject) || Request->GetResponseCode() != 200)
+		if (!Download(*Request, ProjectPath, &GetBuffer, Zen::EContentType::CbObject) || Request->GetResponseCode() != 200)
 		{
-			UE_LOG(LogZenStore, Fatal, TEXT("Zen project '%s' not found"), *FString(InProjectId));
+			UE_LOG(LogZenStore, Fatal, TEXT("Zen project '%s' get FAILED: %s"), *FString(InProjectId), *Request->GetDetailedResultString(ProjectPath));
 			bConnectionSucceeded = false;
 		}
 		else
@@ -279,15 +310,15 @@ void FZenStoreHttpClient::InitializeReadOnly(FStringView InProjectId, FStringVie
 	check(IsConnected());
 
 	{
-		UE::Zen::FZenScopedRequestPtr Request(RequestPool.Get());
+		UE::Zen::FZenScopedRequestPtr Request(RequestPool.Get(), false /* bLogErrors */);
 
 		OplogPath = WriteToString<128>("/prj/", InProjectId, "/oplog/", InOplogId);
-		ProjectPath = WriteToString<128>("/prj/", InProjectId);
 
 		TArray64<uint8> GetBuffer;
 		if (!Download(*Request, OplogPath, &GetBuffer, Zen::EContentType::CbObject) || Request->GetResponseCode() != 200)
 		{
-			UE_LOG(LogZenStore, Fatal, TEXT("Zen oplog '%s'/'%s' not found"), *FString(InProjectId), *FString(InOplogId));
+			UE_LOG(LogZenStore, Fatal, TEXT("Zen oplog '%s/%s' get FAILED: %s"), *FString(InProjectId), *FString(InOplogId), *Request->GetDetailedResultString(OplogPath));
+			bConnectionSucceeded = false;
 		}
 	}
 
@@ -297,7 +328,7 @@ void FZenStoreHttpClient::InitializeReadOnly(FStringView InProjectId, FStringVie
 bool FZenStoreHttpClient::Download(Zen::FZenHttpRequest& Request, FStringView Uri, TArray64<uint8>* Buffer, Zen::EContentType AcceptType)
 {
 	Zen::FZenHttpRequest::Result Res = Request.PerformBlockingDownload(Uri, Buffer, AcceptType);
-	if (ShouldRecoverAndRetry(Request))
+	if (ShouldRetry(Request))
 	{
 		Request.Reset();
 		Res = Request.PerformBlockingDownload(Uri, Buffer, AcceptType);
@@ -308,7 +339,7 @@ bool FZenStoreHttpClient::Download(Zen::FZenHttpRequest& Request, FStringView Ur
 bool FZenStoreHttpClient::Post(Zen::FZenHttpRequest& Request, FStringView Uri, FCbObjectView Obj)
 {
 	Zen::FZenHttpRequest::Result Res = Request.PerformBlockingPost(Uri, Obj);
-	if (ShouldRecoverAndRetry(Request))
+	if (ShouldRetry(Request))
 	{
 		Request.Reset();
 		Res = Request.PerformBlockingPost(Uri, Obj);
@@ -319,7 +350,7 @@ bool FZenStoreHttpClient::Post(Zen::FZenHttpRequest& Request, FStringView Uri, F
 bool FZenStoreHttpClient::Post(Zen::FZenHttpRequest& Request, FStringView Uri, FMemoryView Payload)
 {
 	Zen::FZenHttpRequest::Result Res = Request.PerformBlockingPost(Uri, Payload);
-	if (ShouldRecoverAndRetry(Request))
+	if (ShouldRetry(Request))
 	{
 		Request.Reset();
 		Res = Request.PerformBlockingPost(Uri, Payload);
@@ -373,7 +404,7 @@ void FZenStoreHttpClient::Rpc(FAnsiStringView Uri, FCbObject RpcRequestObject, F
 		Request->Send(&Receiver, HttpResponse);
 		int StatusCode = HttpResponse->GetStatusCode();
 
-		bRetry = Receiver.ShouldRecoverAndRetry(ZenService.GetInstance(), *HttpResponse) && (++Attempt < MaxAttempts);
+		bRetry = Receiver.ShouldRetry(ZenService.GetInstance(), *HttpResponse, /*UseSoftRecovery*/ false) && (++Attempt < MaxAttempts);
 	}
 	while (bRetry);
 }
@@ -405,14 +436,15 @@ void FZenStoreHttpClient::RpcAsync(FAnsiStringView Uri, FCbObject RpcRequestObje
 				delete Receiver;
 			});
 		},
-		/*MaxAttempts*/ 2);
+		/*MaxAttempts*/ 2,
+		/*UseSoftRecovery */ false);
 	PackageReceiver->SendAsync();
 } //-V773
 
 bool FZenStoreHttpClient::Delete(Zen::FZenHttpRequest& Request, FStringView Uri)
 {
 	Zen::FZenHttpRequest::Result Res = Request.PerformBlockingDelete(Uri);
-	if (ShouldRecoverAndRetry(Request))
+	if (ShouldRetry(Request))
 	{
 		Request.Reset();
 		Res = Request.PerformBlockingDelete(Uri);
@@ -420,27 +452,41 @@ bool FZenStoreHttpClient::Delete(Zen::FZenHttpRequest& Request, FStringView Uri)
 	return Res == Zen::FZenHttpRequest::Result::Success;
 }
 
-bool FZenStoreHttpClient::ShouldRecoverAndRetry(Zen::FZenHttpRequest& Request)
+bool FZenStoreHttpClient::ShouldRetry(Zen::FZenHttpRequest& Request)
 {
-	if (!ZenService.GetInstance().IsServiceRunningLocally())
-	{
-		return false;
-	}
-
 	const int ConnectError = 7;
 	const int SSLConnectError = 35;
 	const int TimeoutError = 28;
-
-	if ((Request.GetResponseCode() == ConnectError) ||
-		(Request.GetResponseCode() == SSLConnectError) ||
-		(Request.GetResponseCode() == TimeoutError))
+	const int SendError = 55;
+	const int ReceiveError = 56;
+
+	if ((Request.GetErrorCode() == ConnectError) ||
+		(Request.GetErrorCode() == SSLConnectError) ||
+		(Request.GetErrorCode() == TimeoutError) ||
+		(Request.GetErrorCode() == SendError) ||
+		(Request.GetErrorCode() == ReceiveError))
+	{
+		return TryRecover();
+	}
+	else if (Request.GetResponseCode() == 404 && Request.GetResponseAsString().Contains("No suitable route found"))
 	{
-		return true;
+		// This is a special case if we happen to hit the port while zenserver is starting up - it is worth a retry
+		return TryRecover();
 	}
 
 	return false;
 }
 
+bool FZenStoreHttpClient::TryRecover()
+{
+	Zen::FZenServiceInstance& ZenInstance = ZenService.GetInstance();
+	if (ZenInstance.IsServiceRunningLocally())
+	{
+		return ZenInstance.TryRecovery(Zen::FZenServiceInstance::ERecoveryMode::Hard);
+	}
+	return true;
+}
+
 TIoStatusOr<uint64> FZenStoreHttpClient::AppendOp(FCbPackage OpEntry)
 {
 	check(bAllowEdit);
@@ -454,7 +500,7 @@ TIoStatusOr<uint64> FZenStoreHttpClient::AppendOp(FCbPackage OpEntry)
 		const int32 Salt = SaltGen.Next();
 		bool bIsUsingTempFiles = false;
 
-		UE::Zen::FZenScopedRequestPtr Request(RequestPool.Get());
+		UE::Zen::FZenScopedRequestPtr Request(RequestPool.Get(), false /* bLogErrors */);
 
 		if (TempDirPath.IsEmpty())
 		{
@@ -486,26 +532,19 @@ TIoStatusOr<uint64> FZenStoreHttpClient::AppendOp(FCbPackage OpEntry)
 
 				FCbFieldIterator Prep = Writer.Save();
 
-				bool IsOk = false;			
-				if (Post(*Request, OplogPrepNewEntryPath, Prep.AsObjectView()))
+				if (Post(*Request, OplogPrepNewEntryPath, Prep.AsObjectView()) && Request->GetResponseCode() == 200)
 				{
-					FCbObjectView NeedObject;
+					FCbObjectView NeedObject = FCbObjectView(Request->GetResponseBuffer().GetData());
 
-					if (Request->GetResponseCode() == 200)
+					for (auto& Entry : NeedObject["need"])
 					{
-						NeedObject = FCbObjectView(Request->GetResponseBuffer().GetData());
-
-						for (auto& Entry : NeedObject["need"])
-						{
-							NeedChunks.Add(Entry.AsHash());
-						}
-
-						IsOk = true;
+						NeedChunks.Add(Entry.AsHash());
 					}
 				}
-
-				if (!IsOk)
+				else
 				{
+					UE_LOG(LogZenStore, Warning, TEXT("Oplog prep FAILED: %s"), *Request->GetDetailedResultString(OplogPrepNewEntryPath));
+
 					UE_LOG(LogZenStore, Warning, TEXT("Unable to check for needed chunks, assuming all are needed"));
 					for (const FCbAttachment& Attachment : Attachments)
 					{
@@ -605,12 +644,13 @@ TIoStatusOr<uint64> FZenStoreHttpClient::AppendOp(FCbPackage OpEntry)
 
 		Request->Reset();
 
-		if (Post(*Request, NewOpPostUri, SerializedPackage.GetView()))
+		if (Post(*Request, NewOpPostUri, SerializedPackage.GetView()) && (Request->GetResponseCode() == 200 || Request->GetResponseCode() == 201))
 		{
 			return TIoStatusOr<uint64>(SerializedPackage.TotalSize());
 		}
 		else
 		{
+			UE_LOG(LogZenStore, Warning, TEXT("Oplog write FAILED: %s"), *Request->GetDetailedResultString(OplogPrepNewEntryPath));
 			return TIoStatusOr<uint64>((FIoStatus)(FIoStatusBuilder(EIoErrorCode::Unknown) << TEXT("Append OpLog failed, NewOpLogPath='") << OplogNewEntryPath << TEXT("'")));
 		}
 	}
@@ -622,7 +662,7 @@ TIoStatusOr<uint64> FZenStoreHttpClient::GetChunkSize(const FIoChunkId& Id)
 
 	check(bAllowRead);
 
-	UE::Zen::FZenScopedRequestPtr Request(RequestPool.Get());
+	UE::Zen::FZenScopedRequestPtr Request(RequestPool.Get(), false /* bLogErrors */);
 	TStringBuilder<128> ChunkUri;
 	ChunkUri << OplogPath << '/' << Id << "/info";
 	if (Download(*Request, ChunkUri, nullptr, Zen::EContentType::CbObject) && Request->GetResponseCode() == 200)
@@ -631,9 +671,9 @@ TIoStatusOr<uint64> FZenStoreHttpClient::GetChunkSize(const FIoChunkId& Id)
 		const uint64 ChunkSize = ResponseObj["size"].AsUInt64(0);
 		return ChunkSize;
 	}
-	else if (int ResponseCode = Request->GetResponseCode(); ResponseCode != 404)
+	if (int ResponseCode = Request->GetResponseCode(); ResponseCode != 404)
 	{
-		UE_LOG(LogZenStore, Verbose, TEXT("Request '%s' failed. Error: %s. Status: %s (%d). \"%s\""), *ChunkUri, *FString(Request->GetError()), *FString(UE::Zen::FZenHttpRequest::LexResponseCodeToString(ResponseCode)), ResponseCode, *Request->GetResponseAsString());
+		UE_LOG(LogZenStore, Verbose, TEXT("Get chunk request FAILED: %s"), *Request->GetDetailedResultString(ChunkUri));
 	}
 	return FIoStatus(EIoErrorCode::NotFound);
 }
@@ -827,7 +867,7 @@ TIoStatusOr<FIoBuffer> FZenStoreHttpClient::ReadOpLogUri(FStringBuilderBase& Chu
 {
 	check(bAllowRead);
 
-	UE::Zen::FZenScopedRequestPtr Request(RequestPool.Get());
+	UE::Zen::FZenScopedRequestPtr Request(RequestPool.Get(), false /* bLogErrors */);
 	TArray64<uint8> GetBuffer;
 
 	bool bHaveQuery = false;
@@ -894,15 +934,11 @@ TIoStatusOr<FIoBuffer> FZenStoreHttpClient::ReadOpLogUri(FStringBuilderBase& Chu
 	else
 	{
 		int ResponseCode = Request->GetResponseCode();
+		UE_LOG(LogZenStore, Log, TEXT("Get chunk request FAILED: %s"), *Request->GetDetailedResultString(ChunkUri));
 		if (ResponseCode == 404)
 		{
-			UE_LOG(LogZenStore, Log, TEXT("Request '%s' failed with %s (%d). \"%s\""), *ChunkUri, *FString(UE::Zen::FZenHttpRequest::LexResponseCodeToString(ResponseCode)), ResponseCode, *Request->GetResponseAsString());
 			return FIoStatus(EIoErrorCode::NotFound);
 		}
-		else
-		{
-			UE_LOG(LogZenStore, Log, TEXT("Request '%s' failed. Error: %s. Status: %s (%d). \"%s\""), *ChunkUri, *FString(Request->GetError()), *FString(UE::Zen::FZenHttpRequest::LexResponseCodeToString(ResponseCode)), ResponseCode, *Request->GetResponseAsString());
-		}
 	}
 	return FIoStatus(EIoErrorCode::ReadError);
 }
@@ -1084,15 +1120,16 @@ FZenStoreHttpClient::EndBuildPass(FCbPackage OpEntry)
 
 	UE_LOG(LogZenStore, Verbose, TEXT("Package size: %lld"), SerializedPackage.TotalSize());
 
-	UE::Zen::FZenScopedRequestPtr Request(RequestPool.Get());
+	UE::Zen::FZenScopedRequestPtr Request(RequestPool.Get(), false /* bLogErrors */);
 
 	FMemoryView Payload { SerializedPackage.GetData(), (uint64)SerializedPackage.TotalSize()};
-	if (Post(*Request, OplogNewEntryPath, Payload))
+	if (Post(*Request, OplogNewEntryPath, Payload) && (Request->GetResponseCode() == 200 || Request->GetResponseCode() == 201))
 	{
 		return static_cast<uint64>(Payload.GetSize());
 	}
 	else
 	{
+		UE_LOG(LogZenStore, Warning, TEXT("Oplog end FAILED: %s"), *Request->GetDetailedResultString(OplogNewEntryPath));
 		return (FIoStatus)(FIoStatusBuilder(EIoErrorCode::Unknown) << TEXT("End build pass failed, NewOpLogPath='") << OplogNewEntryPath << TEXT("'"));
 	}
 }
diff --git a/Engine/Source/Developer/IoStoreUtilities/Private/ZenStoreWriter.cpp b/Engine/Source/Developer/IoStoreUtilities/Private/ZenStoreWriter.cpp
--- a/Engine/Source/Developer/IoStoreUtilities/Private/ZenStoreWriter.cpp
+++ b/Engine/Source/Developer/IoStoreUtilities/Private/ZenStoreWriter.cpp
@@ -1023,7 +1023,7 @@ void FZenStoreWriter::CommitPackageInternal(FZenCommitInfo&& ZenCommitInfo)
 		}
 
 		TIoStatusOr<uint64> Status = HttpClient->AppendOp(MoveTemp(OplogEntry));
-		UE_CLOG(!Status.IsOk(), LogZenStoreWriter, Error, TEXT("Failed to commit oplog entry '%s' to Zen"), *CommitInfo.PackageName.ToString());
+		UE_CLOG(!Status.IsOk(), LogZenStoreWriter, Fatal, TEXT("Failed to commit oplog entry '%s' to Zen"), *CommitInfo.PackageName.ToString());
 	}
 	else if (bComputeHash)
 	{
diff --git a/Engine/Source/Developer/Zen/Internal/ZenAsyncCbPackageReceiver.h b/Engine/Source/Developer/Zen/Internal/ZenAsyncCbPackageReceiver.h
--- a/Engine/Source/Developer/Zen/Internal/ZenAsyncCbPackageReceiver.h
+++ b/Engine/Source/Developer/Zen/Internal/ZenAsyncCbPackageReceiver.h
@@ -28,22 +28,21 @@ public:
 		THttpUniquePtr<IHttpRequest>&& InRequest,
 		FZenServiceInstance& InZenServiceInstance,
 		FOnComplete&& InOnComplete,
-		int InMaxAttempts = 1);
+		int InMaxAttempts,
+		bool InUseSoftRecovery);
 
 	ZEN_API const IHttpResponse& GetHttpResponse() { return *Response; }
 	ZEN_API const FCbPackage& GetResponsePackage() { return Package; }
 
-
 	ZEN_API void SendAsync();
 
+	ZEN_API FString GetPayloadAsString() const;
 protected:
 	// IHttpReceiver Interface
 
 	ZEN_API IHttpReceiver* OnCreate(IHttpResponse& LocalResponse);
 	ZEN_API IHttpReceiver* OnComplete(IHttpResponse& LocalResponse);
 
-	ZEN_API FString GetPayloadAsString() const;
-
 protected:
 	THttpUniquePtr<IHttpRequest> Request;
 	THttpUniquePtr<IHttpResponse> Response;
@@ -52,6 +51,7 @@ protected:
 	UE::Zen::FCbPackageReceiver BaseReceiver;
 	FOnComplete OnCompleteCallback;
 	const int MaxAttempts;
+	const bool UseSoftRecovery;
 	int Attempt;
 };
 
diff --git a/Engine/Source/Developer/Zen/Internal/ZenCbPackageReceiver.h b/Engine/Source/Developer/Zen/Internal/ZenCbPackageReceiver.h
--- a/Engine/Source/Developer/Zen/Internal/ZenCbPackageReceiver.h
+++ b/Engine/Source/Developer/Zen/Internal/ZenCbPackageReceiver.h
@@ -27,12 +27,16 @@ public:
 	ZEN_API void Reset();
 	ZEN_API const FMemoryView Body() const;
 
-	ZEN_API static bool ShouldRecoverAndRetry(FZenServiceInstance& ZenServiceInstance, IHttpResponse& LocalResponse);
+	ZEN_API bool ShouldRetry(FZenServiceInstance& ZenServiceInstance, IHttpResponse& LocalResponse, bool UseSoftRecovery) const;
+
+	ZEN_API static FString GetPayloadAsString(EHttpMediaType ContentType, FMemoryView BodyView);
 
 private:
 	ZEN_API IHttpReceiver* OnCreate(IHttpResponse& Response) final;
 	ZEN_API IHttpReceiver* OnComplete(IHttpResponse& Response) final;
 
+	bool TryRecover(Zen::FZenServiceInstance& ZenServiceInstance, bool UseSoftRecovery) const;
+
 private:
 	FCbPackage& Package;
 	IHttpReceiver* Next;
diff --git a/Engine/Source/Developer/Zen/Internal/ZenServerHttp.h b/Engine/Source/Developer/Zen/Internal/ZenServerHttp.h
--- a/Engine/Source/Developer/Zen/Internal/ZenServerHttp.h
+++ b/Engine/Source/Developer/Zen/Internal/ZenServerHttp.h
@@ -107,8 +107,15 @@ public:
 		return int(ResponseCode);
 	}
 
+	inline const int GetErrorCode() const
+	{
+		return int(CurlResult);
+	}
+
 	ZEN_API FAnsiStringView GetError() const;
 
+	ZEN_API FString GetDetailedResultString(FStringView Uri) const;
+
 	inline const bool GetResponseFormatValid() const
 	{
 		return bResponseFormatValid;
diff --git a/Engine/Source/Developer/Zen/Private/ZenAsyncCbPackageReceiver.cpp b/Engine/Source/Developer/Zen/Private/ZenAsyncCbPackageReceiver.cpp
--- a/Engine/Source/Developer/Zen/Private/ZenAsyncCbPackageReceiver.cpp
+++ b/Engine/Source/Developer/Zen/Private/ZenAsyncCbPackageReceiver.cpp
@@ -16,12 +16,14 @@ FAsyncCbPackageReceiver::FAsyncCbPackageReceiver(
 	THttpUniquePtr<IHttpRequest>&& InRequest,
 	Zen::FZenServiceInstance& InZenServiceInstance,
 	FOnComplete&& InOnComplete,
-	int InMaxAttempts)
+	int InMaxAttempts,
+	bool InUseSoftRecovery)
 	: Request(MoveTemp(InRequest))
 	, ZenServiceInstance(InZenServiceInstance)
 	, BaseReceiver(Package, this)
 	, OnCompleteCallback(MoveTemp(InOnComplete))
 	, MaxAttempts(InMaxAttempts)
+	, UseSoftRecovery(InUseSoftRecovery)
 	, Attempt(0)
 {
 }
@@ -38,7 +40,7 @@ IHttpReceiver* FAsyncCbPackageReceiver::OnCreate(IHttpResponse& LocalResponse)
 
 IHttpReceiver* FAsyncCbPackageReceiver::OnComplete(IHttpResponse& LocalResponse)
 {
-	if ((++Attempt < MaxAttempts) && FCbPackageReceiver::ShouldRecoverAndRetry(ZenServiceInstance, LocalResponse) && ZenServiceInstance.TryRecovery())
+	if ((++Attempt < MaxAttempts) && BaseReceiver.ShouldRetry(ZenServiceInstance, LocalResponse, UseSoftRecovery))
 	{
 		BaseReceiver.Reset();
 		SendAsync();
@@ -59,27 +61,7 @@ IHttpReceiver* FAsyncCbPackageReceiver::OnComplete(IHttpResponse& LocalResponse)
 
 FString FAsyncCbPackageReceiver::GetPayloadAsString() const
 {
-	EHttpMediaType ContentType = Response->GetContentType();
-	FMemoryView BodyView = BaseReceiver.Body();
-	switch (ContentType)
-	{
-		case EHttpMediaType::Text:
-		case EHttpMediaType::Yaml:
-		case EHttpMediaType::Json:
-			if (BodyView.GetSize() > 0)
-			{
-				TStringBuilder<128> PayloadText;
-				PayloadText << FAnsiStringView((const ANSICHAR*)BodyView.GetData(), (int)(BodyView.GetSize()));
-				return *PayloadText;
-			}
-			return {};
-		default:
-		{
-			TStringBuilder<128> SizeText;
-			SizeText << "Payload (" << LexToString(ContentType) << "): " << BodyView.GetSize() << " bytes";
-			return *SizeText;
-		}
-	}
+	return BaseReceiver.GetPayloadAsString(Response->GetContentType(), BaseReceiver.Body());
 }
 
 }
diff --git a/Engine/Source/Developer/Zen/Private/ZenCbPackageReceiver.cpp b/Engine/Source/Developer/Zen/Private/ZenCbPackageReceiver.cpp
--- a/Engine/Source/Developer/Zen/Private/ZenCbPackageReceiver.cpp
+++ b/Engine/Source/Developer/Zen/Private/ZenCbPackageReceiver.cpp
@@ -28,23 +28,64 @@ const FMemoryView FCbPackageReceiver::Body() const
 	return MakeMemoryView(BodyArray);
 }
 
-bool FCbPackageReceiver::ShouldRecoverAndRetry(Zen::FZenServiceInstance& ZenServiceInstance, IHttpResponse& LocalResponse)
+bool FCbPackageReceiver::ShouldRetry(Zen::FZenServiceInstance& ZenServiceInstance, IHttpResponse& LocalResponse, bool UseSoftRecovery) const
 {
-	if (!ZenServiceInstance.IsServiceRunningLocally())
+	if (!UseSoftRecovery)
 	{
-		return false;
+		if (LocalResponse.GetErrorCode() == EHttpErrorCode::Unknown)
+		{
+			return TryRecover(ZenServiceInstance, UseSoftRecovery);
+		}
 	}
 
 	if ((LocalResponse.GetErrorCode() == EHttpErrorCode::Connect) ||
 		(LocalResponse.GetErrorCode() == EHttpErrorCode::TlsConnect) ||
 		(LocalResponse.GetErrorCode() == EHttpErrorCode::TimedOut))
 	{
-		return true;
+		return TryRecover(ZenServiceInstance, UseSoftRecovery);
+	}
+
+	if (LocalResponse.GetStatusCode() == 404 && GetPayloadAsString(LocalResponse.GetContentType(), Body()).Contains("No suitable route found"))
+	{
+		// This is a special case if we happen to hit the port while zenserver is starting up - it is worth a retry
+		return TryRecover(ZenServiceInstance, UseSoftRecovery);
 	}
 
 	return false;
 }
 
+bool FCbPackageReceiver::TryRecover(Zen::FZenServiceInstance& ZenServiceInstance, bool UseSoftRecovery) const
+{
+	if (ZenServiceInstance.IsServiceRunningLocally())
+	{
+		return ZenServiceInstance.TryRecovery(UseSoftRecovery ? Zen::FZenServiceInstance::ERecoveryMode::Soft : Zen::FZenServiceInstance::ERecoveryMode::Hard);
+	}
+	return true;
+}
+
+FString FCbPackageReceiver::GetPayloadAsString(EHttpMediaType ContentType, FMemoryView BodyView)
+{
+	switch (ContentType)
+	{
+	case EHttpMediaType::Text:
+	case EHttpMediaType::Yaml:
+	case EHttpMediaType::Json:
+		if (BodyView.GetSize() > 0)
+		{
+			TStringBuilder<128> PayloadText;
+			PayloadText << FAnsiStringView((const ANSICHAR*)BodyView.GetData(), (int)(BodyView.GetSize()));
+			return *PayloadText;
+		}
+		return {};
+	default:
+	{
+		TStringBuilder<128> SizeText;
+		SizeText << "Payload (" << LexToString(ContentType) << "): " << BodyView.GetSize() << " bytes";
+		return *SizeText;
+	}
+	}
+}
+
 IHttpReceiver* FCbPackageReceiver::OnCreate(IHttpResponse& Response)
 {
 	return &BodyReceiver;
diff --git a/Engine/Source/Developer/Zen/Private/ZenServerHttp.cpp b/Engine/Source/Developer/Zen/Private/ZenServerHttp.cpp
--- a/Engine/Source/Developer/Zen/Private/ZenServerHttp.cpp
+++ b/Engine/Source/Developer/Zen/Private/ZenServerHttp.cpp
@@ -576,6 +576,33 @@ namespace UE::Zen {
 		return GetAnsiBufferAsString(ResponseBuffer);
 	}
 
+	FString FZenHttpRequest::GetDetailedResultString(FStringView Uri) const
+	{
+		TStringBuilder<256> SB;
+		if (CurlResult != CURLE_OK)
+		{
+			SB << TEXTVIEW("Failed: ") << GetError();
+			SB << TEXTVIEW(" (") << CurlResult << TEXTVIEW(") ");
+		}
+		SB << TEXTVIEW("Url: ") << Uri;
+
+		if (ResponseCode != 0)
+		{
+			SB << TEXTVIEW(" Response: ") << LexResponseCodeToString(ResponseCode) << TEXTVIEW(" (") << ResponseCode << TEXTVIEW(")");
+		}
+
+		FString ResponseString = GetResponseAsString();
+		if (!ResponseString.IsEmpty())
+		{
+			SB << TEXTVIEW(" \"") << ResponseString << TEXTVIEW("\"");
+		}
+		else if (int64 PayloadSize = ResponseBuffer.Num(); PayloadSize > 0)
+		{
+			SB << TEXTVIEW(" ") << PayloadSize << TEXTVIEW(" bytes");
+		}
+		return SB.ToString();
+	}
+
 	FCbObjectView FZenHttpRequest::GetResponseAsObject() const
 	{
 		return FCbObjectView(ResponseBuffer.GetData());
diff --git a/Engine/Source/Developer/Zen/Private/ZenServerInterface.cpp b/Engine/Source/Developer/Zen/Private/ZenServerInterface.cpp
--- a/Engine/Source/Developer/Zen/Private/ZenServerInterface.cpp
+++ b/Engine/Source/Developer/Zen/Private/ZenServerInterface.cpp
@@ -1562,20 +1562,30 @@ static bool ShutdownZenServerProcess(int Pid, double MaximumWaitDurationSeconds
 				}
 			}
 		}
+		if (!ZenServerState::IsProcessRunning(Pid))
+		{
+			UE_LOG(LogZenServiceInstance, Display, TEXT("Successfully shut down zenserver process with pid %d"), Pid);
+			return true;
+		}
+	}
+	else if (!ZenServerState::IsProcessRunning(Pid))
+	{
+		UE_LOG(LogZenServiceInstance, Display, TEXT("Found no active process with pid %d"), Pid);
+		return true;
+	}
+
+	UE_LOG(LogZenServiceInstance, Display, TEXT("Attempting termination of process with pid %d"), Pid);
+	if (ZenServerState::Terminate(Pid))
+	{
+		UE_LOG(LogZenServiceInstance, Warning, TEXT("Terminated process with pid %d"), Pid);
+		return true;
 	}
 	if (ZenServerState::IsProcessRunning(Pid))
 	{
-		UE_LOG(LogZenServiceInstance, Display, TEXT("Attempting termination of zenserver process with pid %d"), Pid);
-		if (!ZenServerState::Terminate(Pid))
-		{
-			if (ZenServerState::IsProcessRunning(Pid))
-			{
-				UE_LOG(LogZenServiceInstance, Warning, TEXT("Failed to terminate zenserver process with pid %d"), Pid);
-				return false;
-			}
-		}
+		UE_LOG(LogZenServiceInstance, Warning, TEXT("Failed to terminate process with pid %d"), Pid);
+		return false;
 	}
-	UE_LOG(LogZenServiceInstance, Display, TEXT("Successfully shut down zenserver process with pid %d"), Pid);
+	UE_LOG(LogZenServiceInstance, Display, TEXT("No longer any active process with pid %d"), Pid);
 	return true;
 }
 
@@ -2270,11 +2280,11 @@ FZenServiceInstance::IsServiceRunning()
 bool
 FZenServiceInstance::IsServiceReady()
 {
-	return PingService(2);
+	return PingService(2, /*Silent*/ false);
 }
 
 bool 
-FZenServiceInstance::PingService(uint32_t AttemptCount)
+FZenServiceInstance::PingService(uint32_t AttemptCount, bool Silent)
 {
 	uint32 Attempt = 0;
 	while (IsServiceRunning())
@@ -2286,7 +2296,10 @@ FZenServiceInstance::PingService(uint32_t AttemptCount)
 		
 		if (Result == Zen::FZenHttpRequest::Result::Success && Zen::IsSuccessCode(Request.GetResponseCode()))
 		{
-			UE_LOGFMT(LogZenServiceInstance, Display, "Unreal Zen Storage Server HTTP service at {Domain} status: {Response}.", Endpoint.GetURL(), Request.GetResponseAsString());
+			if (!Silent)
+			{
+				UE_LOGFMT(LogZenServiceInstance, Display, "Unreal Zen Storage Server HTTP service at {Domain} status: {Response}.", Endpoint.GetURL(), Request.GetResponseAsString());
+			}
 			return true;
 		}
 
@@ -2294,13 +2307,19 @@ FZenServiceInstance::PingService(uint32_t AttemptCount)
 		{
 			if (Attempt == AttemptCount)
 			{
-				UE_LOGFMT(LogZenServiceInstance, Warning, "Unable to reach Unreal Zen Storage Server HTTP service at {Domain}. Status: {Code}. Response: {Response}", Endpoint.GetURL(), Request.GetResponseCode(), Request.GetResponseAsString());
+				if (!Silent)
+				{
+					UE_LOGFMT(LogZenServiceInstance, Warning, "Unable to reach Local Unreal Zen Storage Server HTTP service at {Domain}: {Result}", Endpoint.GetURL(), Request.GetDetailedResultString(Endpoint.GetURL()));
+				}
 				break;
 			}
 		}
 		else
 		{
-			UE_LOGFMT(LogZenServiceInstance, Display, "Unable to reach Unreal Zen Storage Server HTTP service at {Domain}. Status: {Code}. Response: {Response}", Endpoint.GetURL(), Request.GetResponseCode(), Request.GetResponseAsString());
+			if (!Silent)
+			{
+				UE_LOGFMT(LogZenServiceInstance, Display, "Unable to reach Unreal Zen Storage Server HTTP service at {Domain}: {Result}", Endpoint.GetURL(), Request.GetDetailedResultString(Endpoint.GetURL()));
+			}
 			break;
 		}
 	}
@@ -2308,13 +2327,19 @@ FZenServiceInstance::PingService(uint32_t AttemptCount)
 }
 
 bool 
-FZenServiceInstance::TryRecovery()
+FZenServiceInstance::TryRecovery(ERecoveryMode RecoveryMode)
 {
 	if (!bHasLaunchedLocal)
 	{
 		return false;
 	}
 
+	if (PingService(1, /*Silent*/ true))
+	{
+		UE_LOG(LogZenServiceInstance, Verbose, TEXT("Local ZenServer status: Healthy. Skipping recovery"));
+		return true;
+	}
+
 	const FString ExecutablePath = GetLocalServiceInstallPath();
 	const FString ExecutionContextFilePath = GetServiceRunContextPath();
 
@@ -2325,81 +2350,102 @@ FZenServiceInstance::TryRecovery()
 
 	FTimespan TimespanSinceLastRecovery = FDateTime::UtcNow() - FDateTime(LastRecoveryTicks.load(std::memory_order_relaxed));
 
-	if (TimespanSinceLastRecovery > MinimumDurationSinceLastRecovery)
+	if (RecoveryMode == ERecoveryMode::Hard || (TimespanSinceLastRecovery > MinimumDurationSinceLastRecovery))
 	{
-		FSystemWideCriticalSection RecoveryCriticalSection(TEXT("ZenServerLaunch"), MaximumWaitForHealth);
-		if (!RecoveryCriticalSection.IsValid())
-		{
-			// A recovery is already in progress but did not complete in time, we assume we failed and let recovery continue on a different thread
-			return false;
-		}
-
-		// We test if the service is healthy as a different process might already have triggered a recovery
-		bLastRecoveryResult = PingService(1);
-		if (bLastRecoveryResult)
-		{
-			UE_LOG(LogZenServiceInstance, Display, TEXT("Local ZenServer status: Healthy. Skipping recovery"));
-		}
-		else
+		const uint32_t MaxAttempts = RecoveryMode == ERecoveryMode::Hard ? 3 : 1;
+		uint32_t Attempts = 0;
+		for(;;)
 		{
-			// Update timespan since it may have changed since we waited to enter the crit section
-			TimespanSinceLastRecovery = FDateTime::UtcNow() - FDateTime(LastRecoveryTicks.load(std::memory_order_relaxed));
-			if (TimespanSinceLastRecovery > MinimumDurationSinceLastRecovery)
+			Attempts++;
+			FSystemWideCriticalSection RecoveryCriticalSection(TEXT("ZenServerLaunch"), MaximumWaitForHealth);
+			if (!RecoveryCriticalSection.IsValid())
 			{
-				UE_LOG(LogZenServiceInstance, Display, TEXT("Local ZenServer recovery being attempted..."));
+				// A recovery is already in progress but did not complete in time, we assume we failed and let recovery continue on a different thread
+				if (PingService(1, /*Silent*/ true))
+				{
+					return true;
+				}
+				if (Attempts == MaxAttempts)
+				{
+					UE_LOG(LogZenServiceInstance, Verbose, TEXT("Waiting for other process to complete recovery failed, abandonding recovery attempt"));
+					return false;
+				}
+				else
+				{
+					UE_LOG(LogZenServiceInstance, Verbose, TEXT("Waiting for other process to complete recovery failed, retrying..."));
+					continue;
+				}
+			}
 
-				uint16 Port = GetEndpoint().GetPort();
+			// We test if the service is healthy as a different process might already have triggered a recovery
+			bLastRecoveryResult = PingService(1, /*Silent*/ false);
+			if (!bLastRecoveryResult)
+			{
+				const uint64_t RecoveryStartTime = FPlatformTime::Cycles64();
 
-				bool bShutdownExistingInstance = true;
+				// Update timespan since it may have changed since we waited to enter the crit section
+				TimespanSinceLastRecovery = FDateTime::UtcNow() - FDateTime(LastRecoveryTicks.load(std::memory_order_relaxed));
+				if (RecoveryMode == ERecoveryMode::Hard || (TimespanSinceLastRecovery > MinimumDurationSinceLastRecovery))
 				{
-					const ZenServerState ServerState(/* ReadOnly */true);
-					const ZenServerState::ZenServerEntry* Entry = ServerState.LookupByEffectiveListenPort(Port);
-					if (Entry)
+					UE_LOG(LogZenServiceInstance, Display, TEXT("Local ZenServer is unreachable, initiating recovery..."));
+
+					uint16 Port = GetEndpoint().GetPort();
+
+					bool bShutdownExistingInstance = true;
 					{
-						if (Entry->Pid.load(std::memory_order_relaxed) != AutoLaunchedPid)
+						const ZenServerState ServerState(/* ReadOnly */true);
+						const ZenServerState::ZenServerEntry* Entry = ServerState.LookupByEffectiveListenPort(Port);
+						if (Entry)
 						{
-							// The running process pid is not the same as the one we launched.  The process was relaunched elsewhere. Avoid shutting it down again.
-							bShutdownExistingInstance = false;
+							if (Entry->Pid.load(std::memory_order_relaxed) != AutoLaunchedPid)
+							{
+								// The running process pid is not the same as the one we launched.  The process was relaunched elsewhere. Avoid shutting it down again.
+								bShutdownExistingInstance = false;
+							}
 						}
 					}
-				}
-				if (bShutdownExistingInstance && !ShutdownZenServerProcess((int)AutoLaunchedPid))	// !ShutdownRunningServiceUsingEffectivePort(Port))
-				{
-					return false;
-				}
+					if (bShutdownExistingInstance && !ShutdownZenServerProcess((int)AutoLaunchedPid))	// !ShutdownRunningServiceUsingEffectivePort(Port))
+					{
+						UE_LOG(LogZenServiceInstance, Warning, TEXT("Local ZenServer recovery failed to shut down running instance with pid %d"), AutoLaunchedPid);
+						return false;
+					}
 
-				AutoLaunch(Settings.SettingsVariant.Get<FServiceAutoLaunchSettings>(), ExecutablePath, ExecutionContextFilePath, Endpoint);
-				Port = GetEndpoint().GetPort();
+					AutoLaunch(Settings.SettingsVariant.Get<FServiceAutoLaunchSettings>(), ExecutablePath, ExecutionContextFilePath, Endpoint);
+					Port = GetEndpoint().GetPort();
 
-				FDateTime StartedWaitingForHealth = FDateTime::UtcNow();
-				bLastRecoveryResult = PingService(1);
-				while (!bLastRecoveryResult)
-				{
-					FTimespan WaitForHealth = FDateTime::UtcNow() - StartedWaitingForHealth;
-					if (WaitForHealth > MaximumWaitForHealth)
+					FDateTime StartedWaitingForHealth = FDateTime::UtcNow();
+					bLastRecoveryResult = PingService(1, /*Silent*/ true);
+					while (!bLastRecoveryResult)
 					{
-						UE_LOG(LogZenServiceInstance, Warning, TEXT("Local ZenServer recovery timed out waiting for service to become healthy"));
-						break;
+						FTimespan WaitForHealth = FDateTime::UtcNow() - StartedWaitingForHealth;
+						if (WaitForHealth > MaximumWaitForHealth)
+						{
+							UE_LOG(LogZenServiceInstance, Warning, TEXT("Local ZenServer recovery timed out waiting for service to become healthy"));
+							break;
+						}
+
+						FPlatformProcess::Sleep(0.5f);
+						if (!IsZenProcessUsingEffectivePort(Port))
+						{
+							AutoLaunch(Settings.SettingsVariant.Get<FServiceAutoLaunchSettings>(), ExecutablePath, ExecutionContextFilePath, Endpoint);
+						}
+						bLastRecoveryResult = PingService(1, /*Silent*/ false);
 					}
 
-					FPlatformProcess::Sleep(0.5f);
-					if (!IsZenProcessUsingEffectivePort(Port))
+					double ZenRecoveryDuration = FPlatformTime::ToSeconds64(FPlatformTime::Cycles64() - RecoveryStartTime);
+					if (bLastRecoveryResult)
 					{
-						AutoLaunch(Settings.SettingsVariant.Get<FServiceAutoLaunchSettings>(), ExecutablePath, ExecutionContextFilePath, Endpoint);
+						UE_LOG(LogZenServiceInstance, Display, TEXT("Local ZenServer at %s post recovery finished in %.3lf seconds, status: Healthy"), *Endpoint.GetURL(), ZenRecoveryDuration);
+						LastRecoveryTicks.store(0, std::memory_order_relaxed);
+					}
+					else
+					{
+						UE_LOG(LogZenServiceInstance, Warning, TEXT("Local ZenServer at %s post recovery finished in %.3lf seconds, status: NOT healthy"), *Endpoint.GetURL(), ZenRecoveryDuration);
+						LastRecoveryTicks.store(FDateTime::UtcNow().GetTicks(), std::memory_order_relaxed);
 					}
-					bLastRecoveryResult = PingService(1);
-				}
-				LastRecoveryTicks.store(FDateTime::UtcNow().GetTicks(), std::memory_order_relaxed);
-				UE_LOG(LogZenServiceInstance, Display, TEXT("Local ZenServer recovery finished."));
-				if (bLastRecoveryResult)
-				{
-					UE_LOG(LogZenServiceInstance, Display, TEXT("Local ZenServer post recovery status: Healthy"));
-				}
-				else
-				{
-					UE_LOG(LogZenServiceInstance, Display, TEXT("Local ZenServer post recovery status: NOT healthy"));
 				}
 			}
+			break;
 		}
 	}
 
@@ -2527,10 +2573,14 @@ FZenServiceInstance::Initialize()
 							}
 							AutoLaunchedPort = GetEndpoint().GetPort();
 							bIsRunningLocally = true;
-							if (PingService(2))
+							if (IsServiceReady())
 							{
 								break;
 							}
+							else
+							{
+								UE_LOG(LogZenServiceInstance, Warning, TEXT("Local ZenServer AutoLaunch completed but unable to reach service at %s"), *Endpoint.GetURL());
+							}
 						}
 					}
 				}
@@ -2843,6 +2893,24 @@ bool AddCurrentProcessAsSponsor(uint16 DesiredPort)
 static
 bool UninstallSystemService(FString UtilityPath)
 {
+	{
+		FString Output;
+		int32 ReturnCode;
+		if (RunZenUtility(UtilityPath, TEXT("service status"), Output, ReturnCode))
+		{
+			// 2 = bad arguments (old zen.exe version)
+			if (ReturnCode == 2)
+			{
+				return true;
+			}
+			// 1 = not installed
+			if (ReturnCode == 1 && Output.Contains("is not installed"))
+			{
+				return true;
+			}
+		}
+	}
+
 	FString RunningExecutable;
 	FString RunningCommandLine;
 	uint16 RunningPort;
@@ -2888,13 +2956,13 @@ bool UninstallSystemService(FString UtilityPath)
 
 	// Versions of zen that don't know about the service command will still return 0 for 'service status', leading to warnings
 	// later on in this code when we try to uninstall.
-	if (Output.StartsWith("Unknown command"))
+	if (Output.Contains("Unknown command"))
 	{
 		return true;
 	}
 
 	// zen returns 1 for status of a service that's not installed
-	if (ReturnCode == 1)
+	if (ReturnCode == 1 && Output.Contains("is not installed"))
 	{
 		return true;
 	}
diff --git a/Engine/Source/Developer/Zen/Public/Experimental/ZenServerInterface.h b/Engine/Source/Developer/Zen/Public/Experimental/ZenServerInterface.h
--- a/Engine/Source/Developer/Zen/Public/Experimental/ZenServerInterface.h
+++ b/Engine/Source/Developer/Zen/Public/Experimental/ZenServerInterface.h
@@ -315,7 +315,13 @@ public:
 	bool IsServiceRunningLocally() const { return bIsRunningLocally; }
 	const FZenServiceEndpoint& GetEndpoint() const { return Endpoint; }
 
-	UE_API bool TryRecovery();
+	enum class ERecoveryMode
+	{
+		Soft,	// Try to recover if we haven't done so recently and be relaxed if it fails
+		Hard	// Try to recover always and put more effort and retries to it
+	};
+
+	UE_API bool TryRecovery(ERecoveryMode RecoveryMode);
 
 	UE_API bool GetCacheStats(FZenCacheStats& Stats);
 	UE_API bool GetProjectStats(FZenProjectStats& Stats);
@@ -338,7 +344,7 @@ private:
 	bool WriteLocalPluginsConfig();
 	static bool AutoLaunch(const FServiceAutoLaunchSettings& InSettings, const FString& ExecutablePath, const FString& ExecutionContextFilePath, FZenServiceEndpoint& OutEndpoint);
 	bool ConditionalUpdateSystemServiceInstall(const FSystemServiceSettings& InSettings);
-	bool PingService(uint32_t AttemptCount);
+	bool PingService(uint32_t AttemptCount, bool Silent);
 
 	mutable TFuture<TOptional<FGCStatus>> GCStatusRequest;
 	mutable TOptional<FGCStatus> LastGCStatus;
diff --git a/Engine/Source/Editor/UnrealEd/Private/TargetDomain/TargetDomainUtils.cpp b/Engine/Source/Editor/UnrealEd/Private/TargetDomain/TargetDomainUtils.cpp
--- a/Engine/Source/Editor/UnrealEd/Private/TargetDomain/TargetDomainUtils.cpp
+++ b/Engine/Source/Editor/UnrealEd/Private/TargetDomain/TargetDomainUtils.cpp
@@ -136,14 +136,21 @@ TUniquePtr<FEditorDomainOplog> GEditorDomainOplog;
 TArray<const UTF8CHAR*> FEditorDomainOplog::ReservedOplogKeys;
 
 FEditorDomainOplog::FEditorDomainOplog()
-#if UE_WITH_ZEN
-: HttpClient(TEXT("localhost"), UE::Zen::FZenServiceInstance::GetAutoLaunchedPort() > 0 ? UE::Zen::FZenServiceInstance::GetAutoLaunchedPort() : 8558)
-#else
-: HttpClient(TEXT("localhost"), 8558)
-#endif
 {
 	StaticInit();
 
+#if UE_WITH_ZEN
+	if (UE::Zen::FZenServiceInstance::GetAutoLaunchedPort() != 0)
+	{
+		// Connect to auto-launched service
+		HttpClient = MakeUnique<UE::FZenStoreHttpClient>();
+	}
+	else
+#endif
+	{
+		HttpClient = MakeUnique<UE::FZenStoreHttpClient>(TEXT("localhost"), 8558);
+	}
+
 	FString ProjectId = FApp::GetZenStoreProjectId();
 	FString OplogId = TEXT("EditorDomain");
 
@@ -164,9 +171,9 @@ FEditorDomainOplog::FEditorDomainOplog()
 #if UE_WITH_ZEN
 	if (UE::Zen::IsDefaultServicePresent())
 	{
-		bool IsLocalConnection = HttpClient.GetZenServiceInstance().IsServiceRunningLocally();
-		HttpClient.TryCreateProject(ProjectId, FStringView(), OplogId, AbsServerRoot, AbsEngineDir, AbsProjectDir, IsLocalConnection ? ProjectFilePath : FStringView());
-		HttpClient.TryCreateOplog(ProjectId, OplogId, TEXT("") /*InOplogMarkerFile*/);
+		bool IsLocalConnection = HttpClient->GetZenServiceInstance().IsServiceRunningLocally();
+		HttpClient->TryCreateProject(ProjectId, FStringView(), OplogId, AbsServerRoot, AbsEngineDir, AbsProjectDir, IsLocalConnection ? ProjectFilePath : FStringView());
+		HttpClient->TryCreateOplog(ProjectId, OplogId, TEXT("") /*InOplogMarkerFile*/);
 	}
 #endif
 }
@@ -179,7 +186,7 @@ void FEditorDomainOplog::InitializeRead()
 	}
 	UE_LOG(LogEditorDomain, Display, TEXT("Fetching EditorDomain oplog..."));
 
-	TFuture<FIoStatus> FutureOplogStatus = HttpClient.GetOplog().Next([this](TIoStatusOr<FCbObject> OplogStatus)
+	TFuture<FIoStatus> FutureOplogStatus = HttpClient->GetOplog().Next([this](TIoStatusOr<FCbObject> OplogStatus)
 		{
 			if (!OplogStatus.IsOk())
 			{
@@ -259,7 +266,7 @@ bool FEditorDomainOplog::IsReservedOplogKey(FUtf8StringView Key)
 
 bool FEditorDomainOplog::IsValid() const
 {
-	return HttpClient.IsConnected();
+	return HttpClient->IsConnected();
 }
 
 void FEditorDomainOplog::CommitPackage(FName PackageName, TArrayView<IPackageWriter::FCommitAttachmentInfo> Attachments)
@@ -310,7 +317,7 @@ void FEditorDomainOplog::CommitPackage(FName PackageName, TArrayView<IPackageWri
 
 	FCbObject Obj = PackageObj.Save().AsObject();
 	Pkg.SetObject(Obj);
-	HttpClient.AppendOp(Pkg);
+	HttpClient->AppendOp(Pkg);
 }
 
 // Note that this is destructive - we yank out the buffer memory from the 
@@ -416,7 +423,7 @@ void FEditorDomainOplog::GetOplogAttachments(TArrayView<FName> PackageNames,
 		return;
 	}
 
-	HttpClient.ReadChunks(AttachmentHashes, [Callback = MoveTemp(Callback),&AttachmentHashParams](const FIoHash& RawHash, TIoStatusOr<FIoBuffer> Result)
+	HttpClient->ReadChunks(AttachmentHashes, [Callback = MoveTemp(Callback),&AttachmentHashParams](const FIoHash& RawHash, TIoStatusOr<FIoBuffer> Result)
 	{
 		for (auto It(AttachmentHashParams.CreateConstKeyIterator(RawHash)); It; ++It)
 		{
diff --git a/Engine/Source/Editor/UnrealEd/Private/TargetDomain/TargetDomainUtils.h b/Engine/Source/Editor/UnrealEd/Private/TargetDomain/TargetDomainUtils.h
--- a/Engine/Source/Editor/UnrealEd/Private/TargetDomain/TargetDomainUtils.h
+++ b/Engine/Source/Editor/UnrealEd/Private/TargetDomain/TargetDomainUtils.h
@@ -67,7 +67,7 @@ private:
 	static void StaticInit();
 	static bool IsReservedOplogKey(FUtf8StringView Key);
 
-	UE::FZenStoreHttpClient HttpClient;
+	TUniquePtr<UE::FZenStoreHttpClient> HttpClient;
 	FCriticalSection Lock;
 	TMap<FName, FOplogEntry> Entries;
 	bool bConnectSuccessful = false;
diff --git a/Engine/Source/Runtime/StorageServerClient/Private/StorageServerConnection.cpp b/Engine/Source/Runtime/StorageServerClient/Private/StorageServerConnection.cpp
--- a/Engine/Source/Runtime/StorageServerClient/Private/StorageServerConnection.cpp
+++ b/Engine/Source/Runtime/StorageServerClient/Private/StorageServerConnection.cpp
@@ -921,28 +921,31 @@ TSharedPtr<FInternetAddr> FStorageServerConnection::StringToInternetAddr(const F
 bool FStorageServerConnection::HandshakeRequest()
 {
 	// Handshakes are done with a limited connection timeout so that we can find out if the destination is unreachable in a timely manner.
-	const float ConnectionTimeoutSeconds = 5.0f;
-
-	IStorageServerHttpClient::FResult ResultTuple = HttpClient->RequestSync(
-		BaseURI,
-		EStorageServerContentType::Unknown,
-		"GET",
-		TOptional<FIoBuffer>(),
-		EStorageServerContentType::Unknown,
-		TOptional<FIoBuffer>(),
-		ConnectionTimeoutSeconds,
-		false
-	);
-	TIoStatusOr<FIoBuffer> Result = ResultTuple.Get<0>();
-	if (Result.IsOk())
-	{
-		FMemoryReaderView Reader(Result.ValueOrDie().GetView());
-		FCbObject ResponseObj = LoadCompactBinary(Reader).AsObject();
-		// we currently don't have any concept of protocol versioning, if
-		// we succeed in communicating with the endpoint we're good since
-		// any breaking API change would need to be done in a backward
-		// compatible manner
-		return true;
+	const float ConnectionTimeoutSeconds = 2.0f;
+
+	for (uint32_t Attempt = 0; Attempt < 3; Attempt++)
+	{
+		IStorageServerHttpClient::FResult ResultTuple = HttpClient->RequestSync(
+			BaseURI,
+			EStorageServerContentType::Unknown,
+			"GET",
+			TOptional<FIoBuffer>(),
+			EStorageServerContentType::Unknown,
+			TOptional<FIoBuffer>(),
+			ConnectionTimeoutSeconds,
+			false
+		);
+		TIoStatusOr<FIoBuffer> Result = ResultTuple.Get<0>();
+		if (Result.IsOk())
+		{
+			FMemoryReaderView Reader(Result.ValueOrDie().GetView());
+			FCbObject ResponseObj = LoadCompactBinary(Reader).AsObject();
+			// we currently don't have any concept of protocol versioning, if
+			// we succeed in communicating with the endpoint we're good since
+			// any breaking API change would need to be done in a backward
+			// compatible manner
+			return true;
+		}
 	}
 
 	return false;
