UnrealPak: Fix infinite waits by removing the flawed CalculateCompressionBufferMemory that was using an overflowed int32 result for chunks over 2 GiB because CompressionBufferSize was an int32.
Change type of CompressionBufferSize to uint64 to prevent future problems.
Remove the unnecessary CalculateCompressionBufferMemory function and use the GetSourceBufferSizeEstimate() for all scheduling purposes.
Add progress logging of scheduled memory.

Minor optimizations:
Skip an unnecessary atomic load operation when setting the trace counters.
Never shrink the AvailableCompressionBuffers array.

Resubmit:
Fix the incorrect tracking of UsedBufferMemory after change 44847596 that allows the process to continue on failed source file reads.
UsedBufferMemory was increased when a read started but never decremented on ZenServer read failures which could theoretically throttle subsequent requests and risk starvation (in practice unlikely since this would require either an extreme amount of requests failing or a few very huge ones).

#rnx
#was-backout 45716703

diff --git a/Engine/Source/Developer/IoStoreUtilities/Private/IoStoreUtilities.cpp b/Engine/Source/Developer/IoStoreUtilities/Private/IoStoreUtilities.cpp
--- a/Engine/Source/Developer/IoStoreUtilities/Private/IoStoreUtilities.cpp
+++ b/Engine/Source/Developer/IoStoreUtilities/Private/IoStoreUtilities.cpp
@@ -3392,6 +3392,7 @@ private:
 		{
 			CompletionEvent.Emplace(InCompletionEvent);
 			Manager.ScheduleLoad(this);
+			bLoadScheduled = true;
 		}
 
 		virtual const FIoBuffer* GetSourceBuffer() override
@@ -3401,9 +3402,10 @@ private:
 		
 		virtual void FreeSourceBuffer() override
 		{
-			bSourceBufferValid = false;
-			if (SourceBuffer.DataSize() > 0)
+			if (bLoadScheduled)
 			{
+				bLoadScheduled = false;
+				bSourceBufferValid = false;
 				SourceBuffer = FIoBuffer();
 				Manager.OnBufferMemoryFreed(SourceBufferSize);
 			}
@@ -3446,6 +3448,7 @@ private:
 		TOptional<UE::Tasks::FTaskEvent> CompletionEvent;
 		FIoBuffer SourceBuffer;
 		FQueueEntry* QueueEntry = nullptr;
+		bool bLoadScheduled = false;
 		bool bSourceBufferValid = false;
 	};
 
@@ -5682,6 +5685,7 @@ int32 CreateTarget(const FIoStoreArguments& Arguments, const FIoStoreWriterSetti
 				{
 					ProgressStringBuilder.Appendf(TEXT(" [%llu compression tasks]"), Progress.ScheduledCompressionTasksCount);
 				}
+				ProgressStringBuilder.Appendf(TEXT(" [Scheduled: %llu MiB]"), Progress.ScheduledCompressionMemoryBytes >> 20);
 			}
 			else
 			{
diff --git a/Engine/Source/Developer/IoStoreUtilities/Private/IoStoreWriter.cpp b/Engine/Source/Developer/IoStoreUtilities/Private/IoStoreWriter.cpp
--- a/Engine/Source/Developer/IoStoreUtilities/Private/IoStoreWriter.cpp
+++ b/Engine/Source/Developer/IoStoreUtilities/Private/IoStoreWriter.cpp
@@ -72,7 +72,6 @@ struct FIoStoreWriteQueueEntry
 	/** Hash of the block data as it would be found on disk after compression and encryption */
 	FIoHash ChunkDiskHash;
 
-	uint64 CompressionMemoryEstimate = 0;
 	uint64 Sequence = 0;
 	
 	// We make this optional because at the latest it might not be valid until FinishCompressionBarrior
@@ -471,7 +470,7 @@ public:
 		{
 			CompressionBufferSize = FCompression::CompressMemoryBound(WriterSettings.CompressionMethod, static_cast<int32>(WriterSettings.CompressionBlockSize));
 		}
-		CompressionBufferSize = FMath::Max(CompressionBufferSize, static_cast<int32>(WriterSettings.CompressionBlockSize));
+		CompressionBufferSize = FMath::Max(CompressionBufferSize, WriterSettings.CompressionBlockSize);
 		CompressionBufferSize = Align(CompressionBufferSize, FAES::AESBlockSize);
 
 		MaxCompressionBufferMemory = 2ull << 30;
@@ -512,6 +511,7 @@ public:
 		Progress.CompressedChunksCount = CompressedChunksCount.Load();
 		Progress.SerializedChunksCount = SerializedChunksCount.Load();
 		Progress.ScheduledCompressionTasksCount = ScheduledCompressionTasksCount.Load();
+		Progress.ScheduledCompressionMemoryBytes = ScheduledCompressionMemory.Load();
 		Progress.CompressionDDCGetBytes = CompressionDDCGetBytes.Load();
 		Progress.CompressionDDCPutBytes = CompressionDDCPutBytes.Load();
 		Progress.CompressionDDCMissCount = CompressionDDCMissCount.Load();
@@ -533,7 +533,7 @@ public:
 			FScopeLock Lock(&AvailableCompressionBuffersCritical);
 			if (AvailableCompressionBuffers.Num() > 0)
 			{
-				AllocatedBuffer = AvailableCompressionBuffers.Pop();
+				AllocatedBuffer = AvailableCompressionBuffers.Pop(EAllowShrinking::No);
 			}
 			TRACE_COUNTER_ADD(IoStoreCompressionMemoryUsed, CompressionBufferSize);
 		}
@@ -601,7 +601,7 @@ private:
 	FCriticalSection AvailableCompressionBuffersCritical;
 	TArray<FIoBuffer*> AvailableCompressionBuffers;
 	uint64 MaxCompressionBufferMemory = 0;
-	int32 CompressionBufferSize = -1;
+	uint64 CompressionBufferSize = 0;
 	TArray<TSharedPtr<FIoStoreWriter>> IoStoreWriters;
 
 	friend class FIoStoreWriter;
@@ -890,7 +890,6 @@ public:
 		Entry->ChunkId = ChunkId;
 		Entry->Options = WriteOptions;
 		Entry->CompressionMethod = CompressionMethodForEntry(WriteOptions);
-		Entry->CompressionMemoryEstimate = CalculateCompressionBufferMemory(Request->GetSourceBufferSizeEstimate());
 		Entry->bUseDDCForCompression =
 			WriterSettings.bCompressionEnableDDC &&
 			Entry->CompressionMethod != NAME_None &&
@@ -1644,12 +1643,6 @@ private:
 		return IntCastChecked<int32>(NumChunkBlocks64);
 	}
 
-	uint64 CalculateCompressionBufferMemory(uint64 ChunkSize)
-	{
-		int32 NumBlocks = CalculateNumChunkBlocks(ChunkSize);
-		return WriterContext->CompressionBufferSize * NumBlocks;
-	}
-
 	void AllocateCompressionBuffers(FIoStoreWriteQueueEntry* Entry, const uint8* UncompressedData = nullptr)
 	{
 		check(Entry->ChunkBlocks.Num() == 0);
@@ -1914,13 +1907,14 @@ private:
 		ON_SCOPE_EXIT
 		{
 			TRACE_CPUPROFILER_EVENT_SCOPE(FreeBlocks);
+			const uint64 EntryMemoryEstimate = Entry->Request->GetSourceBufferSizeEstimate();
 			FreeCompressionBuffers(Entry);
 			delete Entry->Request;
 			Entry->Request = nullptr;
 
-			WriterContext->ScheduledCompressionMemory.SubExchange(Entry->CompressionMemoryEstimate);
+			uint64 ScheduledMemory = WriterContext->ScheduledCompressionMemory.SubExchange(EntryMemoryEstimate);
 			WriterContext->CompressionMemoryReleasedEvent->Trigger();
-			TRACE_COUNTER_SET(IoStoreCompressionMemoryScheduled, WriterContext->ScheduledCompressionMemory.Load());
+			TRACE_COUNTER_SET(IoStoreCompressionMemoryScheduled, ScheduledMemory - EntryMemoryEstimate);
 		};
 		const int32* FindExistingIndex = TocBuilder.GetTocEntryIndex(Entry->ChunkId);
 		if (FindExistingIndex)
@@ -2273,10 +2267,10 @@ void FIoStoreWriterContextImpl::ScheduleAllEntries(TArrayView<FIoStoreWriteQueue
 
 	for (FIoStoreWriteQueueEntry* Entry : AllEntries)
 	{
-		uint64 LocalScheduledCompressionMemory = ScheduledCompressionMemory.Load();
+		uint64 ScheduledMemory = ScheduledCompressionMemory.Load();
+		const uint64 EntryMemoryEstimate = Entry->Request->GetSourceBufferSizeEstimate();
 		
-		while (LocalScheduledCompressionMemory > 0 &&
-			LocalScheduledCompressionMemory + Entry->CompressionMemoryEstimate > MaxCompressionBufferMemory)
+		while (ScheduledMemory > 0 && ScheduledMemory + EntryMemoryEstimate > MaxCompressionBufferMemory)
 		{
 			TRACE_CPUPROFILER_EVENT_SCOPE(WaitForCompressionMemory);
 			if (!CompressionMemoryReleasedEvent->Wait(100.f))
@@ -2285,11 +2279,11 @@ void FIoStoreWriterContextImpl::ScheduleAllEntries(TArrayView<FIoStoreWriteQueue
 				// make sure we are not waiting for unsubmitted ddc requests
 				DDCGetRequestDispatcher.DispatchGetRequests(HandleDDCGetResult);
 			}
-			LocalScheduledCompressionMemory = ScheduledCompressionMemory.Load();
+			ScheduledMemory = ScheduledCompressionMemory.Load();
 		}
 
-		ScheduledCompressionMemory.AddExchange(Entry->CompressionMemoryEstimate);
-		TRACE_COUNTER_SET(IoStoreCompressionMemoryScheduled, ScheduledCompressionMemory.Load());
+		ScheduledMemory = ScheduledCompressionMemory.AddExchange(EntryMemoryEstimate);
+		TRACE_COUNTER_SET(IoStoreCompressionMemoryScheduled, ScheduledMemory + EntryMemoryEstimate);
 
 		if (Entry->bLoadingFromReferenceDb)
 		{
diff --git a/Engine/Source/Developer/IoStoreUtilities/Private/IoStoreWriter.h b/Engine/Source/Developer/IoStoreUtilities/Private/IoStoreWriter.h
--- a/Engine/Source/Developer/IoStoreUtilities/Private/IoStoreWriter.h
+++ b/Engine/Source/Developer/IoStoreUtilities/Private/IoStoreWriter.h
@@ -71,6 +71,7 @@ public:
 		uint64 CompressedChunksByType[(int8)EIoChunkType::MAX] = { 0 };
 		uint64 SerializedChunksCount = 0;
 		uint64 ScheduledCompressionTasksCount = 0;
+		uint64 ScheduledCompressionMemoryBytes = 0;
 		uint64 CompressionDDCHitsByType[(int8)EIoChunkType::MAX] = { 0 };
 		uint64 CompressionDDCPutsByType[(int8)EIoChunkType::MAX] = { 0 };
 		uint64 CompressionDDCHitCount = 0;
