Updating comment around compression block alignment

#rnx

diff --git a/Engine/Source/Developer/IoStoreUtilities/Private/IoStoreWriter.cpp b/Engine/Source/Developer/IoStoreUtilities/Private/IoStoreWriter.cpp
--- a/Engine/Source/Developer/IoStoreUtilities/Private/IoStoreWriter.cpp
+++ b/Engine/Source/Developer/IoStoreUtilities/Private/IoStoreWriter.cpp
@@ -1964,10 +1964,9 @@ private:
 			}
 			if (WriterSettings.CompressionBlockAlignment)
 			{
-				// Try and prevent entries from crossing compression alignment blocks if possible. This is to avoid
-				// small entries from causing multiple file system block reads afaict. Large entries necesarily get
-				// aligned to prevent things like a blocksize + 2 entry being at alignment -1, causing 3 low level reads.
-				// ...I think.
+				// This is mostly to support patching systems that can't handle arbitrary insertion sizes. We want to
+				// keep things on the alignment so that they have a chance to match vs the previous release. We also
+				// want to pack small things into a single block so we don't waste a ton of space.
 				bool bCrossesBlockBoundary = Align(TargetPartition->Offset, WriterSettings.CompressionBlockAlignment) != Align(TargetPartition->Offset + Entry->DiskSize - 1, WriterSettings.CompressionBlockAlignment);
 				if (bCrossesBlockBoundary)
 				{
