Adding partition offset/number for container listing in unrealpak

#rnx

diff --git a/Engine/Source/Developer/IoStoreUtilities/Private/IoStoreUtilities.cpp b/Engine/Source/Developer/IoStoreUtilities/Private/IoStoreUtilities.cpp
--- a/Engine/Source/Developer/IoStoreUtilities/Private/IoStoreUtilities.cpp
+++ b/Engine/Source/Developer/IoStoreUtilities/Private/IoStoreUtilities.cpp
@@ -1258,7 +1258,7 @@ private:
 class FChunkEntryCsv
 {
 public:
-
+	FChunkEntryCsv() = default;
 	~FChunkEntryCsv()
 	{
 		if (OutputArchive.IsValid())
@@ -1268,12 +1268,17 @@ public:
 		}
 	}
 
+	void SetPartitionSize(uint64 InPartitionSize)
+	{
+		PartitionSize = InPartitionSize;
+	}
+
 	void CreateOutputFile(const TCHAR* OutputFilename)
 	{
 		OutputArchive.Reset(IFileManager::Get().CreateFileWriter(OutputFilename));
 		if (OutputArchive.IsValid())
 		{
-			OutputArchive->Logf(TEXT("OrderInContainer, ChunkId, PackageId, PackageName, Filename, ContainerName, Offset, OffsetOnDisk, Size, CompressedSize, Hash, ChunkType, ClassType, PakChunk, Platform, InstallType, DeliveryType"));
+			OutputArchive->Logf(TEXT("OrderInContainer, ChunkId, PackageId, PackageName, Filename, ContainerName, Offset, OffsetOnDisk, Size, CompressedSize, Hash, ChunkType, ClassType, PakChunk, Platform, InstallType, DeliveryType, PartitionIndex, OffsetInPartition"));
 		}
 	}
 
@@ -1321,7 +1326,7 @@ public:
 				InstallType = TEXT("Optional");
 			}
 			
-			OutputArchive->Logf(TEXT("%d, %s, 0x%s, %s, %s, %s, %lld, %lld, %lld, %lld, 0x%s, %s, %s, %s, %s, %s, %s"),
+			OutputArchive->Logf(TEXT("%d, %s, 0x%s, %s, %s, %s, %lld, %lld, %lld, %lld, 0x%s, %s, %s, %s, %s, %s, %s, %lld, %lld"),
 				Index,
 				*LexToString(Info.Id),
 				*LexToString(PackageId),
@@ -1338,12 +1343,15 @@ public:
 				*PakChunk,
 				*Platform,
 				*InstallType,
-				*DeliveryType
+				*DeliveryType,
+				Info.OffsetOnDisk / PartitionSize,
+				Info.OffsetOnDisk % PartitionSize 
 			);
 		}
 	}	
 
 private:
+	uint64 PartitionSize=0;
 	TUniquePtr<FArchive> OutputArchive;
 	FChunkDataSummary ChunkDataSummary;
 };
@@ -5862,12 +5870,13 @@ int32 CreateTarget(const FIoStoreArguments& Arguments, const FIoStoreWriterSetti
 	FGraphEventRef WriteCsvFileTask;
 	if (Arguments.CsvPath.Len() > 0)
 	{
-		WriteCsvFileTask = FFunctionGraphTask::CreateAndDispatchWhenReady([&Arguments, &IoStoreWriters, &IoStoreWriterResults]()
+		WriteCsvFileTask = FFunctionGraphTask::CreateAndDispatchWhenReady([&Arguments, &IoStoreWriters, &IoStoreWriterResults, MaxPartitionSize = GeneralIoWriterSettings.MaxPartitionSize]()
 		{
 			TRACE_CPUPROFILER_EVENT_SCOPE(WriteCsvFiles);
 
 			bool bPerContainerCsvFiles = FPaths::DirectoryExists(Arguments.CsvPath);
 			FChunkEntryCsv AllContainersOutCsvFile;
+			AllContainersOutCsvFile.SetPartitionSize(MaxPartitionSize);
 			FChunkEntryCsv* Out = &AllContainersOutCsvFile;
 			if (!bPerContainerCsvFiles)
 			{
@@ -5904,6 +5913,7 @@ int32 CreateTarget(const FIoStoreArguments& Arguments, const FIoStoreWriterSetti
 				{
 					IOSTORE_CPU_SCOPE(WriteCsvFile);
 					FChunkEntryCsv PerContainerOutCsvFile;
+					PerContainerOutCsvFile.SetPartitionSize(MaxPartitionSize);
 					if (bPerContainerCsvFiles)
 					{
 						// When CsvPath is a dir, then create one unique .utoc.csv per container file
@@ -6477,6 +6487,8 @@ int32 ListContainer(
 		}
 		
 		FChunkEntryCsv PerContainerOutCsvFile;
+		PerContainerOutCsvFile.SetPartitionSize(Reader->GetPartitionSize());
+		AllContainersOutCsvFile.SetPartitionSize(Reader->GetPartitionSize());
 		FChunkEntryCsv* Out;
 		
 		if (!bPerContainerCsvFiles)
diff --git a/Engine/Source/Runtime/Core/Private/IO/IoStore.cpp b/Engine/Source/Runtime/Core/Private/IO/IoStore.cpp
--- a/Engine/Source/Runtime/Core/Private/IO/IoStore.cpp
+++ b/Engine/Source/Runtime/Core/Private/IO/IoStore.cpp
@@ -914,6 +914,11 @@ public:
 		return BlockIndex >= FirstBlockIndex && BlockIndex <= LastBlockIndex;
 	}
 
+	uint64 GetPartitionSize() const
+	{
+		return TocReader.GetTocResource().Header.PartitionSize;
+	}
+
 	uint32 GetCompressionBlockSize() const
 	{
 		return TocReader.GetTocResource().Header.CompressionBlockSize;
@@ -1082,6 +1087,11 @@ const FIoDirectoryIndexReader& FIoStoreReader::GetDirectoryIndexReader() const
 	return Impl->GetDirectoryIndexReader();
 }
 
+uint64 FIoStoreReader::GetPartitionSize() const
+{
+	return Impl->GetPartitionSize();
+}
+
 uint32 FIoStoreReader::GetCompressionBlockSize() const
 {
 	return Impl->GetCompressionBlockSize();
diff --git a/Engine/Source/Runtime/Core/Public/IO/IoDispatcher.h b/Engine/Source/Runtime/Core/Public/IO/IoDispatcher.h
--- a/Engine/Source/Runtime/Core/Public/IO/IoDispatcher.h
+++ b/Engine/Source/Runtime/Core/Public/IO/IoDispatcher.h
@@ -625,6 +625,7 @@ public:
 	CORE_API void GetFilenamesByBlockIndex(const TArray<int32>& InBlockIndexList, TArray<FString>& OutFileList) const;
 	CORE_API void GetFilenames(TArray<FString>& OutFileList) const;
 
+	CORE_API uint64 GetPartitionSize() const;
 	CORE_API uint32 GetCompressionBlockSize() const;
 	CORE_API const TArray<FName>& GetCompressionMethods() const;
 	CORE_API void EnumerateCompressedBlocks(TFunction<bool(const FIoStoreTocCompressedBlockInfo&)>&& Callback) const;
