Fix division by zero when writing I/O store CSV file.

#rnx

diff --git a/Engine/Source/Developer/IoStoreUtilities/Private/IoStoreUtilities.cpp b/Engine/Source/Developer/IoStoreUtilities/Private/IoStoreUtilities.cpp
--- a/Engine/Source/Developer/IoStoreUtilities/Private/IoStoreUtilities.cpp
+++ b/Engine/Source/Developer/IoStoreUtilities/Private/IoStoreUtilities.cpp
@@ -1326,7 +1326,7 @@ public:
 				InstallType = TEXT("Optional");
 			}
 			
-			OutputArchive->Logf(TEXT("%d, %s, 0x%s, %s, %s, %s, %lld, %lld, %lld, %lld, 0x%s, %s, %s, %s, %s, %s, %s, %lld, %lld"),
+			OutputArchive->Logf(TEXT("%d, %s, 0x%s, %s, %s, %s, %lld, %lld, %lld, %lld, 0x%s, %s, %s, %s, %s, %s, %s, %llu, %llu"),
 				Index,
 				*LexToString(Info.Id),
 				*LexToString(PackageId),
@@ -1344,7 +1344,7 @@ public:
 				*Platform,
 				*InstallType,
 				*DeliveryType,
-				Info.OffsetOnDisk / PartitionSize,
+				PartitionSize > 0 ? Info.OffsetOnDisk / PartitionSize : 0,
 				Info.OffsetOnDisk % PartitionSize 
 			);
 		}
