Fix for CIS div by zero in UnrealPak

#rnx

diff --git a/Engine/Source/Developer/IoStoreUtilities/Private/IoStoreUtilities.cpp b/Engine/Source/Developer/IoStoreUtilities/Private/IoStoreUtilities.cpp
--- a/Engine/Source/Developer/IoStoreUtilities/Private/IoStoreUtilities.cpp
+++ b/Engine/Source/Developer/IoStoreUtilities/Private/IoStoreUtilities.cpp
@@ -1270,7 +1270,14 @@ public:
 
 	void SetPartitionSize(uint64 InPartitionSize)
 	{
-		PartitionSize = InPartitionSize;
+		if (InPartitionSize > 0)
+		{
+			PartitionSize = InPartitionSize;
+		}
+		else
+		{
+			PartitionSize = TNumericLimits<uint64>::Max();
+		}
 	}
 
 	void CreateOutputFile(const TCHAR* OutputFilename)
@@ -1344,14 +1351,14 @@ public:
 				*Platform,
 				*InstallType,
 				*DeliveryType,
-				PartitionSize > 0 ? Info.OffsetOnDisk / PartitionSize : 0,
+				Info.OffsetOnDisk / PartitionSize,
 				Info.OffsetOnDisk % PartitionSize 
 			);
 		}
 	}	
 
 private:
-	uint64 PartitionSize=0;
+	uint64 PartitionSize=TNumericLimits<uint64>::Max();
 	TUniquePtr<FArchive> OutputArchive;
 	FChunkDataSummary ChunkDataSummary;
 };
