Extra validation and rety of oplog save to zenserver

#rnx

diff --git a/Engine/Source/Developer/IoStoreUtilities/Private/ZenStoreHttpClient.cpp b/Engine/Source/Developer/IoStoreUtilities/Private/ZenStoreHttpClient.cpp
--- a/Engine/Source/Developer/IoStoreUtilities/Private/ZenStoreHttpClient.cpp
+++ b/Engine/Source/Developer/IoStoreUtilities/Private/ZenStoreHttpClient.cpp
@@ -495,6 +495,35 @@ TIoStatusOr<uint64> FZenStoreHttpClient::AppendOp(FCbPackage OpEntry)
 
 	{
 		TRACE_CPUPROFILER_EVENT_SCOPE(Zen_AppendOp_Async);
+
+		TConstArrayView<FCbAttachment> Attachments = OpEntry.GetAttachments();
+
+		{
+			TSet<FIoHash> AttachmentHashes;
+			for (const FCbAttachment& Attachment : Attachments)
+			{
+				AttachmentHashes.Add(Attachment.GetHash());
+			}
+			TStringBuilder<256> ErrorBuffer;
+			OpEntry.GetObject().IterateAttachments([&AttachmentHashes, &ErrorBuffer](FCbFieldView Field)
+				{
+					const FIoHash Hash = Field.AsAttachment();
+					if (!AttachmentHashes.Contains(Hash))
+					{
+						if (ErrorBuffer.Len() > 0)
+						{
+							ErrorBuffer << ", ";
+						}
+						ErrorBuffer << LexToString(Hash);
+					}
+				});
+
+			if (ErrorBuffer.Len() > 0)
+			{
+				return TIoStatusOr<uint64>((FIoStatus)(FIoStatusBuilder(EIoErrorCode::Unknown) << TEXT("Oplog body references attachments that are not included in the package: ") << ErrorBuffer.ToString()));
+			}
+		}
+
 		FLargeMemoryWriter SerializedPackage;
 
 		const int32 Salt = SaltGen.Next();
@@ -510,7 +539,15 @@ TIoStatusOr<uint64> FZenStoreHttpClient::AppendOp(FCbPackage OpEntry)
 		}
 		else
 		{
-			TConstArrayView<FCbAttachment> Attachments = OpEntry.GetAttachments();
+			IPlatformFile& PlatformFile = FPlatformFileManager::Get().GetPlatformFile();
+			TArray<FString> TempFiles;
+			ON_SCOPE_EXIT
+			{
+				for (const FString& TempFile : TempFiles)
+				{
+					(void)PlatformFile.DeleteFile(*TempFile);
+				}
+			};
 
 			// Prep phase
 
@@ -544,7 +581,6 @@ TIoStatusOr<uint64> FZenStoreHttpClient::AppendOp(FCbPackage OpEntry)
 				else
 				{
 					UE_LOG(LogZenStore, Warning, TEXT("Oplog prep FAILED: %s"), *Request->GetDetailedResultString(OplogPrepNewEntryPath));
-
 					UE_LOG(LogZenStore, Warning, TEXT("Unable to check for needed chunks, assuming all are needed"));
 					for (const FCbAttachment& Attachment : Attachments)
 					{
@@ -591,30 +627,37 @@ TIoStatusOr<uint64> FZenStoreHttpClient::AppendOp(FCbPackage OpEntry)
 						FIoHash AttachmentId = FIoHash::HashBuffer(MakeMemoryView(AttachmentSpec));
 
 						FString TempFilePath = TempDirPath / LexToString(AttachmentId);
-						IPlatformFile& PlatformFile = FPlatformFileManager::Get().GetPlatformFile();
 
-						if (IFileHandle* FileHandle = PlatformFile.OpenWrite(*TempFilePath))
+						if (PlatformFile.FileExists(*TempFilePath))
+						{
+							UE_LOG(LogZenStore, Warning, TEXT("A temp file named '%s' already exists, taking slow path for large attachment"), *TempFilePath);
+						}
+						else
 						{
-							bool WriteOK = FileHandle->Write((const uint8*)AttachmentData.GetData(), AttachmentData.GetSize());
-							delete FileHandle;
-							if (WriteOK)
+							if (IFileHandle* FileHandle = PlatformFile.OpenWrite(*TempFilePath))
 							{
-								Writer.AddHash(AttachmentHash);
-								bIsSerialized = true;
-								bIsUsingTempFiles = true;
+								bool WriteOK = FileHandle->Write((const uint8*)AttachmentData.GetData(), AttachmentData.GetSize());
+								delete FileHandle;
+								if (WriteOK)
+								{
+									Writer.AddHash(AttachmentHash);
+									TempFiles.Add(TempFilePath);
+									bIsSerialized = true;
+									bIsUsingTempFiles = true;
+								}
+								else
+								{
+									(void)PlatformFile.DeleteFile(*TempFilePath);
+									// Take the slow path if we can't write the payload file in the large attachment directory
+									UE_LOG(LogZenStore, Warning, TEXT("Could not write to file '%s', taking slow path for large attachment"), *TempFilePath);
+								}
 							}
 							else
 							{
-								(void)PlatformFile.DeleteFile(*TempFilePath);
-								// Take the slow path if we can't write the payload file in the large attachment directory
-								UE_LOG(LogZenStore, Warning, TEXT("Could not write to file '%s', taking slow path for large attachment"), *TempFilePath);
+								// Take the slow path if we can't open the payload file in the large attachment directory
+								UE_LOG(LogZenStore, Warning, TEXT("Could not create file '%s', taking slow path for large attachment"), *TempFilePath);
 							}
 						}
-						else
-						{
-							// Take the slow path if we can't open the payload file in the large attachment directory
-							UE_LOG(LogZenStore, Warning, TEXT("Could not create file '%s', taking slow path for large attachment"), *TempFilePath);
-						}
 					}
 
 					if (!bIsSerialized)
@@ -650,8 +693,7 @@ TIoStatusOr<uint64> FZenStoreHttpClient::AppendOp(FCbPackage OpEntry)
 		}
 		else
 		{
-			UE_LOG(LogZenStore, Warning, TEXT("Oplog write FAILED: %s"), *Request->GetDetailedResultString(OplogPrepNewEntryPath));
-			return TIoStatusOr<uint64>((FIoStatus)(FIoStatusBuilder(EIoErrorCode::Unknown) << TEXT("Append OpLog failed, NewOpLogPath='") << OplogNewEntryPath << TEXT("'")));
+			return TIoStatusOr<uint64>((FIoStatus)(FIoStatusBuilder(EIoErrorCode::Unknown) << TEXT("Append OpLog failed, NewOpLogPath='") << OplogNewEntryPath << TEXT("' ") << Request->GetDetailedResultString(OplogPrepNewEntryPath)));
 		}
 	}
 }
diff --git a/Engine/Source/Developer/IoStoreUtilities/Private/ZenStoreWriter.cpp b/Engine/Source/Developer/IoStoreUtilities/Private/ZenStoreWriter.cpp
--- a/Engine/Source/Developer/IoStoreUtilities/Private/ZenStoreWriter.cpp
+++ b/Engine/Source/Developer/IoStoreUtilities/Private/ZenStoreWriter.cpp
@@ -630,7 +630,7 @@ void FZenStoreWriter::BeginCook(const FCookInfo& Info)
 	{
 		ZenFileSystemManifest->Generate(MetadataDirectoryPath, Info.ReferencedPlugins);
 		TIoStatusOr<uint64> Status = HttpClient->AppendOp(CreateProjectMetaDataOpPackage("CookOnTheFly"));
-		UE_CLOG(!Status.IsOk(), LogZenStoreWriter, Fatal, TEXT("Failed to append OpLog"));
+		UE_CLOG(!Status.IsOk(), LogZenStoreWriter, Error, TEXT("Failed to append OpLog. Reason: %s"), *Status.Status().ToString());
 	}
 
 	if (FPlatformProcess::SupportsMultithreading())
@@ -666,14 +666,13 @@ void FZenStoreWriter::EndCook(const FCookInfo& Info)
 		if (FCbPackage ReferencedSetPackage = CreateReferencedSetOpPackage())
 		{
 			TIoStatusOr<uint64> Status = HttpClient->AppendOp(MoveTemp(ReferencedSetPackage));
-			UE_CLOG(!Status.IsOk(), LogZenStoreWriter, Fatal, TEXT("Failed to append ReferencedSetOp"));
+			UE_CLOG(!Status.IsOk(), LogZenStoreWriter, Error, TEXT("Failed to append ReferencedSetOp. Reason: %s"), *Status.Status().ToString());
 		}
 		{
 			TIoStatusOr<uint64> Status = HttpClient->EndBuildPass(CreateProjectMetaDataOpPackage("EndCook"));
-			UE_CLOG(!Status.IsOk(), LogZenStoreWriter, Fatal, TEXT("Failed to append OpLog and end the build pass"));
+			UE_CLOG(!Status.IsOk(), LogZenStoreWriter, Error, TEXT("Failed to append OpLog and end the build pass. Reason: %s"), *Status.Status().ToString());
 		}
 	}
-
 	UE_LOG(LogZenStoreWriter, Display, TEXT("Output:\t%d Public runtime script objects"), PackageStoreOptimizer->GetTotalScriptObjectCount());
 }
 
@@ -1022,8 +1021,20 @@ void FZenStoreWriter::CommitPackageInternal(FZenCommitInfo&& ZenCommitInfo)
 			CookInfo.Attachments = MoveTemp(CookInfoAttachments);
 		}
 
-		TIoStatusOr<uint64> Status = HttpClient->AppendOp(MoveTemp(OplogEntry));
-		UE_CLOG(!Status.IsOk(), LogZenStoreWriter, Fatal, TEXT("Failed to commit oplog entry '%s' to Zen"), *CommitInfo.PackageName.ToString());
+		TIoStatusOr<uint64> Status = HttpClient->AppendOp(OplogEntry);
+		if (!Status.IsOk())
+		{
+			UE_LOG(LogZenStoreWriter, Warning, TEXT("Failed to commit oplog entry '%s' to Zen, retrying: Reason: %s"), *CommitInfo.PackageName.ToString(), *Status.Status().ToString());
+			Status = HttpClient->AppendOp(MoveTemp(OplogEntry));
+			if (!Status.IsOk())
+			{
+				UE_LOG(LogZenStoreWriter, Error, TEXT("Failed to commit oplog entry '%s' to Zen: Reason: %s"), *CommitInfo.PackageName.ToString(), *Status.Status().ToString());
+			}
+			else
+			{
+				UE_LOG(LogZenStoreWriter, Verbose, TEXT("Commit oplog entry '%s' to Zen succeeded on retry"), *CommitInfo.PackageName.ToString());
+			}
+		}
 	}
 	else if (bComputeHash)
 	{
diff --git a/Engine/Source/Editor/UnrealEd/Private/TargetDomain/TargetDomainUtils.cpp b/Engine/Source/Editor/UnrealEd/Private/TargetDomain/TargetDomainUtils.cpp
--- a/Engine/Source/Editor/UnrealEd/Private/TargetDomain/TargetDomainUtils.cpp
+++ b/Engine/Source/Editor/UnrealEd/Private/TargetDomain/TargetDomainUtils.cpp
@@ -317,7 +317,22 @@ void FEditorDomainOplog::CommitPackage(FName PackageName, TArrayView<IPackageWri
 
 	FCbObject Obj = PackageObj.Save().AsObject();
 	Pkg.SetObject(Obj);
-	HttpClient->AppendOp(Pkg);
+
+
+	TIoStatusOr<uint64> Status = HttpClient->AppendOp(Pkg);
+	if (!Status.IsOk())
+	{
+		UE_LOG(LogEditorDomain, Warning, TEXT("Failed to commit oplog entry '%s' to Zen, retrying: Reason: %s"), *PackageNameKey, *Status.Status().ToString());
+		Status = HttpClient->AppendOp(MoveTemp(Pkg));
+		if (!Status.IsOk())
+		{
+			UE_LOG(LogEditorDomain, Error, TEXT("Failed to commit oplog entry '%s' to Zen: Reason: %s"), *PackageNameKey, *Status.Status().ToString());
+		}
+		else
+		{
+			UE_LOG(LogEditorDomain, Verbose, TEXT("Commit oplog entry '%s' to Zen succeeded on retry"), *PackageNameKey);
+		}
+	}
 }
 
 // Note that this is destructive - we yank out the buffer memory from the 
