[IoStoreUtilities]

Replaces FullName FNames used for Export/Import/Package Desc types to reduce the amount of unique FNames created when describing iostore content. When processing lots of content, these extra names can lead to suprassing the 2GB limit for FName indexing.

#rnx
#FailedValidations Preflight

diff --git a/Engine/Source/Developer/IoStoreUtilities/Private/IoStoreUtilities.cpp b/Engine/Source/Developer/IoStoreUtilities/Private/IoStoreUtilities.cpp
--- a/Engine/Source/Developer/IoStoreUtilities/Private/IoStoreUtilities.cpp
+++ b/Engine/Source/Developer/IoStoreUtilities/Private/IoStoreUtilities.cpp
@@ -7165,7 +7165,7 @@ namespace DescribeUtils
 	{
 		FPackageDesc* Package = nullptr;
 		FName Name;
-		FName FullName;
+		FSoftObjectPath FullName;
 		uint64 PublicExportHash;
 		FPackageObjectIndex OuterIndex;
 		FPackageObjectIndex ClassIndex;
@@ -7185,7 +7185,7 @@ namespace DescribeUtils
 
 	struct FImportDesc
 	{
-		FName Name;
+		FSoftObjectPath Name;
 		FPackageObjectIndex GlobalImportIndex;
 		FExportDesc* Export = nullptr;
 	};
@@ -7193,7 +7193,7 @@ namespace DescribeUtils
 	struct FScriptObjectDesc
 	{
 		FName Name;
-		FName FullName;
+		FSoftObjectPath FullName;
 		FPackageObjectIndex GlobalImportIndex;
 		FPackageObjectIndex OuterIndex;
 	};
@@ -7338,14 +7338,14 @@ namespace DescribeUtils
 		for (auto& KV : ScriptObjectByGlobalIdMap)
 		{
 			FScriptObjectDesc& ScriptObjectDesc = KV.Get<1>();
-			if (ScriptObjectDesc.FullName.IsNone())
+			if (ScriptObjectDesc.FullName.IsNull())
 			{
 				TArray<FScriptObjectDesc*> ScriptObjectStack;
 				FScriptObjectDesc* Current = &ScriptObjectDesc;
 				FString FullName;
 				while (Current)
 				{
-					if (!Current->FullName.IsNone())
+					if (!Current->FullName.IsNull())
 					{
 						FullName = Current->FullName.ToString();
 						break;
@@ -7357,7 +7357,7 @@ namespace DescribeUtils
 				{
 					Current = ScriptObjectStack.Pop();
 					FullName /= Current->Name.ToString();
-					Current->FullName = FName(FullName);
+					Current->FullName = FSoftObjectPath(FullName);
 				}
 			}
 		}
@@ -7634,7 +7634,7 @@ namespace DescribeUtils
 				FPackageDesc* PackageDesc = Packages[Index];
 				for (FExportDesc& ExportDesc : PackageDesc->Exports)
 				{
-					if (ExportDesc.FullName.IsNone())
+					if (ExportDesc.FullName.IsNull())
 					{
 						TRACE_CPUPROFILER_EVENT_SCOPE(GenerateExportFullName);
 
@@ -7645,9 +7645,9 @@ namespace DescribeUtils
 						TCHAR NameBuffer[FName::StringBufferSize];
 						for (;;)
 						{
-							if (!Current->FullName.IsNone())
+							if (!Current->FullName.IsNull())
 							{
-								Current->FullName.ToString(NameBuffer);
+								Current->FullName = FSoftObjectPath(NameBuffer);
 								FullNameBuilder.Append(NameBuffer);
 								break;
 							}
@@ -7666,7 +7666,7 @@ namespace DescribeUtils
 							FullNameBuilder.Append(TEXT("."));
 							Current->Name.ToString(NameBuffer);
 							FullNameBuilder.Append(NameBuffer);
-							Current->FullName = FName(FullNameBuilder);
+							Current->FullName = FSoftObjectPath(FullNameBuilder);
 						}
 					}
 				}
@@ -8199,7 +8199,7 @@ int32 ValidateCrossContainerRefs(
 			FPackageDesc* ImportPackage = Import.Export ? Import.Export->Package : nullptr;
 			if (!ImportPackage) 
 			{
-				UE_CLOG(Import.Name != FName() && !FPackageName::IsScriptPackage(*Import.Name.ToString()),
+				UE_CLOG(Import.Name.IsValid() && !FPackageName::IsScriptPackage(*Import.Name.ToString()),
 					LogIoStore, Error, TEXT("Unresolved import of package %s by package %s"), *Import.Name.ToString(), *Package->PackageName.ToString());
 				continue;
 			}
