[Backout] - CL46591167
#rnx
Original CL Desc
-----------------------------------------------------------------
[IoStoreUtilities]

Replaces FullName FNames used for Export/Import/Package Desc types to reduce the amount of unique FNames created when describing iostore content. When processing lots of content, these extra names can lead to suprassing the 2GB limit for FName indexing.

#rnx
#FailedValidations Preflight

diff --git a/Engine/Source/Developer/IoStoreUtilities/Private/IoStoreUtilities.cpp b/Engine/Source/Developer/IoStoreUtilities/Private/IoStoreUtilities.cpp
--- a/Engine/Source/Developer/IoStoreUtilities/Private/IoStoreUtilities.cpp
+++ b/Engine/Source/Developer/IoStoreUtilities/Private/IoStoreUtilities.cpp
@@ -7165,7 +7165,7 @@ namespace DescribeUtils
 	{
 		FPackageDesc* Package = nullptr;
 		FName Name;
-		FSoftObjectPath FullName;
+		FName FullName;
 		uint64 PublicExportHash;
 		FPackageObjectIndex OuterIndex;
 		FPackageObjectIndex ClassIndex;
@@ -7185,7 +7185,7 @@ namespace DescribeUtils
 
 	struct FImportDesc
 	{
-		FSoftObjectPath Name;
+		FName Name;
 		FPackageObjectIndex GlobalImportIndex;
 		FExportDesc* Export = nullptr;
 	};
@@ -7193,7 +7193,7 @@ namespace DescribeUtils
 	struct FScriptObjectDesc
 	{
 		FName Name;
-		FSoftObjectPath FullName;
+		FName FullName;
 		FPackageObjectIndex GlobalImportIndex;
 		FPackageObjectIndex OuterIndex;
 	};
@@ -7338,14 +7338,14 @@ namespace DescribeUtils
 		for (auto& KV : ScriptObjectByGlobalIdMap)
 		{
 			FScriptObjectDesc& ScriptObjectDesc = KV.Get<1>();
-			if (ScriptObjectDesc.FullName.IsNull())
+			if (ScriptObjectDesc.FullName.IsNone())
 			{
 				TArray<FScriptObjectDesc*> ScriptObjectStack;
 				FScriptObjectDesc* Current = &ScriptObjectDesc;
 				FString FullName;
 				while (Current)
 				{
-					if (!Current->FullName.IsNull())
+					if (!Current->FullName.IsNone())
 					{
 						FullName = Current->FullName.ToString();
 						break;
@@ -7357,7 +7357,7 @@ namespace DescribeUtils
 				{
 					Current = ScriptObjectStack.Pop();
 					FullName /= Current->Name.ToString();
-					Current->FullName = FSoftObjectPath(FullName);
+					Current->FullName = FName(FullName);
 				}
 			}
 		}
@@ -7634,7 +7634,7 @@ namespace DescribeUtils
 				FPackageDesc* PackageDesc = Packages[Index];
 				for (FExportDesc& ExportDesc : PackageDesc->Exports)
 				{
-					if (ExportDesc.FullName.IsNull())
+					if (ExportDesc.FullName.IsNone())
 					{
 						TRACE_CPUPROFILER_EVENT_SCOPE(GenerateExportFullName);
 
@@ -7645,9 +7645,9 @@ namespace DescribeUtils
 						TCHAR NameBuffer[FName::StringBufferSize];
 						for (;;)
 						{
-							if (!Current->FullName.IsNull())
+							if (!Current->FullName.IsNone())
 							{
-								Current->FullName = FSoftObjectPath(NameBuffer);
+								Current->FullName.ToString(NameBuffer);
 								FullNameBuilder.Append(NameBuffer);
 								break;
 							}
@@ -7666,7 +7666,7 @@ namespace DescribeUtils
 							FullNameBuilder.Append(TEXT("."));
 							Current->Name.ToString(NameBuffer);
 							FullNameBuilder.Append(NameBuffer);
-							Current->FullName = FSoftObjectPath(FullNameBuilder);
+							Current->FullName = FName(FullNameBuilder);
 						}
 					}
 				}
@@ -8199,7 +8199,7 @@ int32 ValidateCrossContainerRefs(
 			FPackageDesc* ImportPackage = Import.Export ? Import.Export->Package : nullptr;
 			if (!ImportPackage) 
 			{
-				UE_CLOG(Import.Name.IsValid() && !FPackageName::IsScriptPackage(*Import.Name.ToString()),
+				UE_CLOG(Import.Name != FName() && !FPackageName::IsScriptPackage(*Import.Name.ToString()),
 					LogIoStore, Error, TEXT("Unresolved import of package %s by package %s"), *Import.Name.ToString(), *Package->PackageName.ToString());
 				continue;
 			}
