Fix alignment when logging the container summaries.

#rnx

diff --git a/Engine/Source/Developer/IoStoreUtilities/Private/IoStoreUtilities.cpp b/Engine/Source/Developer/IoStoreUtilities/Private/IoStoreUtilities.cpp
--- a/Engine/Source/Developer/IoStoreUtilities/Private/IoStoreUtilities.cpp
+++ b/Engine/Source/Developer/IoStoreUtilities/Private/IoStoreUtilities.cpp
@@ -3185,14 +3185,24 @@ void LogWriterResults(const TArray<FIoStoreWriterResult>& Results)
 		uint64 PaddingSize = 0;
 	};
 
+	// Container name length can vary wildly, so find the longest one so that we can correctly align the output grids
+	int32 ColumnWidth = 10;
+	for (const FIoStoreWriterResult& Result : Results)
+	{
+		int32 ContainerNamePadding = ((Result.ContainerName.Len() / 10) + 1) * 10;
+		ColumnWidth = FMath::Max(ColumnWidth, ContainerNamePadding);
+	}
+
+	
 	UE_LOG(LogIoStore, Display, TEXT(""));
 	UE_LOG(LogIoStore, Display, TEXT("Container Summary"));
 	UE_LOG(LogIoStore, Display, TEXT("=================="));
 	UE_LOG(LogIoStore, Display, TEXT(""));
-	UE_LOG(LogIoStore, Display, TEXT("%-50s %10s %15s %15s %20s %20s"),
-		TEXT("Container"), TEXT("Flags"), TEXT("Chunk(s) #"), TEXT("TOC (KiB)"), TEXT("Raw Size (MiB)"), TEXT("Size (MiB)"));
-	UE_LOG(LogIoStore, Display, TEXT("----------------------------------------------------------------------------------------------------------------------------------------------"));
-
+	UE_LOG(LogIoStore, Display, TEXT("%-*s %10s %15s %15s %20s %20s"), ColumnWidth, TEXT("Container"), TEXT("Flags"), TEXT("Chunk(s) #"), TEXT("TOC (KiB)"), TEXT("Raw Size (MiB)"), TEXT("Size (MiB)"));
+	
+	const int32 SummaryBorderLength = ColumnWidth + 5 + 10 + 15 + 15 + 20 + 20; // +5 to account for spaces between columns
+	UE_LOG(LogIoStore, Display, TEXT("%s"), *FString::ChrN(SummaryBorderLength, '-'));
+	
 	FContainerStats TotalStats;
 	FContainerStats OnDemandStats;
 	for (const FIoStoreWriterResult& Result : Results)
@@ -3214,7 +3224,7 @@ void LogWriterResults(const TArray<FIoStoreWriterResult>& Results)
 			EnumHasAnyFlags(Result.ContainerFlags, EIoContainerFlags::Indexed) ? TEXT("I") : TEXT("-"),
 			EnumHasAnyFlags(Result.ContainerFlags, EIoContainerFlags::OnDemand) ? TEXT("O") : TEXT("-"));
 
-		UE_LOG(LogIoStore, Display, TEXT("%-50s %10s %15llu %15.2lf %20.2lf %20.2lf %s"),
+		UE_LOG(LogIoStore, Display, TEXT("%-*s %10s %15llu %15.2lf %20.2lf %20.2lf %s"), ColumnWidth,
 			*Result.ContainerName,
 			*ContainerSettings,
 			Result.TocEntryCount,
@@ -3237,12 +3247,12 @@ void LogWriterResults(const TArray<FIoStoreWriterResult>& Results)
 		TotalStats.CompressedContainerSize += Result.CompressedContainerSize;
 		TotalStats.PaddingSize += Result.PaddingSize;
 	}
-	UE_LOG(LogIoStore, Display, TEXT("----------------------------------------------------------------------------------------------------------------------------------------------"));
 
+	UE_LOG(LogIoStore, Display, TEXT("%s"), *FString::ChrN(SummaryBorderLength, '-'));
 
 	if (OnDemandStats.TocCount > 0)
 	{
-		UE_LOG(LogIoStore, Display, TEXT("%-50s %10s %15llu %15.2lf %20.2lf %20.2lf"),
+		UE_LOG(LogIoStore, Display, TEXT("%-*s %10s %15llu %15.2lf %20.2lf %20.2lf"), ColumnWidth,
 			TEXT("Total On Demand"),
 			TEXT(""),
 			OnDemandStats.TocCount,
@@ -3251,7 +3261,7 @@ void LogWriterResults(const TArray<FIoStoreWriterResult>& Results)
 			(double)OnDemandStats.CompressedContainerSize / 1024.0 / 1024.0);
 	}
 
-	UE_LOG(LogIoStore, Display, TEXT("%-50s %10s %15llu %15.2lf %20.2lf %20.2lf"),
+	UE_LOG(LogIoStore, Display, TEXT("%-*s %10s %15llu %15.2lf %20.2lf %20.2lf"), ColumnWidth,
 		TEXT("Total"),
 		TEXT(""),
 		TotalStats.TocCount,
@@ -3268,23 +3278,29 @@ void LogWriterResults(const TArray<FIoStoreWriterResult>& Results)
 	UE_LOG(LogIoStore, Display, TEXT("Container Directory Index"));
 	UE_LOG(LogIoStore, Display, TEXT("=========================="));
 	UE_LOG(LogIoStore, Display, TEXT(""));
-	UE_LOG(LogIoStore, Display, TEXT("%-45s %15s"), TEXT("Container"), TEXT("Size (KiB)"));
-	UE_LOG(LogIoStore, Display, TEXT("----------------------------------------------------------------------------------------------------------------------------------------------"));
+	UE_LOG(LogIoStore, Display, TEXT("%-*s %15s"), ColumnWidth, TEXT("Container"), TEXT("Size (KiB)"));
+	UE_LOG(LogIoStore, Display, TEXT("%s"), *FString::ChrN(ColumnWidth + 1 +15, '-')); // +1 to account for spaces between columns
+
 	for (const FIoStoreWriterResult& Result : Results)
 	{
-		UE_LOG(LogIoStore, Display, TEXT("%-45s %15.2lf"), *Result.ContainerName, double(Result.DirectoryIndexSize) / 1024.0);
+		UE_LOG(LogIoStore, Display, TEXT("%-*s %15.2lf"), ColumnWidth, *Result.ContainerName, double(Result.DirectoryIndexSize) / 1024.0);
 	}
 
 	UE_LOG(LogIoStore, Display, TEXT(""));
 	UE_LOG(LogIoStore, Display, TEXT("Container Patch Report"));
 	UE_LOG(LogIoStore, Display, TEXT("========================"));
 	UE_LOG(LogIoStore, Display, TEXT(""));
-	UE_LOG(LogIoStore, Display, TEXT("%-44s %16s %16s %16s %16s %16s"), TEXT("Container"), TEXT("Total #"), TEXT("Modified #"), TEXT("Added #"), TEXT("Modified (MiB)"), TEXT("Added (MiB)"));
-	UE_LOG(LogIoStore, Display, TEXT("----------------------------------------------------------------------------------------------------------------------------------------------"));
+	UE_LOG(LogIoStore, Display, TEXT("%-*s %16s %16s %16s %16s %16s"), ColumnWidth, TEXT("Container"), TEXT("Total #"), TEXT("Modified #"), TEXT("Added #"), TEXT("Modified (MiB)"), TEXT("Added (MiB)"));
+	
+	const int32 PatchReportBorderLength = ColumnWidth + 5 + 16 + 16 + 16 + 16 + 16; // +5 to account for spaces between columns
+	UE_LOG(LogIoStore, Display, TEXT("%s"), *FString::ChrN(PatchReportBorderLength, '-'));
+
 	for (const FIoStoreWriterResult& Result : Results)
 	{
-		UE_LOG(LogIoStore, Display, TEXT("%-44s %16d %16d %16d %16.2lf %16.2lf"), *Result.ContainerName, Result.TocEntryCount, Result.ModifiedChunksCount, Result.AddedChunksCount, Result.ModifiedChunksSize / 1024.0 / 1024.0, Result.AddedChunksSize / 1024.0 / 1024.0);
+		UE_LOG(LogIoStore, Display, TEXT("%-*s %16d %16d %16d %16.2lf %16.2lf"), ColumnWidth, *Result.ContainerName, Result.TocEntryCount, Result.ModifiedChunksCount, Result.AddedChunksCount, Result.ModifiedChunksSize / 1024.0 / 1024.0, Result.AddedChunksSize / 1024.0 / 1024.0);
 	}
+
+	UE_LOG(LogIoStore, Display, TEXT("%s"), *FString::ChrN(PatchReportBorderLength, '-'));
 }
 
 void LogContainerPackageInfo(const TArray<FContainerTargetSpec*>& ContainerTargets)
@@ -3293,16 +3309,26 @@ void LogContainerPackageInfo(const TArray<FContainerTargetSpec*>& ContainerTarge
 	uint64 TotalPackageCount = 0;
 	uint64 TotalLocalizedPackageCount = 0;
 
+	// Container name length can vary wildly, so find the longest one so that we can correctly align the output grids
+	int32 ColumnWidth = 10;
+	for (const FContainerTargetSpec* ContainerTarget : ContainerTargets)
+	{
+		int32 ContainerNamePadding = ((ContainerTarget->Name.ToString().Len() / 10) + 1) * 10;
+		ColumnWidth = FMath::Max(ColumnWidth, ContainerNamePadding);
+	}
+
 	UE_LOG(LogIoStore, Display, TEXT(""));
 	UE_LOG(LogIoStore, Display, TEXT("PackageStore"));
 	UE_LOG(LogIoStore, Display, TEXT("============="));
 	UE_LOG(LogIoStore, Display, TEXT(""));
-	UE_LOG(LogIoStore, Display, TEXT("%-45s %15s %15s %15s"),
+	UE_LOG(LogIoStore, Display, TEXT("%-*s %15s %15s %15s"), ColumnWidth,
 		TEXT("Container"),
 		TEXT("Size (KiB)"),
 		TEXT("Packages #"),
 		TEXT("Localized #"));
-	UE_LOG(LogIoStore, Display, TEXT("----------------------------------------------------------------------------------------------------------------------------------------------"));
+
+	const int32 BorderLength = ColumnWidth + 3 + 15 + 15 + 15; // +3 to account for spaces between columns
+	UE_LOG(LogIoStore, Display, TEXT("%s"), *FString::ChrN(BorderLength, '-'));
 
 	for (const FContainerTargetSpec* ContainerTarget : ContainerTargets)
 	{
@@ -3310,7 +3336,7 @@ void LogContainerPackageInfo(const TArray<FContainerTargetSpec*>& ContainerTarge
 		uint64 PackageCount = ContainerTarget->Packages.Num();
 		uint64 LocalizedPackageCount = ContainerTarget->Header.LocalizedPackages.Num();
 
-		UE_LOG(LogIoStore, Display, TEXT("%-45s %15.0lf %15llu %15llu"),
+		UE_LOG(LogIoStore, Display, TEXT("%-*s %15.0lf %15llu %15llu"), ColumnWidth,
 			*ContainerTarget->Name.ToString(),
 			(double)StoreSize / 1024.0,
 			PackageCount,
@@ -3320,8 +3346,9 @@ void LogContainerPackageInfo(const TArray<FContainerTargetSpec*>& ContainerTarge
 		TotalPackageCount += PackageCount;
 		TotalLocalizedPackageCount += LocalizedPackageCount;
 	}
-	UE_LOG(LogIoStore, Display, TEXT("----------------------------------------------------------------------------------------------------------------------------------------------"));
-	UE_LOG(LogIoStore, Display, TEXT("%-45s %15.0lf %15llu %15llu"),
+
+	UE_LOG(LogIoStore, Display, TEXT("%s"), *FString::ChrN(BorderLength, '-'));
+	UE_LOG(LogIoStore, Display, TEXT("%-*s %15.0lf %15llu %15llu"), ColumnWidth,
 		TEXT("Total"),
 		(double)TotalStoreSize / 1024.0,
 		TotalPackageCount,
