Added runtime support for the external package feature that was only available for editor packages, for both ZenStore and LinkerLoad paths.
This allows loading cooked packages with objects having their outer resource as an import from another package. One major difference between ZenStore and LinkerLoad is that the ZenStore package optimizer will only allow referencing public exports from a package, meaning that we can't have an object in a package with its outer set to another package if the outer object is not marked as public, but this was supported with the LinkerLoad codepath. This limitation will be adressed in an upcoming change.

diff --git a/Engine/Source/Developer/IoStoreUtilities/Internal/PackageStoreOptimizer.h b/Engine/Source/Developer/IoStoreUtilities/Internal/PackageStoreOptimizer.h
--- a/Engine/Source/Developer/IoStoreUtilities/Internal/PackageStoreOptimizer.h
+++ b/Engine/Source/Developer/IoStoreUtilities/Internal/PackageStoreOptimizer.h
@@ -144,11 +144,13 @@ private:
 	struct FUnresolvedExport
 	{
 		FString FullName;
+		FString FullPath;
 	};
 
 	struct FUnresolvedImport
 	{
 		FString FullName;
+		FString FullPath;
 		FName FromPackageName;
 		int32 FromPackageNameLen = 0;
 		bool bIsScriptImport = false;
@@ -273,10 +275,10 @@ private:
 #endif
 	static uint64 GetCellExportHash(FUtf8StringView VersePath);
 	FCookedHeaderData LoadCookedHeader(const FIoBuffer& CookedHeaderBuffer) const;
-	void ResolveImport(FPackageStorePackage::FUnresolvedImport* Imports, const FObjectImport* ObjectImports, int32 LocalImportIndex) const;
-	void ProcessImports(const FCookedHeaderData& CookedHeaderData, FPackageStorePackage* Package, TArray<FPackageStorePackage::FUnresolvedImport>& UnresolvedImports) const;
-	void ResolveExport(FPackageStorePackage::FUnresolvedExport* Exports, const FObjectExport* ObjectExports, const int32 LocalExportIndex, const FName& PackageName, FPackageStorePackage::FUnresolvedImport* Imports, const FObjectImport* ObjectImports) const;
-	void ProcessExports(const FCookedHeaderData& CookedHeaderData, FPackageStorePackage* Package, FPackageStorePackage::FUnresolvedImport* Imports) const;
+	void ResolveImport(FPackageStorePackage::FUnresolvedImport* Imports, FPackageStorePackage::FUnresolvedExport* Exports, const FObjectImport* ObjectImports, const FObjectExport* ObjectExports, const FName& PackageName, int32 LocalImportIndex) const;
+	void ProcessImports(const FCookedHeaderData& CookedHeaderData, FPackageStorePackage* Package, TArray<FPackageStorePackage::FUnresolvedImport>& UnresolvedImports, TArray<FPackageStorePackage::FUnresolvedExport>& UnresolvedExports) const;
+	void ResolveExport(FPackageStorePackage::FUnresolvedImport* Imports, FPackageStorePackage::FUnresolvedExport* Exports, const FObjectImport* ObjectImports, const FObjectExport* ObjectExports, const FName& PackageName, int32 LocalExportIndex) const;
+	void ProcessExports(const FCookedHeaderData& CookedHeaderData, FPackageStorePackage* Package, TArray<FPackageStorePackage::FUnresolvedImport>& UnresolvedImports, TArray<FPackageStorePackage::FUnresolvedExport>& UnresolvedExports) const;
 	void ProcessPreloadDependencies(const FCookedHeaderData& CookedHeaderData, FPackageStorePackage* Package) const;
 	void ProcessDataResources(const FCookedHeaderData& CookedHeaderData, FPackageStorePackage* Package) const;
 	TArray<FPackageStorePackage::FExportGraphNode*> SortExportGraphNodesInLoadOrder(FPackageStorePackage* Package, FExportGraphEdges& Edges) const;
diff --git a/Engine/Source/Developer/IoStoreUtilities/Private/PackageStoreOptimizer.cpp b/Engine/Source/Developer/IoStoreUtilities/Private/PackageStoreOptimizer.cpp
--- a/Engine/Source/Developer/IoStoreUtilities/Private/PackageStoreOptimizer.cpp
+++ b/Engine/Source/Developer/IoStoreUtilities/Private/PackageStoreOptimizer.cpp
@@ -103,7 +103,7 @@ FPackageStorePackage* FPackageStoreOptimizer::CreatePackageFromCookedHeader(cons
 	FPackageStorePackage* Package = new FPackageStorePackage();
 	Package->Id = FPackageId::FromName(Name);
 	Package->Name = Name;
-	
+
 	FCookedHeaderData CookedHeaderData = LoadCookedHeader(CookedHeaderBuffer);
 	if (!CookedHeaderData.Summary.bUnversioned)
 	{
@@ -122,9 +122,25 @@ FPackageStorePackage* FPackageStoreOptimizer::CreatePackageFromCookedHeader(cons
 
 	Package->SoftPackageReferences = MoveTemp(CookedHeaderData.SoftPackageReferences);
 
-	TArray<FPackageStorePackage::FUnresolvedImport> Imports;
-	ProcessImports(CookedHeaderData, Package, Imports);
-	ProcessExports(CookedHeaderData, Package, Imports.GetData());
+	TArray<FPackageStorePackage::FUnresolvedImport> UnresolvedImports;
+	TArray<FPackageStorePackage::FUnresolvedExport> UnresolvedExports;
+
+	const int32 ImportCount = CookedHeaderData.ObjectImports.Num();
+	const int32 ExportCount = CookedHeaderData.ObjectExports.Num();
+	const int32 CellImportCount = CookedHeaderData.CellImports.Num();	
+	const int32 CellExportCount = CookedHeaderData.CellExports.Num();
+
+	UnresolvedImports.SetNum(ImportCount);
+	UnresolvedExports.SetNum(ExportCount + CellExportCount);
+
+	Package->Imports.SetNum(ImportCount);
+	Package->Exports.SetNum(ExportCount);
+	Package->CellImports.SetNum(CellImportCount);
+	Package->CellExports.SetNum(CellExportCount);
+	Package->ExportGraphNodes.Reserve((ExportCount + CellExportCount) * 2);
+
+	ProcessImports(CookedHeaderData, Package, UnresolvedImports, UnresolvedExports);
+	ProcessExports(CookedHeaderData, Package, UnresolvedImports, UnresolvedExports);
 	ProcessPreloadDependencies(CookedHeaderData, Package);
 	ProcessDataResources(CookedHeaderData, Package);
 
@@ -276,7 +292,7 @@ FPackageStoreOptimizer::FCookedHeaderData FPackageStoreOptimizer::LoadCookedHead
 }
 PRAGMA_ENABLE_DEPRECATION_WARNINGS
 
-void FPackageStoreOptimizer::ResolveImport(FPackageStorePackage::FUnresolvedImport* Imports, const FObjectImport* ObjectImports, int32 LocalImportIndex) const
+void FPackageStoreOptimizer::ResolveImport(FPackageStorePackage::FUnresolvedImport* Imports, FPackageStorePackage::FUnresolvedExport* Exports, const FObjectImport* ObjectImports, const FObjectExport* ObjectExports, const FName& PackageName, int32 LocalImportIndex) const
 {
 	FPackageStorePackage::FUnresolvedImport* Import = Imports + LocalImportIndex;
 	if (Import->FullName.Len() == 0)
@@ -286,20 +302,21 @@ void FPackageStoreOptimizer::ResolveImport(FPackageStorePackage::FUnresolvedImpo
 		const FObjectImport* ObjectImport = ObjectImports + LocalImportIndex;
 		if (ObjectImport->OuterIndex.IsNull())
 		{
-			FName PackageName = ObjectImport->ObjectName;
-			PackageName.AppendString(Import->FullName);
+			FName ImportPackageName = ObjectImport->ObjectName;
+			ImportPackageName.AppendString(Import->FullName);
 			Import->FullName.ToLowerInline();
-			Import->FromPackageName = PackageName;
+			Import->FullPath = Import->FullName;
+			Import->FromPackageName = ImportPackageName;
 			Import->FromPackageNameLen = Import->FullName.Len();
 			Import->bIsScriptImport = Import->FullName.StartsWith(TEXT("/Script/"));
 			Import->bIsVerseVNIImport = !Import->bIsScriptImport && Import->FullName.Contains(TEXT("/_Verse/VNI"));
 			Import->bIsImportOfPackage = true;
 		}
-		else
+		else if (ObjectImport->OuterIndex.IsImport())
 		{
-			const int32 OuterIndex = ObjectImport->OuterIndex.ToImport();
-			ResolveImport(Imports, ObjectImports, OuterIndex);
-			FPackageStorePackage::FUnresolvedImport* OuterImport = Imports + OuterIndex;
+			const int32 OuterImportIndex = ObjectImport->OuterIndex.ToImport();
+			ResolveImport(Imports, Exports, ObjectImports, ObjectExports, PackageName, OuterImportIndex);
+			FPackageStorePackage::FUnresolvedImport* OuterImport = Imports + OuterImportIndex;
 			check(OuterImport->FullName.Len() > 0);
 			Import->bIsScriptImport = OuterImport->bIsScriptImport;
 			Import->bIsVerseVNIImport = OuterImport->bIsVerseVNIImport;
@@ -307,8 +324,41 @@ void FPackageStoreOptimizer::ResolveImport(FPackageStorePackage::FUnresolvedImpo
 			Import->FullName.AppendChar(TEXT('/'));
 			ObjectImport->ObjectName.AppendString(Import->FullName);
 			Import->FullName.ToLowerInline();
-			Import->FromPackageName = OuterImport->FromPackageName;
-			Import->FromPackageNameLen = OuterImport->FromPackageNameLen;
+			if (ObjectImport->HasPackageName())
+			{
+				ObjectImport->GetPackageName().AppendString(Import->FullPath);
+				Import->FullPath.AppendChar(TEXT('/'));
+				ObjectImport->ObjectName.AppendString(Import->FullPath);
+				Import->FullPath.ToLowerInline();
+				Import->FromPackageName = ObjectImport->GetPackageName();
+				Import->FromPackageNameLen = ObjectImport->GetPackageName().ToString().Len();
+			}
+			else
+			{
+				Import->FullPath.Append(OuterImport->FullPath);
+				Import->FullPath.AppendChar(TEXT('/'));
+				ObjectImport->ObjectName.AppendString(Import->FullPath);
+				Import->FullPath.ToLowerInline();
+				Import->FromPackageName = OuterImport->FromPackageName;
+				Import->FromPackageNameLen = OuterImport->FromPackageNameLen;
+			}
+		}
+		else
+		{
+			check(ObjectImport->HasPackageName());
+			int32 OuterExportIndex = ObjectImport->OuterIndex.ToExport();
+			ResolveExport(Imports, Exports, ObjectImports, ObjectExports, PackageName, OuterExportIndex);
+			FPackageStorePackage::FUnresolvedExport* OuterExport = Exports + OuterExportIndex;
+			Import->FullName.Append(OuterExport->FullName);
+			Import->FullName.AppendChar(TEXT('/'));
+			ObjectImport->ObjectName.AppendString(Import->FullName);
+			Import->FullName.ToLowerInline();
+			ObjectImport->GetPackageName().AppendString(Import->FullPath);
+			Import->FullPath.AppendChar(TEXT('/'));
+			ObjectImport->ObjectName.AppendString(Import->FullPath);
+			Import->FullPath.ToLowerInline();
+			Import->FromPackageName = ObjectImport->GetPackageName();
+			Import->FromPackageNameLen = ObjectImport->GetPackageName().ToString().Len();
 		}
 	}
 }
@@ -392,19 +442,15 @@ uint64 FPackageStoreOptimizer::GetCellExportHash(FUtf8StringView VersePath)
 	return CityHash64(reinterpret_cast<const char*>(VersePath.GetData()), VersePath.Len());
 }
 
-void FPackageStoreOptimizer::ProcessImports(const FCookedHeaderData& CookedHeaderData, FPackageStorePackage* Package, TArray<FPackageStorePackage::FUnresolvedImport>& UnresolvedImports) const
+void FPackageStoreOptimizer::ProcessImports(const FCookedHeaderData& CookedHeaderData, FPackageStorePackage* Package, TArray<FPackageStorePackage::FUnresolvedImport>& UnresolvedImports, TArray<FPackageStorePackage::FUnresolvedExport>& UnresolvedExports) const
 {
-	int32 ImportCount = CookedHeaderData.ObjectImports.Num();
-	UnresolvedImports.SetNum(ImportCount);
-	Package->Imports.SetNum(ImportCount);
-
-	int32 CellImportCount = CookedHeaderData.CellImports.Num();
-	Package->CellImports.SetNum(CellImportCount);
+	const int32 ImportCount = CookedHeaderData.ObjectImports.Num();
+	const int32 CellImportCount = CookedHeaderData.CellImports.Num();
 
 	TSet<FName> ImportedPackageNames;
 	for (int32 ImportIndex = 0; ImportIndex < ImportCount; ++ImportIndex)
 	{
-		ResolveImport(UnresolvedImports.GetData(), CookedHeaderData.ObjectImports.GetData(), ImportIndex);
+		ResolveImport(UnresolvedImports.GetData(), UnresolvedExports.GetData(), CookedHeaderData.ObjectImports.GetData(), CookedHeaderData.ObjectExports.GetData(), Package->Name, ImportIndex);
 		FPackageStorePackage::FUnresolvedImport& UnresolvedImport = UnresolvedImports[ImportIndex];
 		// For VerseVNI imports, we still need to add the package name to the imported package names for when the
 		// type being referenced wasn't generated by UHT
@@ -420,7 +466,7 @@ void FPackageStoreOptimizer::ProcessImports(const FCookedHeaderData& CookedHeade
 	{
 		const FCellImport& CellImport = CookedHeaderData.CellImports[CellImportIndex];
 		int32 ImportPackageIndex = CellImport.PackageIndex.ToImport();
-		ResolveImport(UnresolvedImports.GetData(), CookedHeaderData.ObjectImports.GetData(), ImportPackageIndex);
+		ResolveImport(UnresolvedImports.GetData(), UnresolvedExports.GetData(), CookedHeaderData.ObjectImports.GetData(), CookedHeaderData.ObjectExports.GetData(), Package->Name, ImportPackageIndex);
 		FPackageStorePackage::FUnresolvedImport& PackageImport = UnresolvedImports[ImportPackageIndex];
 		if (!PackageImport.bIsScriptImport)
 		{
@@ -468,7 +514,7 @@ void FPackageStoreOptimizer::ProcessImports(const FCookedHeaderData& CookedHeade
 			{
 				if (UnresolvedImport.FromPackageName == Package->ImportedPackages[PackageIndex].Name)
 				{
-					FStringView PackageRelativeName = FStringView(UnresolvedImport.FullName).RightChop(UnresolvedImport.FromPackageNameLen);
+					FStringView PackageRelativeName = FStringView(UnresolvedImport.FullPath).RightChop(UnresolvedImport.FromPackageNameLen);
 					check(PackageRelativeName.Len());
 					FPackageImportReference PackageImportRef(PackageIndex, Package->ImportedPublicExportHashes.Num());
 					Package->Imports[ImportIndex] = FPackageObjectIndex::FromPackageImportRef(PackageImportRef);
@@ -522,18 +568,13 @@ void FPackageStoreOptimizer::ProcessImports(const FCookedHeaderData& CookedHeade
 	}
 }
 
-void FPackageStoreOptimizer::ResolveExport(
-	FPackageStorePackage::FUnresolvedExport* Exports,
-	const FObjectExport* ObjectExports,
-	const int32 LocalExportIndex,
-	const FName& PackageName,
-	FPackageStorePackage::FUnresolvedImport* Imports,
-	const FObjectImport* ObjectImports) const
+void FPackageStoreOptimizer::ResolveExport(FPackageStorePackage::FUnresolvedImport* Imports, FPackageStorePackage::FUnresolvedExport* Exports, const FObjectImport* ObjectImports, const FObjectExport* ObjectExports, const FName& PackageName, int32 LocalExportIndex) const
 {
 	FPackageStorePackage::FUnresolvedExport* Export = Exports + LocalExportIndex;
 	if (Export->FullName.Len() == 0)
 	{
 		Export->FullName.Reserve(256);
+		Export->FullPath.Reserve(256);
 		const FObjectExport* ObjectExport = ObjectExports + LocalExportIndex;
 		if (ObjectExport->OuterIndex.IsNull())
 		{
@@ -542,6 +583,7 @@ void FPackageStoreOptimizer::ResolveExport(
 			ObjectExport->ObjectName.AppendString(Export->FullName);
 			Export->FullName.ToLowerInline();
 			check(Export->FullName.Len() > 0);
+			Export->FullPath = Export->FullName;
 		}
 		else
 		{
@@ -549,17 +591,30 @@ void FPackageStoreOptimizer::ResolveExport(
 			if (ObjectExport->OuterIndex.IsExport())
 			{
 				int32 OuterExportIndex = ObjectExport->OuterIndex.ToExport();
-				ResolveExport(Exports, ObjectExports, OuterExportIndex, PackageName, Imports, ObjectImports);
+				ResolveExport(Imports, Exports, ObjectImports, ObjectExports, PackageName, OuterExportIndex);
 				OuterName = &Exports[OuterExportIndex].FullName;
+				Export->FullPath.Append(*Exports[OuterExportIndex].FullPath);
+				Export->FullPath.AppendChar(TEXT('/'));
+				ObjectExport->ObjectName.AppendString(Export->FullPath);
+				Export->FullPath.ToLowerInline();
 			}
 			else
 			{
 				check(Imports && ObjectImports);
 				int32 OuterImportIndex = ObjectExport->OuterIndex.ToImport();
-				ResolveImport(Imports, ObjectImports, OuterImportIndex);
+				ResolveImport(Imports, Exports, ObjectImports, ObjectExports, PackageName, OuterImportIndex);
 				OuterName = &Imports[OuterImportIndex].FullName;
-
+				Export->FullPath.Append(PackageName.ToString());
+				if (OuterName->StartsWith(PackageName.ToString()))
+				{
+					FStringView PackageRelativeName = FStringView(*OuterName).RightChop(PackageName.ToString().Len());
+					Export->FullPath += PackageRelativeName;
+				}
+				Export->FullPath.AppendChar(TEXT('/'));
+				ObjectExport->ObjectName.AppendString(Export->FullPath);
+				Export->FullPath.ToLowerInline();
 			}
+
 			check(OuterName && OuterName->Len() > 0);
 			Export->FullName.Append(*OuterName);
 			Export->FullName.AppendChar(TEXT('/'));
@@ -569,16 +624,10 @@ void FPackageStoreOptimizer::ResolveExport(
 	}
 }
 
-void FPackageStoreOptimizer::ProcessExports(const FCookedHeaderData& CookedHeaderData, FPackageStorePackage* Package, FPackageStorePackage::FUnresolvedImport* Imports) const
+void FPackageStoreOptimizer::ProcessExports(const FCookedHeaderData& CookedHeaderData, FPackageStorePackage* Package, TArray<FPackageStorePackage::FUnresolvedImport>& UnresolvedImports, TArray<FPackageStorePackage::FUnresolvedExport>& UnresolvedExports) const
 {
-	int32 ExportCount = CookedHeaderData.ObjectExports.Num();
-	int32 CellExportCount = CookedHeaderData.CellExports.Num();
-
-	TArray<FPackageStorePackage::FUnresolvedExport> UnresolvedExports;
-	UnresolvedExports.SetNum(ExportCount + CellExportCount);
-	Package->Exports.SetNum(ExportCount);
-	Package->CellExports.SetNum(CellExportCount);
-	Package->ExportGraphNodes.Reserve((ExportCount + CellExportCount) * 2);
+	const int32 ExportCount = CookedHeaderData.ObjectExports.Num();
+	const int32 CellExportCount = CookedHeaderData.CellExports.Num();
 
 	auto PackageObjectIdFromPackageIndex =
 		[](const TArray<FPackageObjectIndex>& Imports, const FPackageIndex& PackageIndex) -> FPackageObjectIndex
@@ -610,22 +659,23 @@ void FPackageStoreOptimizer::ProcessExports(const FCookedHeaderData& CookedHeade
 		Export.bNotForClient = ObjectExport.bNotForClient;
 		Export.bNotForServer = ObjectExport.bNotForServer;
 		Export.bIsPublic = (Export.ObjectFlags & RF_Public) > 0 || ObjectExport.bGeneratePublicHash;
-		ResolveExport(UnresolvedExports.GetData(), CookedHeaderData.ObjectExports.GetData(), ExportIndex, Package->Name, Imports, CookedHeaderData.ObjectImports.GetData());
+		ResolveExport(UnresolvedImports.GetData(), UnresolvedExports.GetData(), CookedHeaderData.ObjectImports.GetData(), CookedHeaderData.ObjectExports.GetData(), Package->Name, ExportIndex);
 		if (Export.bIsPublic)
 		{
-			check(UnresolvedExport.FullName.Len() > 0);
-			FStringView PackageRelativeName = FStringView(UnresolvedExport.FullName).RightChop(PackageNameStr.Len());
+			check(UnresolvedExport.FullPath.Len() > 0);
+			check(UnresolvedExport.FullPath.StartsWith(PackageNameStr));
+			FStringView PackageRelativeName = FStringView(UnresolvedExport.FullPath).RightChop(PackageNameStr.Len());
 			check(PackageRelativeName.Len());
 			Export.PublicExportHash = GetPublicExportHash(PackageRelativeName);
 			const FPackageStorePackage::FUnresolvedExport* FindCollidingExport = SeenPublicExportHashes.FindRef(Export.PublicExportHash);
 			if (FindCollidingExport)
 			{
-				UE_LOG(LogPackageStoreOptimizer, Fatal, TEXT("Export hash collision in package \"%s\": \"%s\" and \"%s"), *PackageNameStr, PackageRelativeName.GetData(), *FindCollidingExport->FullName);
+				UE_LOG(LogPackageStoreOptimizer, Fatal, TEXT("Export hash collision in package \"%s\": \"%s\" and \"%s"), *PackageNameStr, PackageRelativeName.GetData(), *FindCollidingExport->FullPath);
 			}
 			SeenPublicExportHashes.Add(Export.PublicExportHash, &UnresolvedExport);
 		}
 
-		Export.OuterIndex = PackageObjectIdFromPackageIndex(Package->Imports, ObjectExport.OuterIndex);
+		Export.OuterIndex = PackageObjectIdFromPackageIndex(Package->Imports, ObjectExport.OuterIndex);		
 		Export.ClassIndex = PackageObjectIdFromPackageIndex(Package->Imports, ObjectExport.ClassIndex);
 		Export.SuperIndex = PackageObjectIdFromPackageIndex(Package->Imports, ObjectExport.SuperIndex);
 		Export.TemplateIndex = PackageObjectIdFromPackageIndex(Package->Imports, ObjectExport.TemplateIndex);
diff --git a/Engine/Source/Runtime/CoreUObject/Private/Serialization/AsyncLoading.cpp b/Engine/Source/Runtime/CoreUObject/Private/Serialization/AsyncLoading.cpp
--- a/Engine/Source/Runtime/CoreUObject/Private/Serialization/AsyncLoading.cpp
+++ b/Engine/Source/Runtime/CoreUObject/Private/Serialization/AsyncLoading.cpp
@@ -3404,6 +3404,10 @@ void FAsyncPackage::EventDrivenCreateExport(int32 LocalExportIndex)
 	{
 		Export.bExportLoadFailed = true;
 	}
+	else if (Export.OuterIndex.IsImport())
+	{
+		Export.Object->SetExternalPackage(LinkerRoot);
+	}
 	LastObjectWorkWasPerformedOn = Export.Object;
 	check(Export.Object || Export.bExportLoadFailed);
 }
diff --git a/Engine/Source/Runtime/CoreUObject/Private/Serialization/AsyncLoading2.cpp b/Engine/Source/Runtime/CoreUObject/Private/Serialization/AsyncLoading2.cpp
--- a/Engine/Source/Runtime/CoreUObject/Private/Serialization/AsyncLoading2.cpp
+++ b/Engine/Source/Runtime/CoreUObject/Private/Serialization/AsyncLoading2.cpp
@@ -8757,6 +8757,11 @@ void FAsyncPackage2::EventDrivenCreateExport(const FAsyncPackageHeaderData& Head
 
 	check(Object);
 	EInternalObjectFlags FlagsToSet = EInternalObjectFlags::Async;
+
+	if (Export.OuterIndex.IsImport())
+	{
+		Object->SetExternalPackage(LinkerRoot);
+	}
 	
 	if (Desc.bCanBeImported && Export.PublicExportHash)
 	{
@@ -11978,7 +11983,7 @@ EAsyncPackageState::Type FAsyncLoadingThread2::ProcessLoadingUntilCompleteFromGa
 	bool bUseTimeLimit = TimeLimit > 0.0f;
 	double TimeLoadingPackage = 0.0f;
 
-	// If there's an active sync load context we need to supress is for the duration of this call since we've no idea what the CompletionPredicate is waiting for
+	// If there's an active sync load context we need to suppress is for the duration of this call since we've no idea what the CompletionPredicate is waiting for
 	uint64 NoSyncLoadContextId = 0;
 	uint64& SyncLoadContextId = ThreadState.SyncLoadContextStack.IsEmpty() ? NoSyncLoadContextId : ThreadState.SyncLoadContextStack.Top()->ContextId;
 	TGuardValue<uint64> GuardSyncLoadContextId(SyncLoadContextId, 0);
diff --git a/Engine/Source/Runtime/CoreUObject/Private/UObject/ObjectResource.cpp b/Engine/Source/Runtime/CoreUObject/Private/UObject/ObjectResource.cpp
--- a/Engine/Source/Runtime/CoreUObject/Private/UObject/ObjectResource.cpp
+++ b/Engine/Source/Runtime/CoreUObject/Private/UObject/ObjectResource.cpp
@@ -359,12 +359,10 @@ void operator<<(FStructuredArchive::FSlot Slot, FObjectImport& I)
 	Record << SA_VALUE(TEXT("OuterIndex"), I.OuterIndex);
 	Record << SA_VALUE(TEXT("ObjectName"), I.ObjectName);
 
-#if WITH_EDITORONLY_DATA
-	if (Slot.GetUnderlyingArchive().UEVer() >= VER_UE4_NON_OUTER_PACKAGE_IMPORT && !Slot.GetUnderlyingArchive().IsFilterEditorOnly())
+	if (Slot.GetUnderlyingArchive().UEVer() >= VER_UE4_NON_OUTER_PACKAGE_IMPORT)
 	{
 		Record << SA_VALUE(TEXT("PackageName"), I.PackageName);
 	}
-#endif
 
 	if (Slot.GetUnderlyingArchive().UEVer() >= EUnrealEngineObjectUE5Version::OPTIONAL_RESOURCES)
 	{
diff --git a/Engine/Source/Runtime/CoreUObject/Private/UObject/SavePackage/PackageHarvester.cpp b/Engine/Source/Runtime/CoreUObject/Private/UObject/SavePackage/PackageHarvester.cpp
--- a/Engine/Source/Runtime/CoreUObject/Private/UObject/SavePackage/PackageHarvester.cpp
+++ b/Engine/Source/Runtime/CoreUObject/Private/UObject/SavePackage/PackageHarvester.cpp
@@ -933,11 +933,8 @@ void FPackageHarvester::ProcessImport(TObjectPtr<UObject> InObjPtr,
 	}
 	else
 	{
-		if (!IsFilterEditorOnly())
-		{
-			// operator<<(FStructuredArchive::FSlot Slot, FObjectImport& I) will need to write NAME_None for this empty ExternalPackage pointer
-			HarvestPackageHeaderName(NAME_None);
-		}
+		// operator<<(FStructuredArchive::FSlot Slot, FObjectImport& I) will need to write NAME_None for this empty ExternalPackage pointer
+		HarvestPackageHeaderName(NAME_None);
 	}
 
 	// For things with a BP-created class we need to recurse into that class so the import ClassPackage will load properly
diff --git a/Engine/Source/Runtime/CoreUObject/Public/UObject/FastReferenceCollector.h b/Engine/Source/Runtime/CoreUObject/Public/UObject/FastReferenceCollector.h
--- a/Engine/Source/Runtime/CoreUObject/Public/UObject/FastReferenceCollector.h
+++ b/Engine/Source/Runtime/CoreUObject/Public/UObject/FastReferenceCollector.h
@@ -1015,11 +1015,11 @@ private:
 			// Emit base references
 			Dispatcher.HandleImmutableReference(Class, EMemberlessId::Class, EOrigin::Other);
 			Dispatcher.HandleImmutableReference(Outer, EMemberlessId::Outer, EOrigin::Other);
-#if WITH_EDITOR
+
 			UObject* Package = CurrentObject->GetExternalPackageInternal();
 			Package = Package != CurrentObject ? Package : nullptr;
 			Dispatcher.HandleImmutableReference(Package, EMemberlessId::ExternalPackage, EOrigin::Other);
-#endif
+
 			if (!Schema.IsEmpty())
 			{
 				typename DispatcherType::SchemaStackScopeType SchemaStack(Dispatcher.Context, Schema);
diff --git a/Engine/Source/Runtime/CoreUObject/Public/UObject/ObjectResource.h b/Engine/Source/Runtime/CoreUObject/Public/UObject/ObjectResource.h
--- a/Engine/Source/Runtime/CoreUObject/Public/UObject/ObjectResource.h
+++ b/Engine/Source/Runtime/CoreUObject/Public/UObject/ObjectResource.h
@@ -454,14 +454,12 @@ struct FObjectImport : public FObjectResource
 	 */
 	FName			ClassName;
 
-#if WITH_EDITORONLY_DATA
 	/**
 	 * Package Name this import belongs to. Can be none, in that case follow the outer chain
 	 * until a set PackageName is found or until OuterIndex is null
 	 * Serialized
 	 */
 	FName			PackageName;
-#endif
 
 	/**
 	 * Index into SourceLinker's ExportMap for the export associated with this import's UObject.
@@ -499,42 +497,29 @@ struct FObjectImport : public FObjectResource
 
 	/**
 	 * Accessor function to check if the import has package name set
-	 * Handles editor only code.
 	 * @returns true if the import has a PackageName set
 	 */
 	bool HasPackageName() const
 	{
-#if WITH_EDITORONLY_DATA
 		return !PackageName.IsNone();
-#else
-		return false;
-#endif
 	}
 
 	/**
 	 * Accessor function to get the import package name
-	 * Handles editor only code.
 	 * @returns the import package name, if any
 	 */
 	FName GetPackageName() const
 	{
-#if WITH_EDITORONLY_DATA
 		return PackageName;
-#else
-		return NAME_None;
-#endif
 	}
 
 	/**
 	 * Accessor function to set the import PackageName
-	 * Handles editor only code.
 	 * @param InPackageName the package name to set
 	 */
 	void SetPackageName(FName InPackageName)
 	{
-#if WITH_EDITORONLY_DATA
 		PackageName = InPackageName;
-#endif
 	}
 	
 	/** I/O functions */
