UnrealPak: iostore command FindAndExtract needs to early exit after running (and return success or failure) rather than continuing on to the default path of attempting to create an IoStore container.
#rnx

diff --git a/Engine/Source/Developer/IoStoreUtilities/Private/IoStoreUtilities.cpp b/Engine/Source/Developer/IoStoreUtilities/Private/IoStoreUtilities.cpp
--- a/Engine/Source/Developer/IoStoreUtilities/Private/IoStoreUtilities.cpp
+++ b/Engine/Source/Developer/IoStoreUtilities/Private/IoStoreUtilities.cpp
@@ -6440,7 +6440,7 @@ static TArray<FString> GetContainersFromPathOrWildcard(const FString& InContaine
 	return MoveTemp(ContainerFilePaths);
 }
 
-static void FindAndExtractPackages(
+static int32 FindAndExtractPackages(
 	const FKeyChain& KeyChain,
 	const FString& ContainerPathOrWildcard,
 	const FString& OutputPath,
@@ -6451,15 +6451,17 @@ static void FindAndExtractPackages(
 	if (ContainerFilePaths.Num() == 0)
 	{
 		UE_LOG(LogIoStore, Error, TEXT("Container '%s' doesn't exist and no container matches wildcard."), *ContainerPathOrWildcard);
-		return;
+		return 1;
 	}
 
+	bool bAllSucceeded = true;
 	for (const FString& ContainerFilePath : ContainerFilePaths)
 	{
-		ExtractFilesFromIoStoreContainer(
+		bool bSucceeded = ExtractFilesFromIoStoreContainer(
 			*ContainerFilePath, *OutputPath, KeyChain, &PackageFilter, nullptr, nullptr, nullptr);
+		bAllSucceeded = bSucceeded && bAllSucceeded;
 	}
-
+	return bAllSucceeded ? 0 : 1;
 }
 
 int32 ListContainer(
@@ -9945,7 +9947,7 @@ int32 CreateIoStoreContainerFiles(const TCHAR* CmdLine)
 			return -1;
 		}
 
-		FindAndExtractPackages(Arguments.KeyChain, ContainerPathOrWildcard, OutPath, PackageFilter);
+		return FindAndExtractPackages(Arguments.KeyChain, ContainerPathOrWildcard, OutPath, PackageFilter);
 	}
 	else if (FParse::Value(FCommandLine::Get(), TEXT("Describe="), ArgumentValue))
 	{
