Added packaging option to allow specifying alphabetical-cluster-sort and alphabetical fallback ordering for a specific prefix match of container names.

diff --git a/Engine/Source/Developer/IoStoreUtilities/Private/IoStoreUtilities.cpp b/Engine/Source/Developer/IoStoreUtilities/Private/IoStoreUtilities.cpp
--- a/Engine/Source/Developer/IoStoreUtilities/Private/IoStoreUtilities.cpp
+++ b/Engine/Source/Developer/IoStoreUtilities/Private/IoStoreUtilities.cpp
@@ -1025,6 +1025,7 @@ struct FCookedPackage
 	uint64 TotalBulkDataSize = 0;
 	uint64 LoadOrder = MAX_uint64; // Ordered by dependencies
 	uint64 DiskLayoutOrder = MAX_uint64; // Final order after considering input order files
+	uint64 OverrideDiskLayoutOrder = MAX_uint64; // Final order after applying custom per-container ordering
 	FPackageStoreEntryResource PackageStoreEntry;
 	int32 PreOrderNumber = -1; // Used for sorting in load order
 	bool bPermanentMark = false; // Used for sorting in load order
@@ -1099,8 +1100,23 @@ struct FIoStoreOrderingOptions
 	bool bAlphaSortClusterPackageLists = false;
 
 	EFallbackOrderMode FallbackOrderMode = EFallbackOrderMode::LoadOrder;
+
+	bool operator==(const FIoStoreOrderingOptions& Other) const
+	{
+		return (bClusterByOrderFilePriority == Other.bClusterByOrderFilePriority) &&
+			(bAlphaSortClusterPackageLists == Other.bAlphaSortClusterPackageLists) &&
+			(FallbackOrderMode == Other.FallbackOrderMode);
+	}
 };
 
+uint32 GetTypeHash(const FIoStoreOrderingOptions& OrderingOptions)
+{
+	uint32 Hash = GetTypeHash(OrderingOptions.bClusterByOrderFilePriority);
+	Hash = HashCombine(Hash, GetTypeHash(OrderingOptions.bAlphaSortClusterPackageLists));
+	Hash = HashCombine(Hash, GetTypeHash(OrderingOptions.FallbackOrderMode));
+	return Hash;
+}
+
 struct FIoStoreArguments
 {
 	FString GlobalContainerPath;
@@ -1121,6 +1137,7 @@ struct FIoStoreArguments
 	TUniquePtr<FCookedPackageStore> PackageStore;
 	TUniquePtr<FIoBuffer> ScriptObjects;
 	FIoStoreOrderingOptions OrderingOptions;
+	TMap<FString, FIoStoreOrderingOptions> OrderingOptionsOverrides;
 	bool bSign = false;
 	bool bRemapPluginContentToGame = false;
 	bool bCreateDirectoryIndex = true;
@@ -1508,7 +1525,6 @@ private:
 };
 FClusterStatsCsv ClusterStatsCsv;
 
-
 // If bClusterByOrderFilePriority is false
 //		Order packages by the order of OrderMaps with their associated priority 
 //		e.g. Order map 1 with priority 0 (A, B, C) 
@@ -1525,7 +1541,8 @@ static void AssignPackagesDiskOrder(
 	const TArray<FCookedPackage*>& Packages,
 	const TArray<FFileOrderMap>& OrderMaps,
 	const FPackageIdMap& PackageIdMap,
-	const FIoStoreOrderingOptions& OrderingOptions )
+	const FIoStoreOrderingOptions& OrderingOptions,
+	bool bAssignOverrideDiskLayoutOrder)
 {
 	IOSTORE_CPU_SCOPE(AssignPackagesDiskOrder);
 
@@ -1817,7 +1834,14 @@ static void AssignPackagesDiskOrder(
 	{
 		for (FCookedPackage* Package : Cluster->Packages)
 		{
-			Package->DiskLayoutOrder = LayoutIndex++;
+			if (bAssignOverrideDiskLayoutOrder)
+			{
+				Package->OverrideDiskLayoutOrder = LayoutIndex++;
+			}
+			else
+			{
+				Package->DiskLayoutOrder = LayoutIndex++;
+			}
 		}
 		delete Cluster;
 	}
@@ -1825,16 +1849,38 @@ static void AssignPackagesDiskOrder(
 	ClusterStatsCsv.Close();
 }
 
+using FOrderingOptionsToContainersMap = TMap<FIoStoreOrderingOptions, TArray<const FContainerTargetSpec*>>;
+
 static void CreateDiskLayout(
 	const TArray<FContainerTargetSpec*>& ContainerTargets,
 	const TArray<FCookedPackage*>& Packages,
 	const TArray<FFileOrderMap>& OrderMaps,
 	const FPackageIdMap& PackageIdMap,
-	const FIoStoreOrderingOptions& OrderingOptions)
+	const FIoStoreOrderingOptions& OrderingOptions,
+	const FOrderingOptionsToContainersMap& OrderingOptionsOverrides)
 {
 	IOSTORE_CPU_SCOPE(CreateDiskLayout);
 
-	AssignPackagesDiskOrder(Packages, OrderMaps, PackageIdMap, OrderingOptions);
+	{
+		const bool bAssignOverrideDiskLayoutOrder = false;
+		AssignPackagesDiskOrder(Packages, OrderMaps, PackageIdMap, OrderingOptions, bAssignOverrideDiskLayoutOrder);
+	}
+
+	for (const TPair<FIoStoreOrderingOptions, TArray<const FContainerTargetSpec*>>& Override : OrderingOptionsOverrides)
+	{
+		const bool bAssignOverrideDiskLayoutOrder = true;
+		const FIoStoreOrderingOptions& OverrideOrderingOptions = Override.Key;
+		AssignPackagesDiskOrder(Packages, OrderMaps, PackageIdMap, OverrideOrderingOptions, bAssignOverrideDiskLayoutOrder);
+
+		const TArray<const FContainerTargetSpec*>& ContainersToApplyTo = Override.Value;
+		for (const FContainerTargetSpec* Container : ContainersToApplyTo)
+		{
+			for (FCookedPackage* Package : Container->Packages)
+			{
+				Package->DiskLayoutOrder = Package->OverrideDiskLayoutOrder;
+			}
+		}
+	}
 
 	for (FContainerTargetSpec* ContainerTarget : ContainerTargets)
 	{
@@ -5637,8 +5683,35 @@ int32 CreateTarget(const FIoStoreArguments& Arguments, const FIoStoreWriterSetti
 		IOSTORE_CPU_SCOPE(CreateClusterCSV);
 		ClusterStatsCsv.CreateOutputFile(ClusterCSVPath);
 	}
+
+	// Map and flip the relationship of the ordering options map
+	FOrderingOptionsToContainersMap OrderingOptionsOverrides;
+	for (FContainerTargetSpec* ContainerTarget : ContainerTargets)
+	{
+		const FString* PrevMatchedPrefix = nullptr;
+		const FString ContainerName = ContainerTarget->Name.ToString();
+		for (const TPair<FString, FIoStoreOrderingOptions>& OverridePair : Arguments.OrderingOptionsOverrides)
+		{
+			const FString& ContainerPrefix = OverridePair.Key;
+			const FIoStoreOrderingOptions& OrderingOptions = OverridePair.Value;
+			if (ContainerName.StartsWith(ContainerPrefix, ESearchCase::IgnoreCase))
+			{
+				if (PrevMatchedPrefix != nullptr)
+				{
+					UE_LOG(LogIoStore, Fatal, TEXT("Found multiple ordering option overrides ('%s' vs. '%s') that match '%s'."), **PrevMatchedPrefix, *ContainerPrefix, *ContainerName);
+				}
+
+				PrevMatchedPrefix = &ContainerPrefix;
+				UE_LOG(LogIoStore, Display,
+					TEXT("Applying ordering options override for %s (cluster-by-order: %d, alpha-sort-clusters: %d, fallback-mode: %d)"),
+					*ContainerName, OrderingOptions.bClusterByOrderFilePriority, OrderingOptions.bAlphaSortClusterPackageLists, OrderingOptions.FallbackOrderMode);
+				OrderingOptionsOverrides.FindOrAdd(OrderingOptions).Add(ContainerTarget);
+			}
+		}
+	}
+
 	SortPackagesInLoadOrder(Packages, PackageIdMap);
-	CreateDiskLayout(ContainerTargets, Packages, Arguments.OrderMaps, PackageIdMap, Arguments.OrderingOptions);
+	CreateDiskLayout(ContainerTargets, Packages, Arguments.OrderMaps, PackageIdMap, Arguments.OrderingOptions, OrderingOptionsOverrides);
 
 	{
 		IOSTORE_CPU_SCOPE(AppendContainerHeaderChunks);
@@ -9491,6 +9564,27 @@ static bool ParseSizeArgument(const TCHAR* CmdLine, const TCHAR* Argument, uint6
 	}
 }
 
+static EFallbackOrderMode ParseFallbackOrderMode(const FString& FallbackOrderStr)
+{
+	if (FallbackOrderStr == TEXT("LoadOrder"))
+	{
+		return EFallbackOrderMode::LoadOrder;
+	}
+	else if (FallbackOrderStr == TEXT("AlphabeticalClustered"))
+	{
+		return EFallbackOrderMode::AlphabeticalClustered;
+	}
+	else if (FallbackOrderStr == TEXT("Alphabetical"))
+	{
+		return EFallbackOrderMode::Alphabetical;
+	}
+	else
+	{
+		UE_LOG(LogIoStore, Fatal, TEXT("Unrecognized fallback order '%s'."), *FallbackOrderStr);
+		return EFallbackOrderMode::LoadOrder;
+	}
+}
+
 static bool ParseOrderFileArguments(FIoStoreArguments& Arguments)
 {
 	IOSTORE_CPU_SCOPE(ParseOrderFileArguments);
@@ -9561,21 +9655,37 @@ static bool ParseOrderFileArguments(FIoStoreArguments& Arguments)
 	FString FallbackOrderStr;
 	if (FParse::Value(FCommandLine::Get(), TEXT("-FallbackOrdering="), FallbackOrderStr))
 	{
-		if (FallbackOrderStr == TEXT("LoadOrder"))
-		{
-			Arguments.OrderingOptions.FallbackOrderMode = EFallbackOrderMode::LoadOrder;
-		}
-		else if (FallbackOrderStr == TEXT("AlphabeticalClustered"))
-		{
-			Arguments.OrderingOptions.FallbackOrderMode = EFallbackOrderMode::AlphabeticalClustered;
-		}
-		else if (FallbackOrderStr == TEXT("Alphabetical"))
+		Arguments.OrderingOptions.FallbackOrderMode = ParseFallbackOrderMode(FallbackOrderStr);
+	}
+
+	// Now apply the per-container-prefix overrides on top of the base configuration
+	FString PerContainerAlphaSortClusterPackageListsString;
+	if (FParse::Value(FCommandLine::Get(), TEXT("PerContainerAlphaSortClusterPackageLists="), PerContainerAlphaSortClusterPackageListsString, false))
+	{
+		TArray<FString> ContainerPrefixes;
+		PerContainerAlphaSortClusterPackageListsString.ParseIntoArray(ContainerPrefixes, TEXT(","), true);
+		for (FString& ContainerPrefix : ContainerPrefixes)
 		{
-			Arguments.OrderingOptions.FallbackOrderMode = EFallbackOrderMode::Alphabetical;
+			FIoStoreOrderingOptions& Options = Arguments.OrderingOptionsOverrides.FindOrAdd(MoveTemp(ContainerPrefix), Arguments.OrderingOptions);
+			Options.bAlphaSortClusterPackageLists = true;
 		}
-		else
+	}
+	if (FParse::Value(FCommandLine::Get(), TEXT("PerContainerFallbackOrdering="), PerContainerAlphaSortClusterPackageListsString, false))
+	{
+		TArray<FString> ContainerOrderPairStrings;
+		PerContainerAlphaSortClusterPackageListsString.ParseIntoArray(ContainerOrderPairStrings, TEXT(","), true);
+		for (FString& ContainerOrderPairString : ContainerOrderPairStrings)
 		{
-			UE_LOG(LogIoStore, Fatal, TEXT("Unrecognized fallback order '%s'."), *FallbackOrderStr);
+			TArray<FString> ContainerOrderPair;
+			if (ContainerOrderPairString.ParseIntoArray(ContainerOrderPair, TEXT(":"), true) != 2)
+			{
+				UE_LOG(LogIoStore, Fatal, TEXT("Failed to parse per-container fallback order pair '%s'."), *ContainerOrderPairString);
+			}
+
+			FString& ContainerPrefix = ContainerOrderPair[0];
+			const FString& OverrideFallbackOrderStr = ContainerOrderPair[1];
+			FIoStoreOrderingOptions& Options = Arguments.OrderingOptionsOverrides.FindOrAdd(MoveTemp(ContainerPrefix), Arguments.OrderingOptions);
+			Options.FallbackOrderMode = ParseFallbackOrderMode(OverrideFallbackOrderStr);
 		}
 	}
 	
diff --git a/Engine/Source/Programs/AutomationTool/Scripts/CopyBuildToStagingDirectory.Automation.cs b/Engine/Source/Programs/AutomationTool/Scripts/CopyBuildToStagingDirectory.Automation.cs
--- a/Engine/Source/Programs/AutomationTool/Scripts/CopyBuildToStagingDirectory.Automation.cs
+++ b/Engine/Source/Programs/AutomationTool/Scripts/CopyBuildToStagingDirectory.Automation.cs
@@ -4570,6 +4570,18 @@ namespace AutomationScripts
 					AdditionalArgs += " -AlphaSortClusterPackageLists";
 				}
 
+				List<string> PerContainerAlphaSortClusterPackageLists;
+				if (PlatformGameConfig.GetArray("/Script/UnrealEd.ProjectPackagingSettings", "PerContainerAlphaSortClusterPackageLists", out PerContainerAlphaSortClusterPackageLists))
+				{
+					AdditionalArgs += string.Format(" -PerContainerAlphaSortClusterPackageLists={0}", string.Join(',', PerContainerAlphaSortClusterPackageLists));
+				}
+
+				List<string> PerContainerFallbackOrdering;
+				if (PlatformGameConfig.GetArray("/Script/UnrealEd.ProjectPackagingSettings", "PerContainerFallbackOrdering", out PerContainerFallbackOrdering))
+				{
+					AdditionalArgs += string.Format(" -PerContainerFallbackOrdering={0}", string.Join(',', PerContainerFallbackOrdering));
+				}
+
 				AdditionalArgs += $" -platform={ConfigHierarchy.GetIniPlatformName(SC.StageTargetPlatform.IniPlatformType)}";
 
 				AdditionalArgs += " " + Params.AdditionalIoStoreOptions;
