Optimizing iostore container listings

diff --git a/Engine/Source/Developer/IoStoreUtilities/Private/IoStoreUtilities.cpp b/Engine/Source/Developer/IoStoreUtilities/Private/IoStoreUtilities.cpp
--- a/Engine/Source/Developer/IoStoreUtilities/Private/IoStoreUtilities.cpp
+++ b/Engine/Source/Developer/IoStoreUtilities/Private/IoStoreUtilities.cpp
@@ -6537,13 +6537,76 @@ static int32 FindAndExtractPackages(
 	return bAllSucceeded ? 0 : 1;
 }
 
+void IterateContainerFilesInternal(const FKeyChain& KeyChain, const TMap<FPackageId, TPair<FName, FName>>& AllChunksInfoMap, const FString& ContainerFilePath, FChunkEntryCsv* Out)
+{
+	TUniquePtr<FIoStoreReader> Reader = CreateIoStoreReader(*ContainerFilePath, KeyChain);
+	if (!Reader.IsValid())
+	{
+		UE_LOG(LogIoStore, Warning, TEXT("Failed to read container '%s'"), *ContainerFilePath);
+		return;
+	}
+
+	Out->SetPartitionSize(Reader->GetPartitionSize());
+
+	if (!EnumHasAnyFlags(Reader->GetContainerFlags(), EIoContainerFlags::Indexed))
+	{
+		UE_LOG(LogIoStore, Display, TEXT("No directory index for container '%s'"), *ContainerFilePath);
+	}
+
+	UE_LOG(LogIoStore, Display, TEXT("Listing container '%s'"), *ContainerFilePath);
+
+	FString ContainerName = FPaths::GetBaseFilename(ContainerFilePath);
+	TArray<FIoStoreTocChunkInfo> Chunks;
+
+	{
+		IOSTORE_CPU_SCOPE(EnumerateChunks);
+		Reader->EnumerateChunks([&Chunks](FIoStoreTocChunkInfo&& ChunkInfo)
+			{
+				Chunks.Add(MoveTemp(ChunkInfo));
+				return true;
+			});
+	}
+
+	{
+		IOSTORE_CPU_SCOPE(EnumerateChunks);
+		auto SortKey = [](const FIoStoreTocChunkInfo& ChunkInfo) { return ChunkInfo.OffsetOnDisk; };
+		Algo::SortBy(Chunks, SortKey);
+	}
+
+	{
+		IOSTORE_CPU_SCOPE(WriteCsvFile);
+		FString PackageName;
+		FString ClassType;
+		for (int32 Index = 0; Index < Chunks.Num(); ++Index)
+		{
+			const FIoStoreTocChunkInfo& ChunkInfo = Chunks[Index];
+
+			FPackageId PackageId = PackageIdFromChunkId(ChunkInfo.Id);
+			PackageName.Reset();
+			ClassType.Reset();
+			if (ChunkInfo.bHasValidFileName && FPackageName::TryConvertFilenameToLongPackageName(ChunkInfo.FileName, PackageName, nullptr))
+			{
+				PackageId = FPackageId::FromName(FName(*PackageName));
+			}
+
+			const TPair<FName, FName>* PackageInfo = AllChunksInfoMap.Find(PackageId);
+			if (PackageInfo)
+			{
+				PackageName = PackageInfo->Get<0>().ToString();
+				ClassType = PackageInfo->Get<1>().ToString();
+			}
+
+			Out->AddChunk(ContainerName, Index, ChunkInfo, PackageId, PackageName, ClassType);
+		}
+	}
+}
+
 int32 ListContainer(
 	const FKeyChain& KeyChain,
 	const FString& ContainerPathOrWildcard,
 	const FString& CsvPath)
 {
-	IOSTORE_CPU_SCOPE(ListContainer);
-	
+	IOSTORE_CPU_SCOPE(ListContainer);	
 
 	FString AllChunksInfoFilename;
 	TMap<FPackageId, TPair<FName,FName>> AllChunksInfoMap;
@@ -6578,91 +6641,29 @@ int32 ListContainer(
 	
 	// if CsvPath is a dir, not a file name, then write one csv per container to the dir
 	// otherwise, write all contents to one big csv
-	bool bPerContainerCsvFiles = IFileManager::Get().DirectoryExists(*CsvPath);
-
-	FChunkEntryCsv AllContainersOutCsvFile;
-	if (!bPerContainerCsvFiles)
+	if (IFileManager::Get().DirectoryExists(*CsvPath))
 	{
-		AllContainersOutCsvFile.CreateOutputFile(*CsvPath);
-	}
-
-	for (const FString& ContainerFilePath : ContainerFilePaths)
-	{
-		TUniquePtr<FIoStoreReader> Reader = CreateIoStoreReader(*ContainerFilePath, KeyChain);
-		if (!Reader.IsValid())
+		ParallelFor(ContainerFilePaths.Num(), [&ContainerFilePaths, &AllChunksInfoMap, &KeyChain, &CsvPath](int32 Idx)
 		{
-			UE_LOG(LogIoStore, Warning, TEXT("Failed to read container '%s'"), *ContainerFilePath);
-			continue;
-		}
-		
-		FChunkEntryCsv PerContainerOutCsvFile;
-		PerContainerOutCsvFile.SetPartitionSize(Reader->GetPartitionSize());
-		AllContainersOutCsvFile.SetPartitionSize(Reader->GetPartitionSize());
-		FChunkEntryCsv* Out;
-		
-		if (!bPerContainerCsvFiles)
-		{
-			Out = &AllContainersOutCsvFile;
-		}
-		else
-		{
-			// ContainerFilePath is a .utoc, add .csv and put in CsvPath
-			FString PerContainerCsvPath = CsvPath / FPaths::GetCleanFilename(ContainerFilePath) + TEXT(".csv");
-			PerContainerOutCsvFile.CreateOutputFile(*PerContainerCsvPath);
-			Out = &PerContainerOutCsvFile;
-		}
-
-		if (!EnumHasAnyFlags(Reader->GetContainerFlags(), EIoContainerFlags::Indexed))
-		{
-			UE_LOG(LogIoStore, Display, TEXT("No directory index for container '%s'"), *ContainerFilePath);
-		}
-
-		UE_LOG(LogIoStore, Display, TEXT("Listing container '%s'"), *ContainerFilePath);
-
-		FString ContainerName = FPaths::GetBaseFilename(ContainerFilePath);
-		TArray<FIoStoreTocChunkInfo> Chunks;
-
-		{
-			IOSTORE_CPU_SCOPE(EnumerateChunks);
-			Reader->EnumerateChunks([&Chunks](FIoStoreTocChunkInfo&& ChunkInfo)
-			{
-				Chunks.Add(MoveTemp(ChunkInfo));
-				return true;
-			});
-		}
-
-		{
-			IOSTORE_CPU_SCOPE(EnumerateChunks);
-			auto SortKey = [](const FIoStoreTocChunkInfo& ChunkInfo) { return ChunkInfo.OffsetOnDisk; };
-			Algo::SortBy(Chunks, SortKey);
-		}
-
-		{
-			IOSTORE_CPU_SCOPE(WriteCsvFile);
-			FString PackageName;
-			FString ClassType;
-			for (int32 Index=0; Index < Chunks.Num(); ++Index)
-			{
-				const FIoStoreTocChunkInfo& ChunkInfo = Chunks[Index];
+			const FString& ContainerFilePath = ContainerFilePaths[Idx];
 
-				FPackageId PackageId = PackageIdFromChunkId(ChunkInfo.Id);
-				PackageName.Reset();
-				ClassType.Reset();
-				if (ChunkInfo.bHasValidFileName && FPackageName::TryConvertFilenameToLongPackageName(ChunkInfo.FileName, PackageName, nullptr))
-				{
-					PackageId = FPackageId::FromName(FName(*PackageName));
-				}
+			FChunkEntryCsv PerContainerOutCsvFile;
 
-				TPair<FName,FName>* PackageInfo = AllChunksInfoMap.Find(PackageId);
-				if (PackageInfo)
-				{
-					PackageName = PackageInfo->Get<0>().ToString();
-					ClassType = PackageInfo->Get<1>().ToString();
-				}
+			FString PerContainerCsvPath = CsvPath / FPaths::GetCleanFilename(ContainerFilePath) + TEXT(".csv");
+			PerContainerOutCsvFile.CreateOutputFile(*PerContainerCsvPath);
 
-				Out->AddChunk(ContainerName, Index, ChunkInfo, PackageId, PackageName, ClassType);
-			}
+			IterateContainerFilesInternal(KeyChain, AllChunksInfoMap, ContainerFilePath, &PerContainerOutCsvFile);
+		});
+	}
+	else
+	{
+		FChunkEntryCsv AllContainersOutCsvFile;
+		AllContainersOutCsvFile.CreateOutputFile(*CsvPath);
+		for (const FString& ContainerFilePath : ContainerFilePaths)
+		{			
+			IterateContainerFilesInternal(KeyChain, AllChunksInfoMap, ContainerFilePath, &AllContainersOutCsvFile);
 		}
+		return 0;
 	}
 
 	return 0;
diff --git a/Engine/Source/Runtime/Core/Private/IO/IoStore.cpp b/Engine/Source/Runtime/Core/Private/IO/IoStore.cpp
--- a/Engine/Source/Runtime/Core/Private/IO/IoStore.cpp
+++ b/Engine/Source/Runtime/Core/Private/IO/IoStore.cpp
@@ -409,30 +409,36 @@ public:
 		TocFilePath.Append(InContainerPath);
 		TocFilePath.Append(TEXT(".utoc"));
 
-		FIoStatus TocStatus = TocReader.Read(*TocFilePath, InDecryptionKeys);
-		if (!TocStatus.IsOk())
 		{
-			return TocStatus;
+			TRACE_CPUPROFILER_EVENT_SCOPE(FIoStoreReader::ReadToc);
+			FIoStatus TocStatus = TocReader.Read(*TocFilePath, InDecryptionKeys);
+			if (!TocStatus.IsOk())
+			{
+				return TocStatus;
+			}
 		}
 
-		const FIoStoreTocResourceView& TocResource = TocReader.GetTocResource();
-
-		IPlatformFile& Ipf = FPlatformFileManager::Get().GetPlatformFile();
-		ContainerFileAccessors.Reserve(TocResource.Header.PartitionCount);
-		for (uint32 PartitionIndex = 0; PartitionIndex < TocResource.Header.PartitionCount; ++PartitionIndex)
 		{
-			TStringBuilder<256> ContainerFilePath;
-			ContainerFilePath.Append(InContainerPath);
-			if (PartitionIndex > 0)
-			{
-				ContainerFilePath.Append(FString::Printf(TEXT("_s%d"), PartitionIndex));
-			}
-			ContainerFilePath.Append(TEXT(".ucas"));
+			TRACE_CPUPROFILER_EVENT_SCOPE(FIoStoreReader::GenerateContainerFileAccessors);
+			const FIoStoreTocResourceView& TocResource = TocReader.GetTocResource();
 
-			ContainerFileAccessors.Emplace(new FContainerFileAccess(Ipf, *ContainerFilePath));
-			if (ContainerFileAccessors.Last().IsValid() == false)
+			IPlatformFile& Ipf = FPlatformFileManager::Get().GetPlatformFile();
+			ContainerFileAccessors.Reserve(TocResource.Header.PartitionCount);
+			for (uint32 PartitionIndex = 0; PartitionIndex < TocResource.Header.PartitionCount; ++PartitionIndex)
 			{
-				return FIoStatusBuilder(EIoErrorCode::FileOpenFailed) << TEXT("Failed to open IoStore container file '") << *TocFilePath << TEXT("'");
+				TStringBuilder<256> ContainerFilePath;
+				ContainerFilePath.Append(InContainerPath);
+				if (PartitionIndex > 0)
+				{
+					ContainerFilePath.Append(FString::Printf(TEXT("_s%d"), PartitionIndex));
+				}
+				ContainerFilePath.Append(TEXT(".ucas"));
+
+				ContainerFileAccessors.Emplace(new FContainerFileAccess(Ipf, *ContainerFilePath));
+				if (ContainerFileAccessors.Last().IsValid() == false)
+				{
+					return FIoStatusBuilder(EIoErrorCode::FileOpenFailed) << TEXT("Failed to open IoStore container file '") << *TocFilePath << TEXT("'");
+				}
 			}
 		}
 
