IncrementalCook: CookArtifactReader: Ask the ZenStoreReader to InitializeConnection so the ue.projectstore file is available before asking the ZenCookArtifactReader to connect to the oplog, which uses the ue.projectstore file. Fixes the failure to read CookSettings.txt from zen (and therefore incurring a full recook) when cooking incrementally from snapshot with no previous local cook.
#rnx

diff --git a/Engine/Source/Developer/IoStoreUtilities/Private/IoStoreUtilities.cpp b/Engine/Source/Developer/IoStoreUtilities/Private/IoStoreUtilities.cpp
--- a/Engine/Source/Developer/IoStoreUtilities/Private/IoStoreUtilities.cpp
+++ b/Engine/Source/Developer/IoStoreUtilities/Private/IoStoreUtilities.cpp
@@ -9070,6 +9070,9 @@ int32 Staged2Zen(const FString& BuildPath, const FKeyChain& KeyChain, const FStr
 	
 	ICookedPackageWriter::FCookInfo CookInfo;
 	CookInfo.bFullBuild = true;
+	ZenStoreWriter->InitializeConnection();
+	// Reader InitializeConnection must come after writer InitializeConnection, it uses data written by that function.
+	ZenCookArtifactReader->InitializeConnection();
 	ZenStoreWriter->Initialize(CookInfo);
 	ZenStoreWriter->BeginCook(CookInfo);
 	int32 LocalPackageIndex = 0;
diff --git a/Engine/Source/Developer/IoStoreUtilities/Private/ZenCookArtifactReader.cpp b/Engine/Source/Developer/IoStoreUtilities/Private/ZenCookArtifactReader.cpp
--- a/Engine/Source/Developer/IoStoreUtilities/Private/ZenCookArtifactReader.cpp
+++ b/Engine/Source/Developer/IoStoreUtilities/Private/ZenCookArtifactReader.cpp
@@ -17,12 +17,8 @@ FZenCookArtifactReader::FZenCookArtifactReader(
 #if !UE_BUILD_SHIPPING
 	: ScopeZenService(MakeUnique<UE::Zen::FScopeZenService>())
 	, ZenRootPath(InputPath)
-	, StorageServerPlatformFile(IStorageServerClientModule::Get().TryCreateCustomPlatformFile(*InputPath, &FPlatformFileManager::Get().GetPlatformFile()))
+	, StorageServerPlatformFile(nullptr)
 {
-	if (StorageServerPlatformFile)
-	{
-		StorageServerPlatformFile->SetLowerLevel(nullptr);
-	}
 }
 #else
 {
@@ -33,8 +29,24 @@ FZenCookArtifactReader::~FZenCookArtifactReader()
 {
 }
 
+void FZenCookArtifactReader::InitializeConnection()
+{
+	if (StorageServerPlatformFile)
+	{
+		return;
+	}
+
+	StorageServerPlatformFile.Reset(IStorageServerClientModule::Get().TryCreateCustomPlatformFile(*ZenRootPath,
+		&FPlatformFileManager::Get().GetPlatformFile()));
+	if (StorageServerPlatformFile)
+	{
+		StorageServerPlatformFile->SetLowerLevel(nullptr);
+	}
+}
+
 void FZenCookArtifactReader::Initialize(bool bCleanBuild)
 {
+	InitializeConnection();
 #if !UE_BUILD_SHIPPING
 	if (bCleanBuild && StorageServerPlatformFile)
 	{
diff --git a/Engine/Source/Developer/IoStoreUtilities/Private/ZenStoreWriter.cpp b/Engine/Source/Developer/IoStoreUtilities/Private/ZenStoreWriter.cpp
--- a/Engine/Source/Developer/IoStoreUtilities/Private/ZenStoreWriter.cpp
+++ b/Engine/Source/Developer/IoStoreUtilities/Private/ZenStoreWriter.cpp
@@ -207,8 +207,8 @@ void FZenStoreWriter::StaticInit()
 }
 
 FZenStoreWriter::FZenStoreWriter(
-	const FString& InOutputPath, 
-	const FString& InMetadataDirectoryPath, 
+	const FString& InOutputPath,
+	const FString& InMetadataDirectoryPath,
 	const ITargetPlatform* InTargetPlatform,
 	TSharedRef<ICookArtifactReader> InCookArtifactReader
 )
@@ -220,6 +220,7 @@ FZenStoreWriter::FZenStoreWriter(
 	, PackageStoreOptimizer(new FPackageStoreOptimizer())
 	, CookMode(ICookedPackageWriter::FCookInfo::CookByTheBookMode)
 	, bInitialized(false)
+	, bInitializedConnection(false)
 	, bProvidePerPackageResults(false)
 {
 	StaticInit();
@@ -230,7 +231,7 @@ FZenStoreWriter::FZenStoreWriter(
 	{
 		OplogId = InTargetPlatform->PlatformName();
 	}
-	
+
 	HttpClient = MakeUnique<UE::FZenStoreHttpClient>();
 
 #if UE_WITH_ZEN
@@ -256,9 +257,9 @@ FZenStoreWriter::FZenStoreWriter(
 	PackageStoreOptimizer->Initialize();
 
 	ZenFileSystemManifest = MakeUnique<FZenFileSystemManifest>(TargetPlatform, OutputPath);
-	
+
 	CommitQueue = MakeUnique<FCommitQueue>();
-	
+
 	Compressor = FOodleDataCompression::ECompressor::Mermaid;
 	CompressionLevel = FOodleDataCompression::ECompressionLevel::VeryFast;
 }
@@ -465,12 +466,43 @@ void FZenStoreWriter::RegisterDeterminismHelper(UObject* SourceObject,
 	}
 }
 
+FString FZenStoreWriter::GetProjectStorePath() const
+{
+	return OutputPath / TEXT("ue.projectstore");
+}
+
+void FZenStoreWriter::InitializeConnection()
+{
+	if (bInitializedConnection)
+	{
+		return;
+	}
+	WriteProjectStoreFile();
+	bInitializedConnection = true;
+}
+
+void FZenStoreWriter::WriteProjectStoreFile()
+{
+	// We need to write the ProjectStore file early, for use both during this->Initialize and by
+	// FZenCookArtifactReader->InitializeConnection.
+	FStringView ServiceHostName = TEXT("localhost");
+	uint16 ServicePort = 8558;
+	Zen::FZenProjectStoreWriter ProjectStoreWriter(*GetProjectStorePath());
+#if UE_WITH_ZEN
+	const Zen::FZenServiceInstance& ZenServiceInstance = HttpClient->GetZenServiceInstance();
+	ServiceHostName = ZenServiceInstance.GetEndpoint().GetName();
+	ServicePort = ZenServiceInstance.GetEndpoint().GetPort();
+#endif
+	ProjectStoreWriter.Write(ServiceHostName, ServicePort, IsLocalConnection, *ProjectId, *OplogId, TargetPlatformFName);
+}
+
 void FZenStoreWriter::Initialize(const FCookInfo& Info)
 {
 	CookMode = Info.CookMode;
 
 	if (!bInitialized)
 	{
+		InitializeConnection();
 		bool CleanBuild = Info.bFullBuild && !Info.bWorkerOnSharedSandbox;
 		if (CleanBuild)
 		{
@@ -483,28 +515,15 @@ void FZenStoreWriter::Initialize(const FCookInfo& Info)
 			IFileManager::Get().DeleteDirectory(*OutputPath, bRequireExists, bTree);
 		}
 
-		FString OplogLifetimeMarkerPath = OutputPath / TEXT("ue.projectstore");
-		Zen::FZenProjectStoreWriter ProjectStoreWriter(*OplogLifetimeMarkerPath);
-
-		bool bOplogEstablished = HttpClient->TryCreateOplog(ProjectId, OplogId, OplogLifetimeMarkerPath);
-
-		FStringView ServiceHostName = TEXT("localhost");
-		uint16 ServicePort = 8558;
-		bool IsRunningLocally = false;
-
-#if UE_WITH_ZEN
-		const Zen::FZenServiceInstance& ZenServiceInstance = HttpClient->GetZenServiceInstance();
-		IsRunningLocally = ZenServiceInstance.IsServiceRunningLocally();
-		ServiceHostName = ZenServiceInstance.GetEndpoint().GetName();
-		ServicePort = ZenServiceInstance.GetEndpoint().GetPort();
-#endif
-
-		ProjectStoreWriter.Write(ServiceHostName, ServicePort, IsRunningLocally, *ProjectId, *OplogId, TargetPlatformFName);
-
-		if (!bOplogEstablished && CleanBuild)
+		// If we deleted the directory above for a clean build, or the cooker deleted it, then rewrite
+		// the ProjectStorePath for use by TryCreateOplog and any other systems that want to query the oplog
+		// during cook.
+		if (!IFileManager::Get().FileExists(*GetProjectStorePath()))
 		{
-			IFileManager::Get().Delete(*OplogLifetimeMarkerPath);
+			WriteProjectStoreFile();
 		}
+
+		bool bOplogEstablished = HttpClient->TryCreateOplog(ProjectId, OplogId, GetProjectStorePath());
 		UE_CLOG(!bOplogEstablished, LogZenStoreWriter, Fatal, TEXT("Failed to establish oplog on the ZenServer"));
 
 		if (!Info.bFullBuild)
diff --git a/Engine/Source/Developer/IoStoreUtilities/Public/ZenCookArtifactReader.h b/Engine/Source/Developer/IoStoreUtilities/Public/ZenCookArtifactReader.h
--- a/Engine/Source/Developer/IoStoreUtilities/Public/ZenCookArtifactReader.h
+++ b/Engine/Source/Developer/IoStoreUtilities/Public/ZenCookArtifactReader.h
@@ -18,6 +18,7 @@ public:
 												const FString& MetadataDirectoryPath, 
 												const ITargetPlatform* TargetPlatform);
 	IOSTOREUTILITIES_API virtual ~FZenCookArtifactReader();
+	IOSTOREUTILITIES_API void InitializeConnection();
 	IOSTOREUTILITIES_API void Initialize(bool bCleanBuild) override;
 
 	IOSTOREUTILITIES_API bool FileExists(const TCHAR* Filename) override;
diff --git a/Engine/Source/Developer/IoStoreUtilities/Public/ZenStoreWriter.h b/Engine/Source/Developer/IoStoreUtilities/Public/ZenStoreWriter.h
--- a/Engine/Source/Developer/IoStoreUtilities/Public/ZenStoreWriter.h
+++ b/Engine/Source/Developer/IoStoreUtilities/Public/ZenStoreWriter.h
@@ -130,6 +130,7 @@ public:
 	IOSTOREUTILITIES_API virtual void RegisterDeterminismHelper(UObject* SourceObject,
 		const TRefCountPtr<UE::Cook::IDeterminismHelper>& DeterminismHelper) override;
 
+	IOSTOREUTILITIES_API void InitializeConnection();
 	IOSTOREUTILITIES_API virtual void Initialize(const FCookInfo& Info) override;
 	IOSTOREUTILITIES_API virtual void BeginCook(const FCookInfo& Info) override;
 	IOSTOREUTILITIES_API virtual void EndCook(const FCookInfo& Info) override;
@@ -252,6 +253,8 @@ private:
 	void CommitPackageInternal(FZenCommitInfo&& CommitInfo);
 	FCbAttachment CreateAttachment(FSharedBuffer Buffer) const;
 	FCbAttachment CreateAttachment(FIoBuffer Buffer) const;
+	FString GetProjectStorePath() const;
+	void WriteProjectStoreFile();
 
 	FCriticalSection								PackagesCriticalSection;
 	TMap<FName, TUniquePtr<FPendingPackageState>>	PendingPackages;
@@ -296,6 +299,7 @@ private:
 	TFuture<void>						CommitThread;
 
 	bool								bInitialized;
+	bool								bInitializedConnection;
 	bool								bProvidePerPackageResults;
 
 	static void StaticInit();
diff --git a/Engine/Source/Editor/UnrealEd/Private/CookOnTheFlyServer.cpp b/Engine/Source/Editor/UnrealEd/Private/CookOnTheFlyServer.cpp
--- a/Engine/Source/Editor/UnrealEd/Private/CookOnTheFlyServer.cpp
+++ b/Engine/Source/Editor/UnrealEd/Private/CookOnTheFlyServer.cpp
@@ -10606,6 +10606,9 @@ UE::Cook::FCookSavePackageContext* UCookOnTheFlyServer::CreateSaveContext(const
 		CookArtifactReader = LayeredReader;
 		FZenStoreWriter* ZenWriter = new FZenStoreWriter(ResolvedRootPath, ResolvedMetadataPath, TargetPlatform, LayeredReader);
 		ZenWriter->SetCooker(this);
+		// Reader InitializeConnection must come after writer InitializeConnection, it uses data written by that function.
+		ZenWriter->InitializeConnection();
+		ZenReader->InitializeConnection();
 		PackageWriter = ZenWriter;
 		WriterDebugName = TEXT("ZenStore");
 	}
