Cooker: CookArtifactReader: fix compile error in previous change in the shipping configuration.
#rnx

diff --git a/Engine/Source/Developer/IoStoreUtilities/Private/ZenCookArtifactReader.cpp b/Engine/Source/Developer/IoStoreUtilities/Private/ZenCookArtifactReader.cpp
--- a/Engine/Source/Developer/IoStoreUtilities/Private/ZenCookArtifactReader.cpp
+++ b/Engine/Source/Developer/IoStoreUtilities/Private/ZenCookArtifactReader.cpp
@@ -31,6 +31,7 @@ FZenCookArtifactReader::~FZenCookArtifactReader()
 
 void FZenCookArtifactReader::InitializeConnection()
 {
+#if !UE_BUILD_SHIPPING
 	if (StorageServerPlatformFile)
 	{
 		return;
@@ -42,6 +43,7 @@ void FZenCookArtifactReader::InitializeConnection()
 	{
 		StorageServerPlatformFile->SetLowerLevel(nullptr);
 	}
+#endif
 }
 
 void FZenCookArtifactReader::Initialize(bool bCleanBuild)
