Better logging for cases where super large content causes an iostore partition size assert.

#rnx

diff --git a/Engine/Source/Developer/IoStoreUtilities/Private/IoStoreWriter.cpp b/Engine/Source/Developer/IoStoreUtilities/Private/IoStoreWriter.cpp
--- a/Engine/Source/Developer/IoStoreUtilities/Private/IoStoreWriter.cpp
+++ b/Engine/Source/Developer/IoStoreUtilities/Private/IoStoreWriter.cpp
@@ -2021,7 +2021,17 @@ private:
 		bHasMemoryMappedEntry |= Entry->Options.bIsMemoryMapped;
 		const uint64 ChunkAlignment = Entry->Options.bIsMemoryMapped ? WriterSettings.MemoryMappingAlignment : 0;
 		const uint64 PartitionSizeLimit = MaxPartitionSize > 0 ? MaxPartitionSize : MAX_uint64;
-		checkf(Entry->DiskSize <= PartitionSizeLimit, TEXT("Chunk is too large, increase max partition size!"));
+		if (Entry->DiskSize > PartitionSizeLimit)
+		{
+			// Someone added some super large content, give a detailed log so they know what to fix.
+			const TCHAR* DebugName = Entry->Options.DebugName ? Entry->Options.DebugName : TEXT("NoDebugName");
+
+			// This must be fatal as the loop below doesn't handle partition assignment of chunks larger than the partition size.
+			UE_LOG(LogIoStore, Fatal, TEXT("Chunk is too large: DiskSize = %llu CompressedSize = %llu UncompressedSize = %s ChunkId = %s DebugName = %s PartitionSize = %llu"),
+				Entry->DiskSize, Entry->CompressedSize, Entry->UncompressedSize.IsSet() ? *WriteToString<32>(Entry->UncompressedSize.GetValue()) : TEXT("(Unknown)"),
+				*WriteToString<32>(Entry->ChunkId), DebugName, PartitionSizeLimit);
+		}
+
 		for (;;)
 		{
 			uint64 OffsetBeforePadding = TargetPartition->Offset;
