[IoStore] Move shader library processing in IoStore utilities into separate source files.

This serves as initial preparation to refactor how we package cooked shaders into the IoStore containers to better accommodate pak rules for shaders.
Instead of having loose functions that pass many input and output parameters around, we will refactor the shader library processing into FIoStoreShaderLibraryProcessor.
For simplicity of this change, this mostly just moves code into new private source files and puts its code into UE::IoStore::Private namespace.

#rnx
#Synced-CL 48391606

diff --git a/Engine/Source/Developer/IoStoreUtilities/Private/IoStoreShaderLibraryProcessor.cpp b/Engine/Source/Developer/IoStoreUtilities/Private/IoStoreShaderLibraryProcessor.cpp
new file mode 100644
--- /dev/null
+++ b/Engine/Source/Developer/IoStoreUtilities/Private/IoStoreShaderLibraryProcessor.cpp
@@ -0,0 +1,563 @@
+// Copyright Epic Games, Inc. All Rights Reserved.
+
+#include "IoStoreShaderLibraryProcessor.h"
+
+#include "Async/ParallelFor.h"
+#include "IO/IoStore.h"
+#include "IoStoreUtilitiesCore.h"
+#include "Misc/FileHelper.h"
+#include "Serialization/JsonReader.h"
+#include "Serialization/JsonSerializer.h"
+#include "Serialization/JsonWriter.h"
+#include "Serialization/LargeMemoryReader.h"
+#include "Serialization/LargeMemoryWriter.h"
+#include "Serialization/MemoryReader.h"
+#include "ShaderCodeArchive.h"
+
+
+#define IOSTORE_CPU_SCOPE(NAME) TRACE_CPUPROFILER_EVENT_SCOPE(IoStore##NAME);
+#define IOSTORE_CPU_SCOPE_DATA(NAME, DATA) TRACE_CPUPROFILER_EVENT_SCOPE(IoStore##NAME);
+
+namespace UE::IoStore::Private
+{
+
+static bool LoadShaderAssetInfo(const FString& Filename, FCookedPackageStore* PackageStore, TMap<FSHAHash, TSet<FName>>& OutShaderCodeToAssets)
+{
+	FString JsonText;
+	if (PackageStore && PackageStore->HasZenStoreClient())
+	{
+		FIoChunkId ChunkId = PackageStore->GetChunkIdFromFileName(Filename);
+		if (!ChunkId.IsValid())
+		{
+			return false;
+		}
+		TIoStatusOr<FIoBuffer> Buffer = PackageStore->ReadChunk(ChunkId);
+		if (!Buffer.IsOk())
+		{
+			return false;
+		}
+		FFileHelper::BufferToString(JsonText, Buffer.ValueOrDie().GetData(), (int32)Buffer.ValueOrDie().GetSize());
+	}
+	else
+	{
+		if (!FFileHelper::LoadFileToString(JsonText, *Filename))
+		{
+			return false;
+		}
+	}
+
+	TSharedPtr<FJsonObject> JsonObject;
+	TSharedRef<TJsonReader<TCHAR>> Reader = TJsonReaderFactory<TCHAR>::Create(JsonText);
+
+	if (!FJsonSerializer::Deserialize(Reader, JsonObject) || !JsonObject.IsValid())
+	{
+		return false;
+	}
+
+	TSharedPtr<FJsonValue> AssetInfoArrayValue = JsonObject->Values.FindRef(TEXT("ShaderCodeToAssets"));
+	if (!AssetInfoArrayValue.IsValid())
+	{
+		return false;
+	}
+
+	TArray<TSharedPtr<FJsonValue>> AssetInfoArray = AssetInfoArrayValue->AsArray();
+	for (int32 IdxPair = 0, NumPairs = AssetInfoArray.Num(); IdxPair < NumPairs; ++IdxPair)
+	{
+		TSharedPtr<FJsonObject> Pair = AssetInfoArray[IdxPair]->AsObject();
+		if (Pair.IsValid())
+		{
+			TSharedPtr<FJsonValue> ShaderMapHashJson = Pair->Values.FindRef(TEXT("ShaderMapHash"));
+			if (!ShaderMapHashJson.IsValid())
+			{
+				continue;
+			}
+			FSHAHash ShaderMapHash;
+			ShaderMapHash.FromString(ShaderMapHashJson->AsString());
+
+			TSharedPtr<FJsonValue> AssetPathsArrayValue = Pair->Values.FindRef(TEXT("Assets"));
+			if (!AssetPathsArrayValue.IsValid())
+			{
+				continue;
+			}
+
+			TSet<FName>& Assets = OutShaderCodeToAssets.Add(ShaderMapHash);
+			TArray<TSharedPtr<FJsonValue>> AssetPathsArray = AssetPathsArrayValue->AsArray();
+			for (int32 IdxAsset = 0, NumAssets = AssetPathsArray.Num(); IdxAsset < NumAssets; ++IdxAsset)
+			{
+				Assets.Add(FName(*AssetPathsArray[IdxAsset]->AsString()));
+			}
+		}
+	}
+	return true;
+}
+
+static bool ConvertToIoStoreShaderLibrary(
+	const TCHAR* FileName,
+	FIoChunkId InChunkId,
+	FCookedPackageStore* PackageStore,
+	FOodleDataCompression::ECompressor InShaderOodleCompressor,
+	FOodleDataCompression::ECompressionLevel InShaderOodleLevel,
+	TTuple<FIoChunkId, FIoBuffer>& OutLibraryIoChunk,
+	TArray<TTuple<FIoChunkId, FIoBuffer, uint32>>& OutCodeIoChunks,
+	TArray<TTuple<FSHAHash, TArray<FIoChunkId>>>& OutShaderMaps,
+	TArray<TTuple<FSHAHash, TSet<FName>>>& OutShaderMapAssetAssociations)
+{
+	IOSTORE_CPU_SCOPE(ConvertShaderLibrary);
+	// ShaderArchive-<MyProject>_Chunk<N>-<ShaderFormat>-<ShaderPlatform>.ushaderbytecode
+	// E.g. "ShaderArchive-Lyra_Chunk0-PCD3D_SM6-PCD3D_SM6.ushaderbytecode
+	FStringView BaseFileNameView = FPathViews::GetBaseFilename(FileName);
+	int32 FormatStartIndex = -1;
+	if (!BaseFileNameView.FindLastChar('-', FormatStartIndex))
+	{
+		UE_LOG(LogIoStore, Error, TEXT("Invalid shader code library file name '%s'."), FileName);
+		return false;
+	}
+	int32 LibraryNameStartIndex = -1;
+	if (!BaseFileNameView.FindChar('-', LibraryNameStartIndex) || FormatStartIndex - LibraryNameStartIndex <= 1)
+	{
+		UE_LOG(LogIoStore, Error, TEXT("Invalid shader code library file name '%s'."), FileName);
+		return false;
+	}
+	FString LibraryName(BaseFileNameView.Mid(LibraryNameStartIndex + 1, FormatStartIndex - LibraryNameStartIndex - 1));
+	FName FormatName(BaseFileNameView.RightChop(FormatStartIndex + 1));
+
+	TRACE_CPUPROFILER_EVENT_SCOPE_TEXT(*LibraryName);
+
+	FIoBuffer LibraryBuffer;
+	if (PackageStore && PackageStore->HasZenStoreClient() && InChunkId.IsValid())
+	{
+		LibraryBuffer = PackageStore->ReadChunk(InChunkId).ConsumeValueOrDie();
+	}
+	else
+	{
+		TUniquePtr<FArchive> LibraryAr(IFileManager::Get().CreateFileReader(FileName));
+		if (!LibraryAr)
+		{
+			UE_LOG(LogIoStore, Error, TEXT("Missing shader code library file '%s'."), FileName);
+			return false;
+		}
+
+		LibraryBuffer = FIoBuffer(LibraryAr->TotalSize());
+		LibraryAr->Serialize(LibraryBuffer.GetData(), LibraryBuffer.GetSize());
+	}
+
+	FLargeMemoryReader LibraryAr(LibraryBuffer.GetData(), LibraryBuffer.GetSize());
+	uint32 Version;
+	LibraryAr << Version;
+
+	FSerializedShaderArchive SerializedShaders;
+	LibraryAr << SerializedShaders;
+	int64 OffsetToShaderCode = LibraryAr.Tell();
+
+	FString AssetInfoFileName = FPaths::GetPath(FileName) / FString::Printf(TEXT("ShaderAssetInfo-%s-%s.assetinfo.json"), *LibraryName, *FormatName.ToString());
+	TMap<FSHAHash, FShaderMapAssetPaths> ShaderCodeToAssets;
+	if (!LoadShaderAssetInfo(AssetInfoFileName, PackageStore && InChunkId.IsValid() ? PackageStore : nullptr, ShaderCodeToAssets))
+	{
+		UE_LOG(LogIoStore, Error, TEXT("Failed loading asset asset info file '%s'"), *AssetInfoFileName);
+		return false;
+	}
+
+	FIoStoreShaderCodeArchiveHeader IoStoreLibraryHeader;
+	FIoStoreShaderCodeArchive::CreateIoStoreShaderCodeArchiveHeader(FName(FormatName), SerializedShaders, IoStoreLibraryHeader);
+	checkf(SerializedShaders.ShaderEntries.Num() == IoStoreLibraryHeader.ShaderEntries.Num(), TEXT("IoStore header has different number of shaders (%d) than the original library (%d)"),
+		IoStoreLibraryHeader.ShaderEntries.Num(), SerializedShaders.ShaderEntries.Num());
+	checkf(SerializedShaders.ShaderMapEntries.Num() == IoStoreLibraryHeader.ShaderMapEntries.Num(), TEXT("IoStore header has different number of shadermaps (%d) than the original library (%d)"),
+		IoStoreLibraryHeader.ShaderMapEntries.Num(), SerializedShaders.ShaderMapEntries.Num());
+
+	// Load all the shaders, decompress them and recompress into groups. Each code chunk is shader group chunk now.
+	// This also updates group's compressed size, the only value left to calculate.
+	int32 GroupCount = IoStoreLibraryHeader.ShaderGroupEntries.Num();
+	OutCodeIoChunks.SetNum(GroupCount);
+	FCriticalSection DiskArchiveAccess;
+	ParallelFor(OutCodeIoChunks.Num(),
+		[&OutCodeIoChunks, &IoStoreLibraryHeader, &DiskArchiveAccess, &LibraryBuffer, &SerializedShaders, &OffsetToShaderCode, InShaderOodleCompressor, InShaderOodleLevel](int32 GroupIndex)
+		{
+			FIoStoreShaderGroupEntry& Group = IoStoreLibraryHeader.ShaderGroupEntries[GroupIndex];
+			uint8* UncompressedGroupMemory = reinterpret_cast<uint8*>(FMemory::Malloc(Group.UncompressedSize));
+			check(UncompressedGroupMemory);
+
+			int32 SmallestShaderInGroupByOrdinal = MAX_int32;
+			uint8* ShaderStart = UncompressedGroupMemory;
+			// Not worth to special-case for group size of 1 (by converting to copy) - such groups are very fast to decompress and recompress.
+			for (uint32 IdxShaderInGroup = 0; IdxShaderInGroup < Group.NumShaders; ++IdxShaderInGroup)
+			{
+				int32 ShaderIndex = IoStoreLibraryHeader.ShaderIndices[Group.ShaderIndicesOffset + IdxShaderInGroup];
+				SmallestShaderInGroupByOrdinal = FMath::Min(SmallestShaderInGroupByOrdinal, ShaderIndex);
+				const FIoStoreShaderCodeEntry& Shader = IoStoreLibraryHeader.ShaderEntries[ShaderIndex];
+
+				checkf((reinterpret_cast<uint64>(ShaderStart) - reinterpret_cast<uint64>(UncompressedGroupMemory)) == Shader.UncompressedOffsetInGroup,
+					TEXT("Shader uncompressed offset in group does not agree with its actual placement: Shader.UncompressedOffsetInGroup=%llu, actual=%llu"),
+					Shader.UncompressedOffsetInGroup, (reinterpret_cast<uint64>(ShaderStart) - reinterpret_cast<uint64>(UncompressedGroupMemory))
+				);
+
+				// load and decompress the shader at the desired offset
+				const FShaderCodeEntry& IndividuallyCompressedShader = SerializedShaders.ShaderEntries[ShaderIndex];
+				{
+					// small shaders might be stored without compression, handle them here
+					if (IndividuallyCompressedShader.Size == IndividuallyCompressedShader.UncompressedSize)
+					{
+						FMemory::Memcpy(ShaderStart, LibraryBuffer.GetData() + OffsetToShaderCode + IndividuallyCompressedShader.Offset, IndividuallyCompressedShader.Size);
+					}
+					else
+					{
+						// This function will crash if decompression fails.
+						ShaderCodeArchive::DecompressShaderWithOodle(ShaderStart, IndividuallyCompressedShader.UncompressedSize, LibraryBuffer.GetData() + OffsetToShaderCode + IndividuallyCompressedShader.Offset, IndividuallyCompressedShader.Size);
+					}
+				}
+
+				ShaderStart += IndividuallyCompressedShader.UncompressedSize;
+			}
+
+			checkf((reinterpret_cast<uint64>(ShaderStart) - reinterpret_cast<uint64>(UncompressedGroupMemory)) == Group.UncompressedSize,
+				TEXT("Uncompressed shader group size does not agree with the actual results (Group.UncompressedSize=%llu, actual=%llu)"),
+				Group.UncompressedSize, (reinterpret_cast<uint64>(ShaderStart) - reinterpret_cast<uint64>(UncompressedGroupMemory)));
+
+			// now compress the whole group
+			int64 CompressedGroupSize = 0;
+			uint8* CompressedShaderGroupMemory = nullptr;
+			ShaderCodeArchive::CompressShaderWithOodle(nullptr, CompressedGroupSize, UncompressedGroupMemory, Group.UncompressedSize, InShaderOodleCompressor, InShaderOodleLevel);
+			checkf(CompressedGroupSize > 0, TEXT("CompressedGroupSize estimate seems wrong (%lld)"), CompressedGroupSize);
+
+			CompressedShaderGroupMemory = reinterpret_cast<uint8*>(FMemory::Malloc(CompressedGroupSize));
+			bool bCompressed = ShaderCodeArchive::CompressShaderWithOodle(CompressedShaderGroupMemory, CompressedGroupSize, UncompressedGroupMemory, Group.UncompressedSize, InShaderOodleCompressor, InShaderOodleLevel);
+			checkf(bCompressed, TEXT("We could not compress the shader group after providing an estimated memory."));
+
+			TTuple<FIoChunkId, FIoBuffer, uint32>& OutCodeIoChunk = OutCodeIoChunks[GroupIndex];
+			OutCodeIoChunk.Get<0>() = IoStoreLibraryHeader.ShaderGroupIoHashes[GroupIndex];
+			// This value is the load order factor for the group (the smaller, the more likely), that IoStore will use to sort the chunks.
+			// We calculate it as the smallest shader index in the group, which shouldn't be bad. Arguably a better way would be to get the lowest-numbered shadermap that references _any_ shader in the group, 
+			// but it's much, much slower to calculate and the resulting order would be likely the same/similar most of the time
+			checkf(SmallestShaderInGroupByOrdinal >= 0 && SmallestShaderInGroupByOrdinal < IoStoreLibraryHeader.ShaderEntries.Num(), TEXT("SmallestShaderInGroupByOrdinal has an invalid value of %d (not within [0, %d) range as expected)"),
+				SmallestShaderInGroupByOrdinal, IoStoreLibraryHeader.ShaderEntries.Num());
+			OutCodeIoChunk.Get<2>() = static_cast<uint32>(SmallestShaderInGroupByOrdinal);
+
+			if (CompressedGroupSize < Group.UncompressedSize)
+			{
+				OutCodeIoChunk.Get<1>() = FIoBuffer(CompressedGroupSize);
+				FMemory::Memcpy(OutCodeIoChunk.Get<1>().GetData(), CompressedShaderGroupMemory, CompressedGroupSize);
+				Group.CompressedSize = CompressedGroupSize;
+			}
+			else
+			{
+				// store uncompressed (unlikely, but happens for a 200-byte sized shader that happens to get its own group)
+				OutCodeIoChunk.Get<1>() = FIoBuffer(Group.UncompressedSize);
+				FMemory::Memcpy(OutCodeIoChunk.Get<1>().GetData(), UncompressedGroupMemory, Group.UncompressedSize);
+				Group.CompressedSize = Group.UncompressedSize;
+			}
+			FMemory::Free(UncompressedGroupMemory);
+			FMemory::Free(CompressedShaderGroupMemory);
+		},
+		EParallelForFlags::Unbalanced
+	);
+
+	// calculate and log stats
+	int64 TotalUncompressedSizeViaGroups = 0;	// calculate total uncompressed size twice for a sanity check
+	int64 TotalUncompressedSizeViaShaders = 0;
+	int64 TotalIndividuallyCompressedSize = 0;
+	int64 TotalGroupCompressedSize = 0;
+	for (int32 IdxGroup = 0, NumGroups = IoStoreLibraryHeader.ShaderGroupEntries.Num(); IdxGroup < NumGroups; ++IdxGroup)
+	{
+		FIoStoreShaderGroupEntry& Group = IoStoreLibraryHeader.ShaderGroupEntries[IdxGroup];
+		TotalGroupCompressedSize += Group.CompressedSize;
+		TotalUncompressedSizeViaGroups += Group.UncompressedSize;
+
+		// now go via shaders
+		for (uint32 IdxShaderInGroup = 0; IdxShaderInGroup < Group.NumShaders; ++IdxShaderInGroup)
+		{
+			int32 ShaderIndex = IoStoreLibraryHeader.ShaderIndices[Group.ShaderIndicesOffset + IdxShaderInGroup];
+			const FShaderCodeEntry& IndividuallyCompressedShader = SerializedShaders.ShaderEntries[ShaderIndex];
+
+			TotalIndividuallyCompressedSize += IndividuallyCompressedShader.Size;
+			TotalUncompressedSizeViaShaders += IndividuallyCompressedShader.UncompressedSize;
+		}
+	}
+
+	checkf(TotalUncompressedSizeViaGroups == TotalUncompressedSizeViaShaders, TEXT("Sanity check failure: total uncompressed shader size differs if calculated via shader groups (%lld) or individual shaders (%lld)"),
+		TotalUncompressedSizeViaGroups, TotalUncompressedSizeViaShaders
+	);
+
+	UE_LOG(LogIoStore, Display, TEXT("%s(%s): Recompressed %d shaders as %d groups. Library size changed from %lld KB (%.2f : 1 ratio) to %lld KB (%.2f : 1 ratio), %.2f%% of previous."),
+		*LibraryName, *FormatName.ToString(),
+		SerializedShaders.ShaderEntries.Num(), IoStoreLibraryHeader.ShaderGroupEntries.Num(),
+		TotalIndividuallyCompressedSize / 1024, static_cast<double>(TotalUncompressedSizeViaShaders) / static_cast<double>(TotalIndividuallyCompressedSize),
+		TotalGroupCompressedSize / 1024, static_cast<double>(TotalUncompressedSizeViaShaders) / static_cast<double>(TotalGroupCompressedSize),
+		100.0 * static_cast<double>(TotalGroupCompressedSize) / static_cast<double>(TotalIndividuallyCompressedSize)
+	);
+
+	FLargeMemoryWriter IoStoreLibraryAr(0, true);
+	FIoStoreShaderCodeArchive::SaveIoStoreShaderCodeArchive(IoStoreLibraryHeader, IoStoreLibraryAr);
+	OutLibraryIoChunk.Key = FIoStoreShaderCodeArchive::GetShaderCodeArchiveChunkId(LibraryName, FormatName);
+	int64 TotalSize = IoStoreLibraryAr.TotalSize();
+	OutLibraryIoChunk.Value = FIoBuffer(FIoBuffer::AssumeOwnership, IoStoreLibraryAr.ReleaseOwnership(), TotalSize);
+
+	int32 ShaderMapCount = IoStoreLibraryHeader.ShaderMapHashes.Num();
+	OutShaderMaps.SetNum(ShaderMapCount);
+	ParallelFor(OutShaderMaps.Num(),
+		[&OutShaderMaps, &IoStoreLibraryHeader](int32 ShaderMapIndex)
+		{
+			TTuple<FSHAHash, TArray<FIoChunkId>>& OutShaderMap = OutShaderMaps[ShaderMapIndex];
+			OutShaderMap.Key = IoStoreLibraryHeader.ShaderMapHashes[ShaderMapIndex];
+			const FIoStoreShaderMapEntry& ShaderMapEntry = IoStoreLibraryHeader.ShaderMapEntries[ShaderMapIndex];
+			int32 LookupIndexEnd = ShaderMapEntry.ShaderIndicesOffset + ShaderMapEntry.NumShaders;
+			OutShaderMap.Value.Reserve(ShaderMapEntry.NumShaders);	// worst-case, 1 group == 1 shader
+			for (int32 ShaderLookupIndex = ShaderMapEntry.ShaderIndicesOffset; ShaderLookupIndex < LookupIndexEnd; ++ShaderLookupIndex)
+			{
+				int32 ShaderIndex = IoStoreLibraryHeader.ShaderIndices[ShaderLookupIndex];
+				int32 GroupIndex = IoStoreLibraryHeader.ShaderEntries[ShaderIndex].ShaderGroupIndex;
+				OutShaderMap.Value.AddUnique(IoStoreLibraryHeader.ShaderGroupIoHashes[GroupIndex]);
+			}
+		},
+		EParallelForFlags::Unbalanced
+	);
+
+	OutShaderMapAssetAssociations.Reserve(ShaderCodeToAssets.Num());
+	for (auto& KV : ShaderCodeToAssets)
+	{
+		OutShaderMapAssetAssociations.Emplace(KV.Key, MoveTemp(KV.Value));
+	}
+
+	return true;
+}
+
+void FIoStoreShaderLibraryProcessor::ProcessShaderLibraries(const FIoStoreArguments& Arguments, TArray<FContainerTargetSpec*>& ContainerTargets, FIoStoreShaderLibraryOutput& Output)
+{
+	IOSTORE_CPU_SCOPE(ProcessShaderLibraries);
+
+	double LibraryStart = FPlatformTime::Seconds();
+
+	TMap<FIoChunkId, FShaderInfo*> ChunkIdToShaderInfoMap;
+	TMap<FSHAHash, TArray<FIoChunkId>> ShaderChunkIdsByShaderMapHash;
+	TMap<FName, TSet<FSHAHash>> PackageNameToShaderMaps;
+	TMap<FContainerTargetSpec*, TSet<FShaderInfo*>> AllContainerShaderLibraryShadersMap;
+
+	{
+		IOSTORE_CPU_SCOPE(ConvertShaderLibraries);
+		for (FContainerTargetSpec* ContainerTarget : ContainerTargets)
+		{
+			for (FContainerTargetFile& TargetFile : ContainerTarget->TargetFiles)
+			{
+				if (TargetFile.ChunkType == EContainerChunkType::ShaderCodeLibrary)
+				{
+					TSet<FShaderInfo*>& ContainerShaderLibraryShaders = AllContainerShaderLibraryShadersMap.FindOrAdd(ContainerTarget);
+
+					TArray<TTuple<FSHAHash, TArray<FIoChunkId>>> ShaderMaps;
+					TTuple<FIoChunkId, FIoBuffer> LibraryChunk;
+					TArray<TTuple<FIoChunkId, FIoBuffer, uint32>> CodeChunks;
+					TArray<TTuple<FSHAHash, TSet<FName>>> ShaderMapAssetAssociations;
+					if (!ConvertToIoStoreShaderLibrary(*TargetFile.NormalizedSourcePath, TargetFile.ChunkId, Arguments.PackageStore.Get(), Arguments.ShaderOodleCompressor, Arguments.ShaderOodleLevel, LibraryChunk, CodeChunks, ShaderMaps, ShaderMapAssetAssociations))
+					{
+						UE_LOG(LogIoStore, Warning, TEXT("Failed converting shader library '%s'"), *TargetFile.NormalizedSourcePath);
+						continue;
+					}
+					TargetFile.ChunkId = LibraryChunk.Key;
+					TargetFile.SourceBuffer.Emplace(LibraryChunk.Value);
+					TargetFile.SourceSize = LibraryChunk.Value.GetSize();
+
+					const bool bIsGlobalShaderLibrary = FPaths::GetCleanFilename(TargetFile.NormalizedSourcePath).StartsWith(TEXT("ShaderArchive-Global-"));
+					const FShaderInfo::EShaderType ShaderType = bIsGlobalShaderLibrary ? FShaderInfo::Global : FShaderInfo::Normal;
+					for (const TTuple<FIoChunkId, FIoBuffer, uint32>& CodeChunk : CodeChunks)
+					{
+						const FIoChunkId& ShaderChunkId = CodeChunk.Get<0>();
+						FShaderInfo* ShaderInfo = ChunkIdToShaderInfoMap.FindRef(ShaderChunkId);
+						if (!ShaderInfo)
+						{
+							ShaderInfo = new FShaderInfo();
+							ShaderInfo->ChunkId = ShaderChunkId;
+							ShaderInfo->CodeIoBuffer = CodeChunk.Get<1>();
+							ShaderInfo->LoadOrderFactor = CodeChunk.Get<2>();
+							Output.ShaderInfos.Add(ShaderInfo);
+							ChunkIdToShaderInfoMap.Add(ShaderChunkId, ShaderInfo);
+						}
+						else
+						{
+							// first, make sure that the code is exactly the same
+							if (ShaderInfo->CodeIoBuffer != CodeChunk.Get<1>())
+							{
+								UE_LOG(LogIoStore, Error, TEXT("Collision of two shader code chunks, same Id (%s), different code. Packaged game will likely crash, not being able to decompress the shaders."), *LexToString(ShaderChunkId));
+							}
+
+							// If we already exist, then we have two separate LoadOrderFactors,
+							// which one we got first affects build determinism. Take the lower.
+							if (CodeChunk.Get<2>() < ShaderInfo->LoadOrderFactor)
+							{
+								ShaderInfo->LoadOrderFactor = CodeChunk.Get<2>();
+							}
+						}
+
+
+						const FShaderInfo::EShaderType* CurrentShaderTypeInContainer = ShaderInfo->TypeInContainer.Find(ContainerTarget);
+						if (!CurrentShaderTypeInContainer || *CurrentShaderTypeInContainer != FShaderInfo::Global)
+						{
+							// If a shader is both global and shared consider it to be global
+							ShaderInfo->TypeInContainer.Add(ContainerTarget, ShaderType);
+						}
+
+						ContainerShaderLibraryShaders.Add(ShaderInfo);
+					}
+					for (const TTuple<FSHAHash, TSet<FName>>& ShaderMapAssetAssociation : ShaderMapAssetAssociations)
+					{
+						for (const FName& PackageName : ShaderMapAssetAssociation.Value)
+						{
+							TSet<FSHAHash>& PackageShaderMaps = PackageNameToShaderMaps.FindOrAdd(PackageName);
+							PackageShaderMaps.Add(ShaderMapAssetAssociation.Key);
+						}
+					}
+
+					for (TTuple<FSHAHash, TArray<FIoChunkId>>& ShaderMap : ShaderMaps)
+					{
+						TArray<FIoChunkId>& ShaderMapChunkIds = ShaderChunkIdsByShaderMapHash.FindOrAdd(ShaderMap.Key);
+						ShaderMapChunkIds.Append(MoveTemp(ShaderMap.Value));
+					}
+				} // end if containerchunktype shadercodelibrary
+			} // end foreach targetfile
+		} // end foreach container
+	}
+
+	// 1. Update ShaderInfos with which packages we reference.
+	// 2. Add to packages which shaders we use.
+	// 3. Add to PackageStore what shaders we use.
+	{
+		IOSTORE_CPU_SCOPE(UpdatePackageStoreShaders);
+		for (FContainerTargetSpec* ContainerTarget : ContainerTargets)
+		{
+			for (FCookedPackage* Package : ContainerTarget->Packages)
+			{
+				TSet<FSHAHash>* FindShaderMapHashes = PackageNameToShaderMaps.Find(Package->PackageName);
+				if (FindShaderMapHashes)
+				{
+					for (const FSHAHash& ShaderMapHash : *FindShaderMapHashes)
+					{
+						const TArray<FIoChunkId>* FindChunkIds = ShaderChunkIdsByShaderMapHash.Find(ShaderMapHash);
+						if (!FindChunkIds)
+						{
+							UE_LOG(LogIoStore, Warning, TEXT("Package '%s' in '%s' referencing missing shader map '%s'"), *Package->PackageName.ToString(), *ContainerTarget->Name.ToString(), *ShaderMapHash.ToString());
+							continue;
+						}
+						Package->ShaderMapHashes.Add(ShaderMapHash);
+						for (const FIoChunkId& ShaderChunkId : *FindChunkIds)
+						{
+							FShaderInfo* ShaderInfo = ChunkIdToShaderInfoMap.FindRef(ShaderChunkId);
+							if (!ShaderInfo)
+							{
+								UE_LOG(LogIoStore, Warning, TEXT("Package '%s' in '%s' referencing missing shader with chunk id '%s'"), *Package->PackageName.ToString(), *ContainerTarget->Name.ToString(), *BytesToHex(ShaderChunkId.GetData(), ShaderChunkId.GetSize()));
+								continue;
+							}
+
+							check(ShaderInfo);
+							ShaderInfo->ReferencedByPackages.Add(Package);
+							Package->Shaders.AddUnique(ShaderInfo);
+						}
+					}
+				}
+				Algo::Sort(Package->ShaderMapHashes);
+			}
+		}
+	}
+
+	{
+		IOSTORE_CPU_SCOPE(AddShaderChunks);
+		for (FContainerTargetSpec* ContainerTarget : ContainerTargets)
+		{
+			auto AddShaderTargetFile = [&ContainerTarget](FShaderInfo* ShaderInfo)
+				{
+					FContainerTargetFile& ShaderTargetFile = ContainerTarget->TargetFiles.AddDefaulted_GetRef();
+					ShaderTargetFile.ContainerTarget = ContainerTarget;
+					ShaderTargetFile.ChunkId = ShaderInfo->ChunkId;
+					ShaderTargetFile.ChunkType = EContainerChunkType::ShaderCode;
+					ShaderTargetFile.bForceUncompressed = true;
+					ShaderTargetFile.SourceBuffer.Emplace(ShaderInfo->CodeIoBuffer);
+					ShaderTargetFile.SourceSize = ShaderInfo->CodeIoBuffer.DataSize();
+				};
+
+			const TSet<FShaderInfo*>* FindContainerShaderLibraryShaders = AllContainerShaderLibraryShadersMap.Find(ContainerTarget);
+			if (FindContainerShaderLibraryShaders)
+			{
+				for (FCookedPackage* Package : ContainerTarget->Packages)
+				{
+					for (FShaderInfo* ShaderInfo : Package->Shaders)
+					{
+						if (ShaderInfo->ReferencedByPackages.Num() == 1)
+						{
+							FShaderInfo::EShaderType* ShaderType = ShaderInfo->TypeInContainer.Find(ContainerTarget);
+							if (ShaderType && *ShaderType != FShaderInfo::Global)
+							{
+								*ShaderType = FShaderInfo::Inline;
+							}
+						}
+					}
+				}
+
+				for (FShaderInfo* ShaderInfo : *FindContainerShaderLibraryShaders)
+				{
+					FShaderAssociationInfo::FShaderChunkInfoKey ShaderChunkInfoKey = { ContainerTarget->Name, ShaderInfo->ChunkId };
+					FShaderAssociationInfo::FShaderChunkInfo& ShaderChunkInfo = Output.AssocInfo.ShaderChunkInfos.Add(ShaderChunkInfoKey);
+
+					FShaderInfo::EShaderType* ShaderType = ShaderInfo->TypeInContainer.Find(ContainerTarget);
+					check(ShaderType);
+					if (*ShaderType == FShaderInfo::Global)
+					{
+						ContainerTarget->GlobalShaders.Add(ShaderInfo);
+						ShaderChunkInfo.Type = FShaderAssociationInfo::FShaderChunkInfo::Global;
+					}
+					else if (*ShaderType == FShaderInfo::Inline)
+					{
+						ContainerTarget->InlineShaders.Add(ShaderInfo);
+
+						ShaderChunkInfo.Type = FShaderAssociationInfo::FShaderChunkInfo::Package;
+						checkf(ShaderInfo->ReferencedByPackages.Num() == 1, TEXT("Inline shader chunks must be referenced by 1 package only, but shader chunk %s is referenced by %d"),
+							*LexToString(ShaderInfo->ChunkId), ShaderInfo->ReferencedByPackages.Num());
+
+						ShaderChunkInfo.ReferencedByPackages.Add(ShaderInfo->ReferencedByPackages[FSetElementId::FromInteger(0)]->PackageName);
+					}
+					else if (ShaderInfo->ReferencedByPackages.Num() > 1)
+					{
+						ContainerTarget->SharedShaders.Add(ShaderInfo);
+
+						ShaderChunkInfo.Type = FShaderAssociationInfo::FShaderChunkInfo::Package;
+
+						for (FCookedPackage* Package : ShaderInfo->ReferencedByPackages)
+						{
+							ShaderChunkInfo.ReferencedByPackages.Add(Package->PackageName);
+
+							TArray<FShaderAssociationInfo::FShaderChunkInfoKey>& PackageShaders = Output.AssocInfo.PackageShaderMap.FindOrAdd(Package->PackageName);
+							PackageShaders.Add(ShaderChunkInfoKey);
+						}
+					}
+					else
+					{
+						//
+						// Note that we can get here with shaders that get split off in to another container (e.g. sm6 shaders). 
+						// Since they are in a different pakChunk they can't get inlined or shared. However, they still "belong" to
+						// the referencing packages for association purposes.
+						//
+						if (ShaderInfo->ReferencedByPackages.Num())
+						{
+							ShaderChunkInfo.Type = FShaderAssociationInfo::FShaderChunkInfo::Package;
+
+							for (FCookedPackage* Package : ShaderInfo->ReferencedByPackages)
+							{
+								ShaderChunkInfo.ReferencedByPackages.Add(Package->PackageName);
+
+								TArray<FShaderAssociationInfo::FShaderChunkInfoKey>& PackageShaders = Output.AssocInfo.PackageShaderMap.FindOrAdd(Package->PackageName);
+								PackageShaders.Add(ShaderChunkInfoKey);
+							}
+						}
+						else
+						{
+							ShaderChunkInfo.Type = FShaderAssociationInfo::FShaderChunkInfo::Orphan;
+						}
+
+						ContainerTarget->UniqueShaders.Add(ShaderInfo);
+					}
+					AddShaderTargetFile(ShaderInfo);
+				}
+			}
+		}
+	}
+
+	double LibraryEnd = FPlatformTime::Seconds();
+	UE_LOG(LogIoStore, Display, TEXT("Shaders processed in %.2f seconds"), LibraryEnd - LibraryStart);
+}
+
+} // namespace UE::IoStore::Private
+
diff --git a/Engine/Source/Developer/IoStoreUtilities/Private/IoStoreShaderLibraryProcessor.h b/Engine/Source/Developer/IoStoreUtilities/Private/IoStoreShaderLibraryProcessor.h
new file mode 100644
--- /dev/null
+++ b/Engine/Source/Developer/IoStoreUtilities/Private/IoStoreShaderLibraryProcessor.h
@@ -0,0 +1,81 @@
+// Copyright Epic Games, Inc. All Rights Reserved.
+
+#pragma once
+
+#include "CookMetadata.h"
+#include "Containers/Array.h"
+#include "Containers/Map.h"
+#include "IO/IoChunkId.h"
+#include "Templates/UniquePtr.h"
+#include "UObject/NameTypes.h"
+
+namespace UE::IoStore::Private
+{
+
+struct FContainerTargetSpec;
+struct FIoStoreArguments;
+struct FShaderInfo;
+
+// Carries association between shaders and packages from shader processing to size assignment.
+struct FShaderAssociationInfo
+{
+	struct FShaderChunkInfo
+	{
+		enum EType
+		{
+			// This shader is referenced by one or more packages and can assign its size to
+			// those packages
+			Package,
+			// This shader is needed by the baseline engine in some manner and is a flat
+			// size cost.
+			Global,
+			// This shader isn't global or referenced by packages - theoretically shouldn't
+			// exist.
+			Orphan
+		};
+
+		EType Type;
+		TArray<FName> ReferencedByPackages;
+		uint32 CompressedSize = 0;
+		UE::Cook::EPluginSizeTypes SizeType;
+
+		// If we are shared across plugins, this is our pseudo plugin name. This is generated during size assignment.
+		FString SharedPluginName;
+	};
+
+	struct FShaderChunkInfoKey
+	{
+		FName PakChunkName;
+		FIoChunkId IoChunkId;
+
+		friend uint32 GetTypeHash(const FShaderChunkInfoKey& Key)
+		{
+			return GetTypeHash(Key.IoChunkId);
+		}
+		friend bool operator == (const FShaderChunkInfoKey& LHS, const FShaderChunkInfoKey& RHS)
+		{
+			return LHS.PakChunkName == RHS.PakChunkName && LHS.IoChunkId == RHS.IoChunkId;
+		}
+	};
+
+	TMap<FShaderChunkInfoKey, FShaderChunkInfo> ShaderChunkInfos;
+
+	// The list of shaders referenced by each package. This can be used to look up in to ContainerShaderInfos.
+	TMap<FName /* PackageName */, TArray<FShaderChunkInfoKey>> PackageShaderMap;
+};
+
+struct FIoStoreShaderLibraryOutput
+{
+	TArray<FShaderInfo*> ShaderInfos;
+	FShaderAssociationInfo AssocInfo;
+};
+
+// Processor to load cooked shader libraries and package them into IoStore entries
+class FIoStoreShaderLibraryProcessor
+{
+public:
+	void ProcessShaderLibraries(const FIoStoreArguments& Arguments, TArray<FContainerTargetSpec*>& ContainerTargets, FIoStoreShaderLibraryOutput& Output);
+
+};
+
+} // namespace UE::IoStore::Private
diff --git a/Engine/Source/Developer/IoStoreUtilities/Private/IoStoreUtilities.cpp b/Engine/Source/Developer/IoStoreUtilities/Private/IoStoreUtilities.cpp
--- a/Engine/Source/Developer/IoStoreUtilities/Private/IoStoreUtilities.cpp
+++ b/Engine/Source/Developer/IoStoreUtilities/Private/IoStoreUtilities.cpp
@@ -1,7 +1,9 @@
 // Copyright Epic Games, Inc. All Rights Reserved.
 
 #include "IoStoreUtilities.h"
+#include "IoStoreUtilitiesCore.h"
 #include "IoStoreWriter.h"
+#include "IoStoreShaderLibraryProcessor.h"
 
 #include "IoStoreLooseFiles.h"
 #include "Algo/BinarySearch.h"
@@ -94,6 +96,8 @@
 
 IMPLEMENT_MODULE(FDefaultModuleImpl, IoStoreUtilities);
 
+using namespace UE::IoStore::Private;
+
 #define IOSTORE_CPU_SCOPE(NAME) TRACE_CPUPROFILER_EVENT_SCOPE(IoStore##NAME);
 #define IOSTORE_CPU_SCOPE_DATA(NAME, DATA) TRACE_CPUPROFILER_EVENT_SCOPE(IoStore##NAME);
 
@@ -610,12 +614,6 @@ public:
 	}
 };
 
-struct FReleasedPackages
-{
-	TSet<FName> PackageNames;
-	TMap<FPackageId, FName> PackageIdToName;
-};
-
 [[nodiscard]] static bool LoadKeyChain(const TCHAR* CmdLine, FKeyChain& OutCryptoSettings)
 {
 	OutCryptoSettings.SetSigningKey(InvalidRSAKeyHandle);
@@ -814,371 +812,6 @@ struct FReleasedPackages
 	return true;
 }
 
-struct FContainerSourceFile 
-{
-	FString NormalizedPath;
-	FString DestinationPath;
-	bool bNeedsCompression = false;
-	bool bNeedsEncryption = false;
-};
-
-struct FContainerSourceSpec
-{
-	FName Name;
-	FString OutputPath;
-	FString OptionalOutputPath;
-	FString StageLooseFileRootPath;
-	TArray<FContainerSourceFile> SourceFiles;
-	FString PatchTargetFile;
-	TArray<FString> PatchSourceContainerFiles;
-	FString EncryptionKeyOverrideGuid;
-	bool bGenerateDiffPatch = false;
-	bool bOnDemand = false;
-};
-
-struct FCookedFileStatData
-{
-	enum EFileType
-	{
-		PackageHeader,
-		PackageData,
-		BulkData,
-		OptionalBulkData,
-		MemoryMappedBulkData,
-		ShaderLibrary,
-		Regions,
-		OptionalSegmentPackageHeader,
-		OptionalSegmentPackageData,
-		OptionalSegmentBulkData,
-		Invalid
-	};
-
-	int64 FileSize = 0;
-	EFileType FileType = Invalid;
-};
-
-/**
- * Finds where the extension starts in a path.
- * 
- * We can't just find the left most '.' in the file name because the path might
- * contain a bulkdata cooked index (@see FBulkDataCookedIndex) and we don't want
- * to include the index part. We are attempting to match up the extensions with
- * those found in the constructor of FCookedFileStatMap.
- * 
- * For example take the following filenames and the extensions we want to find:
- * <packagename>.m.ubulk   -> .m.ubulk
- * <packagename>.o.ubulk   -> .o.ubulk
- * <packagename>.001.ubulk ->.ubulk
- * 
- * So we want to find the left most '.' unless that is followed by only numeric
- * values in which case we have found a cooked index and we want to find the '.'
- * after those values.
-*/
-static int32 GetFullExtensionStartIndex(FStringView Path)
-{
-	int32 ExtensionStartIndex = -1;
-	for (int32 Index = Path.Len() - 1; Index >= 0; --Index)
-	{
-		// Check if we reached the end of the filename
-		if (FPathViews::IsSeparator(Path[Index]))
-		{
-			break;
-		}
-		else if (Path[Index] == '.')
-		{
-			if (ExtensionStartIndex != -1)
-			{
-				// As we have already found a '.' we need to check for the cooked index
-				FStringView Extension = Path.SubStr(Index + 1, (ExtensionStartIndex - Index) - 1);
-				if (UE::String::IsNumericOnlyDigits(Extension))
-				{
-					return ExtensionStartIndex;
-				}
-			}
-
-			ExtensionStartIndex = Index;
-		}
-	}
-	return ExtensionStartIndex;
-}
-
-static FStringView GetFullExtension(FStringView Path)
-{
-	int32 ExtensionStartIndex = GetFullExtensionStartIndex(Path);
-	if (ExtensionStartIndex < 0)
-	{
-		return FStringView();
-	}
-	else
-	{
-		return Path.RightChop(ExtensionStartIndex);
-	}
-}
-
-class FCookedFileStatMap
-{
-public:
-	FCookedFileStatMap()
-	{
-		Extensions.Emplace(TEXT(".umap"), FCookedFileStatData::PackageHeader);
-		Extensions.Emplace(TEXT(".uasset"), FCookedFileStatData::PackageHeader);
-		Extensions.Emplace(TEXT(".uexp"), FCookedFileStatData::PackageData);
-		Extensions.Emplace(TEXT(".ubulk"), FCookedFileStatData::BulkData);
-		Extensions.Emplace(TEXT(".uptnl"), FCookedFileStatData::OptionalBulkData);
-		Extensions.Emplace(TEXT(".m.ubulk"), FCookedFileStatData::MemoryMappedBulkData);
-		Extensions.Emplace(*FString::Printf(TEXT(".uexp%s"), FFileRegion::RegionsFileExtension), FCookedFileStatData::Regions);
-		Extensions.Emplace(*FString::Printf(TEXT(".uptnl%s"), FFileRegion::RegionsFileExtension), FCookedFileStatData::Regions);
-		Extensions.Emplace(*FString::Printf(TEXT(".ubulk%s"), FFileRegion::RegionsFileExtension), FCookedFileStatData::Regions);
-		Extensions.Emplace(*FString::Printf(TEXT(".m.ubulk%s"), FFileRegion::RegionsFileExtension), FCookedFileStatData::Regions);
-		Extensions.Emplace(TEXT(".ushaderbytecode"), FCookedFileStatData::ShaderLibrary);
-		Extensions.Emplace(TEXT(".o.umap"), FCookedFileStatData::OptionalSegmentPackageHeader);
-		Extensions.Emplace(TEXT(".o.uasset"), FCookedFileStatData::OptionalSegmentPackageHeader);
-		Extensions.Emplace(TEXT(".o.uexp"), FCookedFileStatData::OptionalSegmentPackageData);
-		Extensions.Emplace(TEXT(".o.ubulk"), FCookedFileStatData::OptionalSegmentBulkData);
-	}
-
-	int32 Num() const
-	{
-		return Map.Num();
-	}
-
-	void Add(const TCHAR* Path, int64 FileSize)
-	{
-		FString NormalizedPath(Path);
-		FPaths::NormalizeFilename(NormalizedPath);
-
-		FStringView NormalizePathView(NormalizedPath);
-		int32 ExtensionStartIndex = GetFullExtensionStartIndex(NormalizePathView);
-		if (ExtensionStartIndex < 0)
-		{
-			return;
-		}
-		FStringView Extension = NormalizePathView.RightChop(ExtensionStartIndex);
-		const FCookedFileStatData::EFileType* FindFileType = FindFileTypeFromExtension(Extension);
-		if (!FindFileType)
-		{
-			return;
-		}
-
-		FCookedFileStatData& CookedFileStatData = Map.Add(MoveTemp(NormalizedPath));
-		CookedFileStatData.FileType = *FindFileType;
-		CookedFileStatData.FileSize = FileSize;
-	}
-
-	const FCookedFileStatData* Find(FStringView NormalizedPath) const
-	{
-		return Map.FindByHash(GetTypeHash(NormalizedPath), NormalizedPath);
-	}
-
-private:
-	const FCookedFileStatData::EFileType* FindFileTypeFromExtension(const FStringView& Extension)
-	{
-		for (const auto& Pair : Extensions)
-		{
-			if (Pair.Key == Extension)
-			{
-				return &Pair.Value;
-			}
-		}
-		return nullptr;
-	}
-
-	TArray<TTuple<FString, FCookedFileStatData::EFileType>> Extensions;
-	TMap<FString, FCookedFileStatData> Map;
-};
-
-struct FContainerTargetSpec;
-
-struct FShaderInfo
-{
-	enum EShaderType { Normal, Global, Inline };
-
-	FIoChunkId ChunkId;
-	FIoBuffer CodeIoBuffer;
-	// the smaller the order, the more likely this will be loaded
-	uint32 LoadOrderFactor = MAX_uint32;
-	TSet<struct FCookedPackage*> ReferencedByPackages;
-	TMap<FContainerTargetSpec*, EShaderType> TypeInContainer;
-
-	// This is used to ensure build determinism and such must be stable
-	// across builds.
-	static bool Sort(const FShaderInfo* A, const FShaderInfo* B)
-	{
-		if (A->LoadOrderFactor == B->LoadOrderFactor)
-		{
-			// Shader chunk IDs are the hash of the shader so this is consistent across builds.
-			return FMemory::Memcmp(A->ChunkId.GetData(), B->ChunkId.GetData(), A->ChunkId.GetSize()) < 0;
-		}
-		return A->LoadOrderFactor < B->LoadOrderFactor;
-	}
-};
-
-struct FCookedPackage
-{
-	FPackageId GlobalPackageId;
-	FName PackageName;
-	FName SourcePackageName; // Source package name for redirected package
-	TArray<FShaderInfo*> Shaders;
-	TArray<FSHAHash> ShaderMapHashes;
-	uint64 UAssetSize = 0;
-	uint64 OptionalSegmentUAssetSize = 0;
-	uint64 UExpSize = 0;
-	uint64 TotalBulkDataSize = 0;
-	uint64 LoadOrder = MAX_uint64; // Ordered by dependencies
-	uint64 DiskLayoutOrder = MAX_uint64; // Final order after considering input order files
-	uint64 OverrideDiskLayoutOrder = MAX_uint64; // Final order after applying custom per-container ordering
-	FPackageStoreEntryResource PackageStoreEntry;
-	int32 PreOrderNumber = -1; // Used for sorting in load order
-	bool bPermanentMark = false; // Used for sorting in load order
-	bool bIsLocalized = false;
-};
-
-struct FLegacyCookedPackage
-	: public FCookedPackage
-{
-	FString FileName;
-	FString OptionalSegmentFileName;
-	FPackageStorePackage* OptimizedPackage = nullptr;
-	FPackageStorePackage* OptimizedOptionalSegmentPackage = nullptr;
-};
-
-enum class EContainerChunkType
-{
-	ShaderCodeLibrary,
-	ShaderCode,
-	PackageData,
-	BulkData,
-	OptionalBulkData,
-	MemoryMappedBulkData,
-	OptionalSegmentPackageData,
-	OptionalSegmentBulkData,
-	Invalid
-};
-
-struct FContainerTargetFile
-{
-	FContainerTargetSpec* ContainerTarget = nullptr;
-	FCookedPackage* Package = nullptr;
-	FString NormalizedSourcePath;
-	TOptional<FIoBuffer> SourceBuffer;
-	FString DestinationPath;
-	uint64 SourceSize = 0;
-	uint64 IdealOrder = 0;
-	FIoChunkId ChunkId;
-	FIoHash ChunkHash;
-	FBulkDataCookedIndex BulkDataCookedIndex;
-	TArray<uint8> PackageHeaderData;
-	EContainerChunkType ChunkType;
-	bool bForceUncompressed = false;
-
-	TArray<FFileRegion> FileRegions;
-};
-
-struct FFileOrderMap
-{
-	TMap<FName, int64> PackageNameToOrder;
-	FString Name;
-	int32 Priority;
-	int32 Index;
-
-	FFileOrderMap(int32 InPriority, int32 InIndex)
-		: Priority(InPriority)
-		, Index(InIndex)
-	{
-	}
-};
-
-enum class EFallbackOrderMode
-{
-	LoadOrder, // Default fallback ordering. Uses LoadOrder dependency ordering
-	AlphabeticalClustered, // As above but clusters are sorted alphabetically at the end for improved patching stability
-	Alphabetical // Pure alphabetical fallback ordering. Best for patching stability but not optimized to reduce seeks
-};
-
-struct FIoStoreOrderingOptions
-{
-	bool bClusterByOrderFilePriority = false;
-	bool bAlphaSortClusterPackageLists = false;
-
-	EFallbackOrderMode FallbackOrderMode = EFallbackOrderMode::LoadOrder;
-
-	bool operator==(const FIoStoreOrderingOptions& Other) const
-	{
-		return (bClusterByOrderFilePriority == Other.bClusterByOrderFilePriority) &&
-			(bAlphaSortClusterPackageLists == Other.bAlphaSortClusterPackageLists) &&
-			(FallbackOrderMode == Other.FallbackOrderMode);
-	}
-};
-
-uint32 GetTypeHash(const FIoStoreOrderingOptions& OrderingOptions)
-{
-	uint32 Hash = GetTypeHash(OrderingOptions.bClusterByOrderFilePriority);
-	Hash = HashCombine(Hash, GetTypeHash(OrderingOptions.bAlphaSortClusterPackageLists));
-	Hash = HashCombine(Hash, GetTypeHash(OrderingOptions.FallbackOrderMode));
-	return Hash;
-}
-
-struct FIoStoreArguments
-{
-	FString GlobalContainerPath;
-	FString CookedDir;
-	TArray<FContainerSourceSpec> Containers;
-	FCookedFileStatMap CookedFileStatMap;
-	TArray<FFileOrderMap> OrderMaps;
-	FKeyChain KeyChain;
-	FKeyChain PatchKeyChain;
-	FString DLCPluginPath;
-	FString DLCName;
-	FString ReferenceChunkGlobalContainerFileName;
-	FString ReferenceChunkChangesCSVFileName;
-	FString ReferenceChunkAdditionalContainersPath;
-	FString CsvPath;
-	FKeyChain ReferenceChunkKeys;
-	FReleasedPackages ReleasedPackages;
-	TUniquePtr<FCookedPackageStore> PackageStore;
-	TUniquePtr<FIoBuffer> ScriptObjects;
-	FIoStoreOrderingOptions OrderingOptions;
-	TMap<FString, FIoStoreOrderingOptions> OrderingOptionsOverrides;
-	bool bSign = false;
-	bool bRemapPluginContentToGame = false;
-	bool bCreateDirectoryIndex = true;
-	bool bFileRegions = false;
-	EAssetRegistryWritebackMethod WriteBackMetadataToAssetRegistry = EAssetRegistryWritebackMethod::Disabled;
-	bool bWritePluginSizeSummaryJsons = false; // Only valid if WriteBackMetadataToAssetRegistry != Disabled.
-
-	FOodleDataCompression::ECompressor ShaderOodleCompressor = FOodleDataCompression::ECompressor::Mermaid;
-	FOodleDataCompression::ECompressionLevel ShaderOodleLevel = FOodleDataCompression::ECompressionLevel::Normal;
-
-	bool IsDLC() const
-	{
-		return DLCPluginPath.Len() > 0;
-	}
-};
-
-struct FContainerTargetSpec
-{
-	FIoContainerId ContainerId;
-	FIoContainerId OptionalSegmentContainerId;
-	FIoContainerHeader Header;
-	FIoContainerHeader OptionalSegmentHeader;
-	FName Name;
-	FGuid EncryptionKeyGuid;
-	FString OutputPath;
-	FString OptionalSegmentOutputPath;
-	FString StageLooseFileRootPath;
-	TSharedPtr<IIoStoreWriter> IoStoreWriter;
-	TSharedPtr<IIoStoreWriter> OptionalSegmentIoStoreWriter;
-	TArray<FContainerTargetFile> TargetFiles;
-	TArray<TUniquePtr<FIoStoreReader>> PatchSourceReaders;
-	EIoContainerFlags ContainerFlags = EIoContainerFlags::None;
-	TArray<FCookedPackage*> Packages;
-	TSet<FShaderInfo*> GlobalShaders;
-	TSet<FShaderInfo*> SharedShaders;
-	TSet<FShaderInfo*> UniqueShaders;
-	TSet<FShaderInfo*> InlineShaders;
-	bool bGenerateDiffPatch = false;
-};
-
 using FPackageNameMap = TMap<FName, FCookedPackage*>;
 using FPackageIdMap = TMap<FPackageId, FCookedPackage*>;
 
@@ -2274,584 +1907,7 @@ TArray<TUniquePtr<FIoStoreReader>> CreatePatchSourceReaders(const TArray<FString
 	return Readers;
 }
 
-static bool LoadShaderAssetInfo(const FString& Filename, FCookedPackageStore* PackageStore, TMap<FSHAHash, TSet<FName>>& OutShaderCodeToAssets)
-{
-	FString JsonText; 
-	if (PackageStore && PackageStore->HasZenStoreClient())
-	{
-		FIoChunkId ChunkId = PackageStore->GetChunkIdFromFileName(Filename);
-		if (!ChunkId.IsValid())
-		{
-			return false;
-		}
-		TIoStatusOr<FIoBuffer> Buffer = PackageStore->ReadChunk(ChunkId);
-		if (!Buffer.IsOk())
-		{
-			return false;
-		}
-		FFileHelper::BufferToString(JsonText, Buffer.ValueOrDie().GetData(), (int32)Buffer.ValueOrDie().GetSize());
-	}
-	else
-	{
-		if (!FFileHelper::LoadFileToString(JsonText, *Filename))
-		{
-			return false;
-		}
-	}
-
-	TSharedPtr<FJsonObject> JsonObject;
-	TSharedRef<TJsonReader<TCHAR>> Reader = TJsonReaderFactory<TCHAR>::Create(JsonText);
-
-	if (!FJsonSerializer::Deserialize(Reader, JsonObject) || !JsonObject.IsValid())
-	{
-		return false;
-	}
-
-	TSharedPtr<FJsonValue> AssetInfoArrayValue = JsonObject->Values.FindRef(TEXT("ShaderCodeToAssets"));
-	if (!AssetInfoArrayValue.IsValid())
-	{
-		return false;
-	}
-
-	TArray<TSharedPtr<FJsonValue>> AssetInfoArray = AssetInfoArrayValue->AsArray();
-	for (int32 IdxPair = 0, NumPairs = AssetInfoArray.Num(); IdxPair < NumPairs; ++IdxPair)
-	{
-		TSharedPtr<FJsonObject> Pair = AssetInfoArray[IdxPair]->AsObject();
-		if (Pair.IsValid())
-		{
-			TSharedPtr<FJsonValue> ShaderMapHashJson = Pair->Values.FindRef(TEXT("ShaderMapHash"));
-			if (!ShaderMapHashJson.IsValid())
-			{
-				continue;
-			}
-			FSHAHash ShaderMapHash;
-			ShaderMapHash.FromString(ShaderMapHashJson->AsString());
-
-			TSharedPtr<FJsonValue> AssetPathsArrayValue = Pair->Values.FindRef(TEXT("Assets"));
-			if (!AssetPathsArrayValue.IsValid())
-			{
-				continue;
-			}
-
-			TSet<FName>& Assets = OutShaderCodeToAssets.Add(ShaderMapHash);
-			TArray<TSharedPtr<FJsonValue>> AssetPathsArray = AssetPathsArrayValue->AsArray();
-			for (int32 IdxAsset = 0, NumAssets = AssetPathsArray.Num(); IdxAsset < NumAssets; ++IdxAsset)
-			{
-				Assets.Add(FName(*AssetPathsArray[IdxAsset]->AsString()));
-			}
-		}
-	}
-	return true;
-}
-
-static bool ConvertToIoStoreShaderLibrary(
-	const TCHAR* FileName,
-	FIoChunkId InChunkId,
-	FCookedPackageStore* PackageStore,
-	FOodleDataCompression::ECompressor InShaderOodleCompressor,
-	FOodleDataCompression::ECompressionLevel InShaderOodleLevel,
-	TTuple<FIoChunkId, FIoBuffer>& OutLibraryIoChunk,
-	TArray<TTuple<FIoChunkId, FIoBuffer, uint32>>& OutCodeIoChunks,
-	TArray<TTuple<FSHAHash, TArray<FIoChunkId>>>& OutShaderMaps,
-	TArray<TTuple<FSHAHash, TSet<FName>>>& OutShaderMapAssetAssociations)
-{
-	IOSTORE_CPU_SCOPE(ConvertShaderLibrary);
-	// ShaderArchive-MyProject-PCD3D.ushaderbytecode
-	FStringView BaseFileNameView = FPathViews::GetBaseFilename(FileName);
-	int32 FormatStartIndex = -1;
-	if (!BaseFileNameView.FindLastChar('-', FormatStartIndex))
-	{
-		UE_LOG(LogIoStore, Error, TEXT("Invalid shader code library file name '%s'."), FileName);
-		return false;
-	}
-	int32 LibraryNameStartIndex = -1;
-	if (!BaseFileNameView.FindChar('-', LibraryNameStartIndex) || FormatStartIndex - LibraryNameStartIndex <= 1)
-	{
-		UE_LOG(LogIoStore, Error, TEXT("Invalid shader code library file name '%s'."), FileName);
-		return false;
-	}
-	FString LibraryName(BaseFileNameView.Mid(LibraryNameStartIndex + 1, FormatStartIndex - LibraryNameStartIndex - 1));
-	FName FormatName(BaseFileNameView.RightChop(FormatStartIndex + 1));
-	
-	FIoBuffer LibraryBuffer;
-	if (PackageStore && PackageStore->HasZenStoreClient() && InChunkId.IsValid())
-	{
-		LibraryBuffer = PackageStore->ReadChunk(InChunkId).ConsumeValueOrDie();
-	}
-	else
-	{
-		TUniquePtr<FArchive> LibraryAr(IFileManager::Get().CreateFileReader(FileName));
-		if (!LibraryAr)
-		{
-			UE_LOG(LogIoStore, Error, TEXT("Missing shader code library file '%s'."), FileName);
-			return false;
-		}
-
-		LibraryBuffer = FIoBuffer(LibraryAr->TotalSize());
-		LibraryAr->Serialize(LibraryBuffer.GetData(), LibraryBuffer.GetSize());
-	}
-
-	FLargeMemoryReader LibraryAr(LibraryBuffer.GetData(), LibraryBuffer.GetSize());
-	uint32 Version;
-	LibraryAr << Version;
-
-	FSerializedShaderArchive SerializedShaders;
-	LibraryAr << SerializedShaders;
-	int64 OffsetToShaderCode = LibraryAr.Tell();
-
-	FString AssetInfoFileName = FPaths::GetPath(FileName) / FString::Printf(TEXT("ShaderAssetInfo-%s-%s.assetinfo.json"), *LibraryName, *FormatName.ToString());
-	TMap<FSHAHash, FShaderMapAssetPaths> ShaderCodeToAssets;
-	if (!LoadShaderAssetInfo(AssetInfoFileName, PackageStore && InChunkId.IsValid() ? PackageStore : nullptr, ShaderCodeToAssets))
-	{
-		UE_LOG(LogIoStore, Error, TEXT("Failed loading asset asset info file '%s'"), *AssetInfoFileName);
-		return false;
-	}
-
-	FIoStoreShaderCodeArchiveHeader IoStoreLibraryHeader;
-	FIoStoreShaderCodeArchive::CreateIoStoreShaderCodeArchiveHeader(FName(FormatName), SerializedShaders, IoStoreLibraryHeader);
-	checkf(SerializedShaders.ShaderEntries.Num() == IoStoreLibraryHeader.ShaderEntries.Num(), TEXT("IoStore header has different number of shaders (%d) than the original library (%d)"), 
-		IoStoreLibraryHeader.ShaderEntries.Num(), SerializedShaders.ShaderEntries.Num());
-	checkf(SerializedShaders.ShaderMapEntries.Num() == IoStoreLibraryHeader.ShaderMapEntries.Num(), TEXT("IoStore header has different number of shadermaps (%d) than the original library (%d)"),
-		IoStoreLibraryHeader.ShaderMapEntries.Num(), SerializedShaders.ShaderMapEntries.Num());
-
-	// Load all the shaders, decompress them and recompress into groups. Each code chunk is shader group chunk now.
-	// This also updates group's compressed size, the only value left to calculate.
-	int32 GroupCount = IoStoreLibraryHeader.ShaderGroupEntries.Num();
-	OutCodeIoChunks.SetNum(GroupCount);
-	FCriticalSection DiskArchiveAccess;
-	ParallelFor(OutCodeIoChunks.Num(),
-		[&OutCodeIoChunks, &IoStoreLibraryHeader, &DiskArchiveAccess, &LibraryBuffer, &SerializedShaders, &OffsetToShaderCode, InShaderOodleCompressor, InShaderOodleLevel](int32 GroupIndex)
-		{
-			FIoStoreShaderGroupEntry& Group = IoStoreLibraryHeader.ShaderGroupEntries[GroupIndex];
-			uint8* UncompressedGroupMemory = reinterpret_cast<uint8*>(FMemory::Malloc(Group.UncompressedSize));
-			check(UncompressedGroupMemory);
-
-			int32 SmallestShaderInGroupByOrdinal = MAX_int32;
-			uint8* ShaderStart = UncompressedGroupMemory;
-			// Not worth to special-case for group size of 1 (by converting to copy) - such groups are very fast to decompress and recompress.
-			for (uint32 IdxShaderInGroup = 0; IdxShaderInGroup < Group.NumShaders; ++IdxShaderInGroup)
-			{
-				int32 ShaderIndex = IoStoreLibraryHeader.ShaderIndices[Group.ShaderIndicesOffset + IdxShaderInGroup];
-				SmallestShaderInGroupByOrdinal = FMath::Min(SmallestShaderInGroupByOrdinal, ShaderIndex);
-				const FIoStoreShaderCodeEntry& Shader = IoStoreLibraryHeader.ShaderEntries[ShaderIndex];
-
-				checkf((reinterpret_cast<uint64>(ShaderStart) - reinterpret_cast<uint64>(UncompressedGroupMemory)) == Shader.UncompressedOffsetInGroup,
-					TEXT("Shader uncompressed offset in group does not agree with its actual placement: Shader.UncompressedOffsetInGroup=%llu, actual=%llu"),
-					Shader.UncompressedOffsetInGroup, (reinterpret_cast<uint64>(ShaderStart) - reinterpret_cast<uint64>(UncompressedGroupMemory))
-				);
-
-				// load and decompress the shader at the desired offset
-				const FShaderCodeEntry& IndividuallyCompressedShader = SerializedShaders.ShaderEntries[ShaderIndex];
-				{
-					// small shaders might be stored without compression, handle them here
-					if (IndividuallyCompressedShader.Size == IndividuallyCompressedShader.UncompressedSize)
-					{
-						FMemory::Memcpy(ShaderStart, LibraryBuffer.GetData() + OffsetToShaderCode + IndividuallyCompressedShader.Offset, IndividuallyCompressedShader.Size);
-					}
-					else
-					{
-						// This function will crash if decompression fails.
-						ShaderCodeArchive::DecompressShaderWithOodle(ShaderStart, IndividuallyCompressedShader.UncompressedSize, LibraryBuffer.GetData() + OffsetToShaderCode + IndividuallyCompressedShader.Offset, IndividuallyCompressedShader.Size);
-					}
-				}
-
-				ShaderStart += IndividuallyCompressedShader.UncompressedSize;
-			}
-
-			checkf((reinterpret_cast<uint64>(ShaderStart) - reinterpret_cast<uint64>(UncompressedGroupMemory)) == Group.UncompressedSize,
-				TEXT("Uncompressed shader group size does not agree with the actual results (Group.UncompressedSize=%llu, actual=%llu)"), 
-				Group.UncompressedSize, (reinterpret_cast<uint64>(ShaderStart) - reinterpret_cast<uint64>(UncompressedGroupMemory)));
-
-			// now compress the whole group
-			int64 CompressedGroupSize = 0;
-			uint8* CompressedShaderGroupMemory = nullptr;
-			ShaderCodeArchive::CompressShaderWithOodle(nullptr, CompressedGroupSize, UncompressedGroupMemory, Group.UncompressedSize, InShaderOodleCompressor, InShaderOodleLevel);
-			checkf(CompressedGroupSize > 0, TEXT("CompressedGroupSize estimate seems wrong (%lld)"), CompressedGroupSize);
-
-			CompressedShaderGroupMemory = reinterpret_cast<uint8*>(FMemory::Malloc(CompressedGroupSize));
-			bool bCompressed = ShaderCodeArchive::CompressShaderWithOodle(CompressedShaderGroupMemory, CompressedGroupSize, UncompressedGroupMemory, Group.UncompressedSize, InShaderOodleCompressor, InShaderOodleLevel);
-			checkf(bCompressed, TEXT("We could not compress the shader group after providing an estimated memory."));
-
-			TTuple<FIoChunkId, FIoBuffer, uint32>& OutCodeIoChunk = OutCodeIoChunks[GroupIndex];
-			OutCodeIoChunk.Get<0>() = IoStoreLibraryHeader.ShaderGroupIoHashes[GroupIndex];
-			// This value is the load order factor for the group (the smaller, the more likely), that IoStore will use to sort the chunks.
-			// We calculate it as the smallest shader index in the group, which shouldn't be bad. Arguably a better way would be to get the lowest-numbered shadermap that references _any_ shader in the group, 
-			// but it's much, much slower to calculate and the resulting order would be likely the same/similar most of the time
-			checkf(SmallestShaderInGroupByOrdinal >= 0 && SmallestShaderInGroupByOrdinal < IoStoreLibraryHeader.ShaderEntries.Num(), TEXT("SmallestShaderInGroupByOrdinal has an invalid value of %d (not within [0, %d) range as expected)"),
-				SmallestShaderInGroupByOrdinal, IoStoreLibraryHeader.ShaderEntries.Num());
-			OutCodeIoChunk.Get<2>() = static_cast<uint32>(SmallestShaderInGroupByOrdinal);
-
-			if (CompressedGroupSize < Group.UncompressedSize)
-			{
-				OutCodeIoChunk.Get<1>() = FIoBuffer(CompressedGroupSize);
-				FMemory::Memcpy(OutCodeIoChunk.Get<1>().GetData(), CompressedShaderGroupMemory, CompressedGroupSize);
-				Group.CompressedSize = CompressedGroupSize;
-			}
-			else
-			{
-				// store uncompressed (unlikely, but happens for a 200-byte sized shader that happens to get its own group)
-				OutCodeIoChunk.Get<1>() = FIoBuffer(Group.UncompressedSize);
-				FMemory::Memcpy(OutCodeIoChunk.Get<1>().GetData(), UncompressedGroupMemory, Group.UncompressedSize);
-				Group.CompressedSize = Group.UncompressedSize;
-			}
-			FMemory::Free(UncompressedGroupMemory);
-			FMemory::Free(CompressedShaderGroupMemory);
-		},
-		EParallelForFlags::Unbalanced
-	);
-
-	// calculate and log stats
-	int64 TotalUncompressedSizeViaGroups = 0;	// calculate total uncompressed size twice for a sanity check
-	int64 TotalUncompressedSizeViaShaders = 0;
-	int64 TotalIndividuallyCompressedSize = 0;
-	int64 TotalGroupCompressedSize = 0;
-	for (int32 IdxGroup = 0, NumGroups = IoStoreLibraryHeader.ShaderGroupEntries.Num(); IdxGroup < NumGroups; ++IdxGroup)
-	{
-		FIoStoreShaderGroupEntry& Group = IoStoreLibraryHeader.ShaderGroupEntries[IdxGroup];
-		TotalGroupCompressedSize += Group.CompressedSize;
-		TotalUncompressedSizeViaGroups += Group.UncompressedSize;
-
-		// now go via shaders
-		for (uint32 IdxShaderInGroup = 0; IdxShaderInGroup < Group.NumShaders; ++IdxShaderInGroup)
-		{
-			int32 ShaderIndex = IoStoreLibraryHeader.ShaderIndices[Group.ShaderIndicesOffset + IdxShaderInGroup];
-			const FShaderCodeEntry& IndividuallyCompressedShader = SerializedShaders.ShaderEntries[ShaderIndex];
-
-			TotalIndividuallyCompressedSize += IndividuallyCompressedShader.Size;
-			TotalUncompressedSizeViaShaders += IndividuallyCompressedShader.UncompressedSize;
-		}
-	}
-
-	checkf(TotalUncompressedSizeViaGroups == TotalUncompressedSizeViaShaders, TEXT("Sanity check failure: total uncompressed shader size differs if calculated via shader groups (%lld) or individual shaders (%lld)"),
-		TotalUncompressedSizeViaGroups, TotalUncompressedSizeViaShaders
-		);
-
-	UE_LOG(LogIoStore, Display, TEXT("%s(%s): Recompressed %d shaders as %d groups. Library size changed from %lld KB (%.2f : 1 ratio) to %lld KB (%.2f : 1 ratio), %.2f%% of previous."),
-		*LibraryName, *FormatName.ToString(),
-		SerializedShaders.ShaderEntries.Num(), IoStoreLibraryHeader.ShaderGroupEntries.Num(),
-		TotalIndividuallyCompressedSize / 1024, static_cast<double>(TotalUncompressedSizeViaShaders) / static_cast<double>(TotalIndividuallyCompressedSize),
-		TotalGroupCompressedSize / 1024, static_cast<double>(TotalUncompressedSizeViaShaders) / static_cast<double>(TotalGroupCompressedSize),
-		100.0 * static_cast<double>(TotalGroupCompressedSize) / static_cast<double>(TotalIndividuallyCompressedSize)
-		);
-
-	FLargeMemoryWriter IoStoreLibraryAr(0, true);
-	FIoStoreShaderCodeArchive::SaveIoStoreShaderCodeArchive(IoStoreLibraryHeader, IoStoreLibraryAr);
-	OutLibraryIoChunk.Key = FIoStoreShaderCodeArchive::GetShaderCodeArchiveChunkId(LibraryName, FormatName);
-	int64 TotalSize = IoStoreLibraryAr.TotalSize();
-	OutLibraryIoChunk.Value = FIoBuffer(FIoBuffer::AssumeOwnership, IoStoreLibraryAr.ReleaseOwnership(), TotalSize);
-
-	int32 ShaderMapCount = IoStoreLibraryHeader.ShaderMapHashes.Num();
-	OutShaderMaps.SetNum(ShaderMapCount);
-	ParallelFor(OutShaderMaps.Num(),
-		[&OutShaderMaps, &IoStoreLibraryHeader](int32 ShaderMapIndex)
-		{
-			TTuple<FSHAHash, TArray<FIoChunkId>>& OutShaderMap = OutShaderMaps[ShaderMapIndex];
-			OutShaderMap.Key = IoStoreLibraryHeader.ShaderMapHashes[ShaderMapIndex];
-			const FIoStoreShaderMapEntry& ShaderMapEntry = IoStoreLibraryHeader.ShaderMapEntries[ShaderMapIndex];
-			int32 LookupIndexEnd = ShaderMapEntry.ShaderIndicesOffset + ShaderMapEntry.NumShaders;
-			OutShaderMap.Value.Reserve(ShaderMapEntry.NumShaders);	// worst-case, 1 group == 1 shader
-			for (int32 ShaderLookupIndex = ShaderMapEntry.ShaderIndicesOffset; ShaderLookupIndex < LookupIndexEnd; ++ShaderLookupIndex)
-			{
-				int32 ShaderIndex = IoStoreLibraryHeader.ShaderIndices[ShaderLookupIndex];
-				int32 GroupIndex = IoStoreLibraryHeader.ShaderEntries[ShaderIndex].ShaderGroupIndex;
-				OutShaderMap.Value.AddUnique(IoStoreLibraryHeader.ShaderGroupIoHashes[GroupIndex]);
-			}
-		},
-		EParallelForFlags::Unbalanced
-	);
-
-	OutShaderMapAssetAssociations.Reserve(ShaderCodeToAssets.Num());
-	for (auto& KV : ShaderCodeToAssets)
-	{
-		OutShaderMapAssetAssociations.Emplace(KV.Key, MoveTemp(KV.Value));
-	}
-
-	return true;
-}
-
-// Carries association between shaders and packages from shader processing to size assignment.
-struct FShaderAssociationInfo
-{
-	struct FShaderChunkInfo
-	{
-		enum EType 
-		{
-			// This shader is referenced by one or more packages and can assign its size to
-			// those packages
-			Package, 
-			// This shader is needed by the baseline engine in some manner and is a flat
-			// size cost.
-			Global, 
-			// This shader isn't global or referenced by packages - theoretically shouldn't
-			// exist.
-			Orphan
-		};
-
-		EType Type;
-		TArray<FName> ReferencedByPackages;
-		uint32 CompressedSize = 0;
-		UE::Cook::EPluginSizeTypes SizeType;
-
-		// If we are shared across plugins, this is our pseudo plugin name. This is generated during size assignment.
-		FString SharedPluginName;
-	};
-	
-	struct FShaderChunkInfoKey
-	{
-		FName PakChunkName;
-		FIoChunkId IoChunkId;
-
-		friend uint32 GetTypeHash(const FShaderChunkInfoKey& Key) { return GetTypeHash(Key.IoChunkId); }
-		friend bool operator == (const FShaderChunkInfoKey& LHS, const FShaderChunkInfoKey& RHS) { return LHS.PakChunkName == RHS.PakChunkName && LHS.IoChunkId == RHS.IoChunkId; }
-	};
-
-	TMap<FShaderChunkInfoKey, FShaderChunkInfo> ShaderChunkInfos;
-
-	// The list of shaders referenced by each package. This can be used to look up in to ContainerShaderInfos.
-	TMap<FName /* PackageName */, TArray<FShaderChunkInfoKey>> PackageShaderMap;
-};
-
-static void ProcessShaderLibraries(const FIoStoreArguments& Arguments, TArray<FContainerTargetSpec*>& ContainerTargets, TArray<FShaderInfo*> OutShaders, FShaderAssociationInfo& OutAssocInfo)
-{
-	IOSTORE_CPU_SCOPE(ProcessShaderLibraries);
-
-	double LibraryStart = FPlatformTime::Seconds();
-
-	TMap<FIoChunkId, FShaderInfo*> ChunkIdToShaderInfoMap;
-	TMap<FSHAHash, TArray<FIoChunkId>> ShaderChunkIdsByShaderMapHash;
-	TMap<FName, TSet<FSHAHash>> PackageNameToShaderMaps;
-	TMap<FContainerTargetSpec*, TSet<FShaderInfo*>> AllContainerShaderLibraryShadersMap;
-
-	{
-		IOSTORE_CPU_SCOPE(ConvertShaderLibraries);
-		for (FContainerTargetSpec* ContainerTarget : ContainerTargets)
-		{
-			for (FContainerTargetFile& TargetFile : ContainerTarget->TargetFiles)
-			{
-				if (TargetFile.ChunkType == EContainerChunkType::ShaderCodeLibrary)
-				{
-					TSet<FShaderInfo*>& ContainerShaderLibraryShaders = AllContainerShaderLibraryShadersMap.FindOrAdd(ContainerTarget);
-
-					TArray<TTuple<FSHAHash, TArray<FIoChunkId>>> ShaderMaps;
-					TTuple<FIoChunkId, FIoBuffer> LibraryChunk;
-					TArray<TTuple<FIoChunkId, FIoBuffer, uint32>> CodeChunks;
-					TArray<TTuple<FSHAHash, TSet<FName>>> ShaderMapAssetAssociations;
-				    if (!ConvertToIoStoreShaderLibrary(*TargetFile.NormalizedSourcePath, TargetFile.ChunkId, Arguments.PackageStore.Get(), Arguments.ShaderOodleCompressor, Arguments.ShaderOodleLevel, LibraryChunk, CodeChunks, ShaderMaps, ShaderMapAssetAssociations))
-					{
-						UE_LOG(LogIoStore, Warning, TEXT("Failed converting shader library '%s'"), *TargetFile.NormalizedSourcePath);
-						continue;
-					}
-					TargetFile.ChunkId = LibraryChunk.Key;
-					TargetFile.SourceBuffer.Emplace(LibraryChunk.Value);
-					TargetFile.SourceSize = LibraryChunk.Value.GetSize();
-
-					const bool bIsGlobalShaderLibrary = FPaths::GetCleanFilename(TargetFile.NormalizedSourcePath).StartsWith(TEXT("ShaderArchive-Global-"));
-					const FShaderInfo::EShaderType ShaderType = bIsGlobalShaderLibrary ? FShaderInfo::Global : FShaderInfo::Normal;
-					for (const TTuple<FIoChunkId, FIoBuffer, uint32>& CodeChunk : CodeChunks)
-					{
-						const FIoChunkId& ShaderChunkId = CodeChunk.Get<0>();
-						FShaderInfo* ShaderInfo = ChunkIdToShaderInfoMap.FindRef(ShaderChunkId);
-						if (!ShaderInfo)
-						{
-							ShaderInfo = new FShaderInfo();
-							ShaderInfo->ChunkId = ShaderChunkId;
-							ShaderInfo->CodeIoBuffer = CodeChunk.Get<1>();
-							ShaderInfo->LoadOrderFactor = CodeChunk.Get<2>();
-							OutShaders.Add(ShaderInfo);
-							ChunkIdToShaderInfoMap.Add(ShaderChunkId, ShaderInfo);
-						}
-						else
-						{
-							// first, make sure that the code is exactly the same
-							if (ShaderInfo->CodeIoBuffer != CodeChunk.Get<1>())
-							{
-								UE_LOG(LogIoStore, Error, TEXT("Collision of two shader code chunks, same Id (%s), different code. Packaged game will likely crash, not being able to decompress the shaders."), *LexToString(ShaderChunkId) );
-							}
-
-							// If we already exist, then we have two separate LoadOrderFactors,
-							// which one we got first affects build determinism. Take the lower.
-							if (CodeChunk.Get<2>() < ShaderInfo->LoadOrderFactor)
-							{
-								ShaderInfo->LoadOrderFactor = CodeChunk.Get<2>();
-							}
-						}
-
-
-						const FShaderInfo::EShaderType* CurrentShaderTypeInContainer = ShaderInfo->TypeInContainer.Find(ContainerTarget);
-						if (!CurrentShaderTypeInContainer || *CurrentShaderTypeInContainer != FShaderInfo::Global)
-						{
-							// If a shader is both global and shared consider it to be global
-							ShaderInfo->TypeInContainer.Add(ContainerTarget, ShaderType);
-						}
-						
-						ContainerShaderLibraryShaders.Add(ShaderInfo);
-					}
-					for (const TTuple<FSHAHash, TSet<FName>>& ShaderMapAssetAssociation : ShaderMapAssetAssociations)
-					{
-						for (const FName& PackageName : ShaderMapAssetAssociation.Value)
-						{
-							TSet<FSHAHash>& PackageShaderMaps = PackageNameToShaderMaps.FindOrAdd(PackageName);
-							PackageShaderMaps.Add(ShaderMapAssetAssociation.Key);
-						}
-					}
-
-					for (TTuple<FSHAHash, TArray<FIoChunkId>>& ShaderMap : ShaderMaps)
-					{
-						TArray<FIoChunkId>& ShaderMapChunkIds = ShaderChunkIdsByShaderMapHash.FindOrAdd(ShaderMap.Key);
-						ShaderMapChunkIds.Append(MoveTemp(ShaderMap.Value));
-					}
-				} // end if containerchunktype shadercodelibrary
-			} // end foreach targetfile
-		} // end foreach container
-	}
-
-	// 1. Update ShaderInfos with which packages we reference.
-	// 2. Add to packages which shaders we use.
-	// 3. Add to PackageStore what shaders we use.
-	{
-		IOSTORE_CPU_SCOPE(UpdatePackageStoreShaders);
-		for (FContainerTargetSpec* ContainerTarget : ContainerTargets)
-		{
-			for (FCookedPackage* Package : ContainerTarget->Packages)
-			{
-				TSet<FSHAHash>* FindShaderMapHashes = PackageNameToShaderMaps.Find(Package->PackageName);
-				if (FindShaderMapHashes)
-				{
-					for (const FSHAHash& ShaderMapHash : *FindShaderMapHashes)
-					{
-						const TArray<FIoChunkId>* FindChunkIds = ShaderChunkIdsByShaderMapHash.Find(ShaderMapHash);
-						if (!FindChunkIds)
-						{
-							UE_LOG(LogIoStore, Warning, TEXT("Package '%s' in '%s' referencing missing shader map '%s'"), *Package->PackageName.ToString(), *ContainerTarget->Name.ToString(), *ShaderMapHash.ToString());
-							continue;
-						}
-						Package->ShaderMapHashes.Add(ShaderMapHash);
-						for (const FIoChunkId& ShaderChunkId : *FindChunkIds)
-						{
-							FShaderInfo* ShaderInfo = ChunkIdToShaderInfoMap.FindRef(ShaderChunkId);
-							if (!ShaderInfo)
-							{
-								UE_LOG(LogIoStore, Warning, TEXT("Package '%s' in '%s' referencing missing shader with chunk id '%s'"), *Package->PackageName.ToString(), *ContainerTarget->Name.ToString(), *BytesToHex(ShaderChunkId.GetData(), ShaderChunkId.GetSize()));
-								continue;
-							}
-
-							check(ShaderInfo);
-							ShaderInfo->ReferencedByPackages.Add(Package);
-							Package->Shaders.AddUnique(ShaderInfo);
-						}
-					}
-				}
-				Algo::Sort(Package->ShaderMapHashes);
-			}
-		}
-	}
-
-	{
-		IOSTORE_CPU_SCOPE(AddShaderChunks);
-		for (FContainerTargetSpec* ContainerTarget : ContainerTargets)
-		{
-			auto AddShaderTargetFile = [&ContainerTarget](FShaderInfo* ShaderInfo)
-			{
-				FContainerTargetFile& ShaderTargetFile = ContainerTarget->TargetFiles.AddDefaulted_GetRef();
-				ShaderTargetFile.ContainerTarget = ContainerTarget;
-				ShaderTargetFile.ChunkId = ShaderInfo->ChunkId;
-				ShaderTargetFile.ChunkType = EContainerChunkType::ShaderCode;
-				ShaderTargetFile.bForceUncompressed = true;
-				ShaderTargetFile.SourceBuffer.Emplace(ShaderInfo->CodeIoBuffer);
-				ShaderTargetFile.SourceSize = ShaderInfo->CodeIoBuffer.DataSize();
-			};
-
-			const TSet<FShaderInfo*>* FindContainerShaderLibraryShaders = AllContainerShaderLibraryShadersMap.Find(ContainerTarget);
-			if (FindContainerShaderLibraryShaders)
-			{
-				for (FCookedPackage* Package : ContainerTarget->Packages)
-				{
-					for (FShaderInfo* ShaderInfo : Package->Shaders)
-					{
-						if (ShaderInfo->ReferencedByPackages.Num() == 1)
-						{
-							FShaderInfo::EShaderType* ShaderType = ShaderInfo->TypeInContainer.Find(ContainerTarget);
-							if (ShaderType && *ShaderType != FShaderInfo::Global)
-							{
-								*ShaderType = FShaderInfo::Inline;
-							}
-						}
-					}
-				}
-
-				for (FShaderInfo* ShaderInfo : *FindContainerShaderLibraryShaders)
-				{
-					FShaderAssociationInfo::FShaderChunkInfoKey ShaderChunkInfoKey = { ContainerTarget->Name, ShaderInfo->ChunkId };
-					FShaderAssociationInfo::FShaderChunkInfo& ShaderChunkInfo = OutAssocInfo.ShaderChunkInfos.Add(ShaderChunkInfoKey);
-
-					FShaderInfo::EShaderType* ShaderType = ShaderInfo->TypeInContainer.Find(ContainerTarget);
-					check(ShaderType);
-					if (*ShaderType == FShaderInfo::Global)
-					{
-						ContainerTarget->GlobalShaders.Add(ShaderInfo);
-						ShaderChunkInfo.Type = FShaderAssociationInfo::FShaderChunkInfo::Global;
-					}
-					else if (*ShaderType == FShaderInfo::Inline)
-					{
-						ContainerTarget->InlineShaders.Add(ShaderInfo);
-
-						ShaderChunkInfo.Type = FShaderAssociationInfo::FShaderChunkInfo::Package;
-						checkf(ShaderInfo->ReferencedByPackages.Num() == 1, TEXT("Inline shader chunks must be referenced by 1 package only, but shader chunk %s is referenced by %d"),
-							*LexToString(ShaderInfo->ChunkId), ShaderInfo->ReferencedByPackages.Num());
-
-						ShaderChunkInfo.ReferencedByPackages.Add(ShaderInfo->ReferencedByPackages[FSetElementId::FromInteger(0)]->PackageName);
-					}
-					else if (ShaderInfo->ReferencedByPackages.Num() > 1)
-					{
-						ContainerTarget->SharedShaders.Add(ShaderInfo);
-
-						ShaderChunkInfo.Type = FShaderAssociationInfo::FShaderChunkInfo::Package;
-						
-						for (FCookedPackage* Package : ShaderInfo->ReferencedByPackages)
-						{
-							ShaderChunkInfo.ReferencedByPackages.Add(Package->PackageName);
-
-							TArray<FShaderAssociationInfo::FShaderChunkInfoKey>& PackageShaders = OutAssocInfo.PackageShaderMap.FindOrAdd(Package->PackageName);
-							PackageShaders.Add(ShaderChunkInfoKey);
-						}
-					}
-					else
-					{
-						//
-						// Note that we can get here with shaders that get split off in to another container (e.g. sm6 shaders). 
-						// Since they are in a different pakChunk they can't get inlined or shared. However, they still "belong" to
-						// the referencing packages for association purposes.
-						//
-						if (ShaderInfo->ReferencedByPackages.Num())
-						{
-							ShaderChunkInfo.Type = FShaderAssociationInfo::FShaderChunkInfo::Package;
-
-							for (FCookedPackage* Package : ShaderInfo->ReferencedByPackages)
-							{
-								ShaderChunkInfo.ReferencedByPackages.Add(Package->PackageName);
-
-								TArray<FShaderAssociationInfo::FShaderChunkInfoKey>& PackageShaders = OutAssocInfo.PackageShaderMap.FindOrAdd(Package->PackageName);
-								PackageShaders.Add(ShaderChunkInfoKey);
-							}
-						}
-						else
-						{
-							ShaderChunkInfo.Type = FShaderAssociationInfo::FShaderChunkInfo::Orphan;
-						}
-
-						ContainerTarget->UniqueShaders.Add(ShaderInfo);
-					}
-					AddShaderTargetFile(ShaderInfo);
-				}
-			}
-		}
-	}
-
-	double LibraryEnd = FPlatformTime::Seconds();
-	UE_LOG(LogIoStore, Display, TEXT("Shaders processed in %.2f seconds"), LibraryEnd - LibraryStart);
-}
-
-void InitializeContainerTargetsAndPackages(
+static void InitializeContainerTargetsAndPackages(
 	const FIoStoreArguments& Arguments,
 	TArray<FCookedPackage*>& Packages,
 	FPackageNameMap& PackageNameMap,
@@ -5427,7 +4483,7 @@ void CreateContainerHeader(FContainerTargetSpec& ContainerTarget, bool bIsOption
 	OptionalSegmentStoreEntriesWriter.Flush(Header.OptionalSegmentStoreEntries);
 }
 
-int32 CreateTarget(const FIoStoreArguments& Arguments, const FIoStoreWriterSettings& GeneralIoWriterSettings)
+static int32 CreateTarget(const FIoStoreArguments& Arguments, const FIoStoreWriterSettings& GeneralIoWriterSettings)
 {
 	IOSTORE_CPU_SCOPE(CreateTarget);
 	TGuardValue<int32> GuardAllowUnversionedContentInEditor(GAllowUnversionedContentInEditor, 1);
@@ -5611,9 +4667,10 @@ int32 CreateTarget(const FIoStoreArguments& Arguments, const FIoStoreWriterSetti
 	}
 
 	UE_LOG(LogIoStore, Display, TEXT("Processing shader libraries, compressing with Oodle %s, level %d (%s)"), FOodleDataCompression::ECompressorToString(Arguments.ShaderOodleCompressor), (int32)Arguments.ShaderOodleLevel, FOodleDataCompression::ECompressionLevelToString(Arguments.ShaderOodleLevel));
-	TArray<FShaderInfo*> Shaders;
-	FShaderAssociationInfo ShaderAssocInfo;
-	ProcessShaderLibraries(Arguments, ContainerTargets, Shaders, ShaderAssocInfo);
+	FIoStoreShaderLibraryProcessor ShaderLibraryProcessor;
+	FIoStoreShaderLibraryOutput ShaderLibraryOutput;
+	ShaderLibraryProcessor.ProcessShaderLibraries(Arguments, ContainerTargets, ShaderLibraryOutput);
+	FShaderAssociationInfo& ShaderAssocInfo = ShaderLibraryOutput.AssocInfo;
 
 	FIoStoreWriteRequestManager WriteRequestManager(PackageStoreOptimizer, Arguments.PackageStore.Get());
 
diff --git a/Engine/Source/Developer/IoStoreUtilities/Private/IoStoreUtilitiesCore.cpp b/Engine/Source/Developer/IoStoreUtilities/Private/IoStoreUtilitiesCore.cpp
new file mode 100644
--- /dev/null
+++ b/Engine/Source/Developer/IoStoreUtilities/Private/IoStoreUtilitiesCore.cpp
@@ -0,0 +1,118 @@
+// Copyright Epic Games, Inc. All Rights Reserved.
+
+#include "IoStoreUtilitiesCore.h"
+#include "Misc/Paths.h"
+
+
+namespace UE::IoStore::Private
+{
+
+FCookedFileStatMap::FCookedFileStatMap()
+{
+	Extensions.Emplace(TEXT(".umap"), FCookedFileStatData::PackageHeader);
+	Extensions.Emplace(TEXT(".uasset"), FCookedFileStatData::PackageHeader);
+	Extensions.Emplace(TEXT(".uexp"), FCookedFileStatData::PackageData);
+	Extensions.Emplace(TEXT(".ubulk"), FCookedFileStatData::BulkData);
+	Extensions.Emplace(TEXT(".uptnl"), FCookedFileStatData::OptionalBulkData);
+	Extensions.Emplace(TEXT(".m.ubulk"), FCookedFileStatData::MemoryMappedBulkData);
+	Extensions.Emplace(*FString::Printf(TEXT(".uexp%s"), FFileRegion::RegionsFileExtension), FCookedFileStatData::Regions);
+	Extensions.Emplace(*FString::Printf(TEXT(".uptnl%s"), FFileRegion::RegionsFileExtension), FCookedFileStatData::Regions);
+	Extensions.Emplace(*FString::Printf(TEXT(".ubulk%s"), FFileRegion::RegionsFileExtension), FCookedFileStatData::Regions);
+	Extensions.Emplace(*FString::Printf(TEXT(".m.ubulk%s"), FFileRegion::RegionsFileExtension), FCookedFileStatData::Regions);
+	Extensions.Emplace(TEXT(".ushaderbytecode"), FCookedFileStatData::ShaderLibrary);
+	Extensions.Emplace(TEXT(".o.umap"), FCookedFileStatData::OptionalSegmentPackageHeader);
+	Extensions.Emplace(TEXT(".o.uasset"), FCookedFileStatData::OptionalSegmentPackageHeader);
+	Extensions.Emplace(TEXT(".o.uexp"), FCookedFileStatData::OptionalSegmentPackageData);
+	Extensions.Emplace(TEXT(".o.ubulk"), FCookedFileStatData::OptionalSegmentBulkData);
+}
+
+void FCookedFileStatMap::Add(const TCHAR* Path, int64 FileSize)
+{
+	FString NormalizedPath(Path);
+	FPaths::NormalizeFilename(NormalizedPath);
+
+	FStringView NormalizePathView(NormalizedPath);
+	int32 ExtensionStartIndex = GetFullExtensionStartIndex(NormalizePathView);
+	if (ExtensionStartIndex < 0)
+	{
+		return;
+	}
+	FStringView Extension = NormalizePathView.RightChop(ExtensionStartIndex);
+	const FCookedFileStatData::EFileType* FindFileType = FindFileTypeFromExtension(Extension);
+	if (!FindFileType)
+	{
+		return;
+	}
+
+	FCookedFileStatData& CookedFileStatData = Map.Add(MoveTemp(NormalizedPath));
+	CookedFileStatData.FileType = *FindFileType;
+	CookedFileStatData.FileSize = FileSize;
+}
+
+const FCookedFileStatData* FCookedFileStatMap::Find(FStringView NormalizedPath) const
+{
+	return Map.FindByHash(GetTypeHash(NormalizedPath), NormalizedPath);
+}
+
+const FCookedFileStatData::EFileType* FCookedFileStatMap::FindFileTypeFromExtension(const FStringView& Extension)
+{
+	for (const auto& Pair : Extensions)
+	{
+		if (Pair.Key == Extension)
+		{
+			return &Pair.Value;
+		}
+	}
+	return nullptr;
+}
+
+uint32 GetTypeHash(const FIoStoreOrderingOptions& OrderingOptions)
+{
+	uint32 Hash = ::GetTypeHash(OrderingOptions.bClusterByOrderFilePriority);
+	Hash = HashCombine(Hash, ::GetTypeHash(OrderingOptions.bAlphaSortClusterPackageLists));
+	Hash = HashCombine(Hash, ::GetTypeHash(OrderingOptions.FallbackOrderMode));
+	return Hash;
+}
+
+int32 GetFullExtensionStartIndex(FStringView Path)
+{
+	int32 ExtensionStartIndex = -1;
+	for (int32 Index = Path.Len() - 1; Index >= 0; --Index)
+	{
+		// Check if we reached the end of the filename
+		if (FPathViews::IsSeparator(Path[Index]))
+		{
+			break;
+		}
+		else if (Path[Index] == '.')
+		{
+			if (ExtensionStartIndex != -1)
+			{
+				// As we have already found a '.' we need to check for the cooked index
+				FStringView Extension = Path.SubStr(Index + 1, (ExtensionStartIndex - Index) - 1);
+				if (UE::String::IsNumericOnlyDigits(Extension))
+				{
+					return ExtensionStartIndex;
+				}
+			}
+
+			ExtensionStartIndex = Index;
+		}
+	}
+	return ExtensionStartIndex;
+}
+
+FStringView GetFullExtension(FStringView Path)
+{
+	int32 ExtensionStartIndex = GetFullExtensionStartIndex(Path);
+	if (ExtensionStartIndex < 0)
+	{
+		return FStringView();
+	}
+	else
+	{
+		return Path.RightChop(ExtensionStartIndex);
+	}
+}
+
+} // namespace UE::IoStore::Private
diff --git a/Engine/Source/Developer/IoStoreUtilities/Private/IoStoreUtilitiesCore.h b/Engine/Source/Developer/IoStoreUtilities/Private/IoStoreUtilitiesCore.h
new file mode 100644
--- /dev/null
+++ b/Engine/Source/Developer/IoStoreUtilities/Private/IoStoreUtilitiesCore.h
@@ -0,0 +1,305 @@
+// Copyright Epic Games, Inc. All Rights Reserved.
+
+#pragma once
+
+#include "Containers/Map.h"
+#include "Containers/UnrealString.h"
+#include "Compression/OodleDataCompression.h"
+#include "CookedPackageStore.h"
+#include "IO/IoContainerId.h"
+#include "IO/IoContainerHeader.h"
+#include "IO/IoDispatcher.h"
+#include "IoStoreWriter.h"
+#include "Misc/KeyChainUtilities.h"
+#include "PackageStoreOptimizer.h"
+#include "Settings/ProjectPackagingSettings.h" // for EAssetRegistryWritebackMethod
+#include "Templates/SharedPointer.h"
+#include "Templates/UniquePtr.h"
+#include "UObject/NameTypes.h"
+
+
+#define IOSTORE_CPU_SCOPE(NAME) TRACE_CPUPROFILER_EVENT_SCOPE(IoStore##NAME);
+#define IOSTORE_CPU_SCOPE_DATA(NAME, DATA) TRACE_CPUPROFILER_EVENT_SCOPE(IoStore##NAME);
+
+namespace UE::IoStore::Private
+{
+
+struct FContainerTargetSpec;
+
+struct FContainerSourceFile
+{
+	FString NormalizedPath;
+	FString DestinationPath;
+	bool bNeedsCompression = false;
+	bool bNeedsEncryption = false;
+};
+
+struct FContainerSourceSpec
+{
+	FName Name;
+	FString OutputPath;
+	FString OptionalOutputPath;
+	FString StageLooseFileRootPath;
+	TArray<FContainerSourceFile> SourceFiles;
+	FString PatchTargetFile;
+	TArray<FString> PatchSourceContainerFiles;
+	FString EncryptionKeyOverrideGuid;
+	bool bGenerateDiffPatch = false;
+	bool bOnDemand = false;
+};
+
+struct FCookedFileStatData
+{
+	enum EFileType
+	{
+		PackageHeader,
+		PackageData,
+		BulkData,
+		OptionalBulkData,
+		MemoryMappedBulkData,
+		ShaderLibrary,
+		Regions,
+		OptionalSegmentPackageHeader,
+		OptionalSegmentPackageData,
+		OptionalSegmentBulkData,
+		Invalid
+	};
+
+	int64 FileSize = 0;
+	EFileType FileType = Invalid;
+};
+
+class FCookedFileStatMap
+{
+public:
+	FCookedFileStatMap();
+
+	FORCEINLINE int32 Num() const
+	{
+		return Map.Num();
+	}
+
+	void Add(const TCHAR* Path, int64 FileSize);
+
+	const FCookedFileStatData* Find(FStringView NormalizedPath) const;
+
+private:
+	const FCookedFileStatData::EFileType* FindFileTypeFromExtension(const FStringView& Extension);
+
+	TArray<TTuple<FString, FCookedFileStatData::EFileType>> Extensions;
+	TMap<FString, FCookedFileStatData> Map;
+};
+
+enum class EFallbackOrderMode
+{
+	LoadOrder, // Default fallback ordering. Uses LoadOrder dependency ordering
+	AlphabeticalClustered, // As above but clusters are sorted alphabetically at the end for improved patching stability
+	Alphabetical // Pure alphabetical fallback ordering. Best for patching stability but not optimized to reduce seeks
+};
+
+struct FIoStoreOrderingOptions
+{
+	bool bClusterByOrderFilePriority = false;
+	bool bAlphaSortClusterPackageLists = false;
+
+	EFallbackOrderMode FallbackOrderMode = EFallbackOrderMode::LoadOrder;
+
+	bool operator==(const FIoStoreOrderingOptions& Other) const
+	{
+		return (bClusterByOrderFilePriority == Other.bClusterByOrderFilePriority) &&
+			(bAlphaSortClusterPackageLists == Other.bAlphaSortClusterPackageLists) &&
+			(FallbackOrderMode == Other.FallbackOrderMode);
+	}
+};
+
+uint32 GetTypeHash(const FIoStoreOrderingOptions& OrderingOptions);
+
+struct FReleasedPackages
+{
+	TSet<FName> PackageNames;
+	TMap<FPackageId, FName> PackageIdToName;
+};
+
+struct FFileOrderMap
+{
+	TMap<FName, int64> PackageNameToOrder;
+	FString Name;
+	int32 Priority;
+	int32 Index;
+
+	FFileOrderMap(int32 InPriority, int32 InIndex)
+		: Priority(InPriority)
+		, Index(InIndex)
+	{
+	}
+};
+
+struct FIoStoreArguments
+{
+	FString GlobalContainerPath;
+	FString CookedDir;
+	TArray<FContainerSourceSpec> Containers;
+	FCookedFileStatMap CookedFileStatMap;
+	TArray<FFileOrderMap> OrderMaps;
+	FKeyChain KeyChain;
+	FKeyChain PatchKeyChain;
+	FString DLCPluginPath;
+	FString DLCName;
+	FString ReferenceChunkGlobalContainerFileName;
+	FString ReferenceChunkChangesCSVFileName;
+	FString ReferenceChunkAdditionalContainersPath;
+	FString CsvPath;
+	FKeyChain ReferenceChunkKeys;
+	FReleasedPackages ReleasedPackages;
+	TUniquePtr<FCookedPackageStore> PackageStore;
+	TUniquePtr<FIoBuffer> ScriptObjects;
+	FIoStoreOrderingOptions OrderingOptions;
+	TMap<FString, FIoStoreOrderingOptions> OrderingOptionsOverrides;
+	bool bSign = false;
+	bool bRemapPluginContentToGame = false;
+	bool bCreateDirectoryIndex = true;
+	bool bFileRegions = false;
+	EAssetRegistryWritebackMethod WriteBackMetadataToAssetRegistry = EAssetRegistryWritebackMethod::Disabled;
+	bool bWritePluginSizeSummaryJsons = false; // Only valid if WriteBackMetadataToAssetRegistry != Disabled.
+
+	FOodleDataCompression::ECompressor ShaderOodleCompressor = FOodleDataCompression::ECompressor::Mermaid;
+	FOodleDataCompression::ECompressionLevel ShaderOodleLevel = FOodleDataCompression::ECompressionLevel::Normal;
+
+	bool IsDLC() const
+	{
+		return DLCPluginPath.Len() > 0;
+	}
+};
+
+struct FShaderInfo
+{
+	enum EShaderType { Normal, Global, Inline };
+
+	FIoChunkId ChunkId;
+	FIoBuffer CodeIoBuffer;
+	// the smaller the order, the more likely this will be loaded
+	uint32 LoadOrderFactor = MAX_uint32;
+	TSet<struct FCookedPackage*> ReferencedByPackages;
+	TMap<FContainerTargetSpec*, EShaderType> TypeInContainer;
+
+	// This is used to ensure build determinism and such must be stable
+	// across builds.
+	static bool Sort(const FShaderInfo* A, const FShaderInfo* B)
+	{
+		if (A->LoadOrderFactor == B->LoadOrderFactor)
+		{
+			// Shader chunk IDs are the hash of the shader so this is consistent across builds.
+			return FMemory::Memcmp(A->ChunkId.GetData(), B->ChunkId.GetData(), A->ChunkId.GetSize()) < 0;
+		}
+		return A->LoadOrderFactor < B->LoadOrderFactor;
+	}
+};
+
+struct FCookedPackage
+{
+	FPackageId GlobalPackageId;
+	FName PackageName;
+	FName SourcePackageName; // Source package name for redirected package
+	TArray<FShaderInfo*> Shaders;
+	TArray<FSHAHash> ShaderMapHashes;
+	uint64 UAssetSize = 0;
+	uint64 OptionalSegmentUAssetSize = 0;
+	uint64 UExpSize = 0;
+	uint64 TotalBulkDataSize = 0;
+	uint64 LoadOrder = MAX_uint64; // Ordered by dependencies
+	uint64 DiskLayoutOrder = MAX_uint64; // Final order after considering input order files
+	uint64 OverrideDiskLayoutOrder = MAX_uint64; // Final order after applying custom per-container ordering
+	FPackageStoreEntryResource PackageStoreEntry;
+	int32 PreOrderNumber = -1; // Used for sorting in load order
+	bool bPermanentMark = false; // Used for sorting in load order
+	bool bIsLocalized = false;
+};
+
+struct FLegacyCookedPackage
+	: public FCookedPackage
+{
+	FString FileName;
+	FString OptionalSegmentFileName;
+	FPackageStorePackage* OptimizedPackage = nullptr;
+	FPackageStorePackage* OptimizedOptionalSegmentPackage = nullptr;
+};
+
+enum class EContainerChunkType
+{
+	ShaderCodeLibrary,
+	ShaderCode,
+	PackageData,
+	BulkData,
+	OptionalBulkData,
+	MemoryMappedBulkData,
+	OptionalSegmentPackageData,
+	OptionalSegmentBulkData,
+	Invalid
+};
+
+struct FContainerTargetFile
+{
+	FContainerTargetSpec* ContainerTarget = nullptr;
+	FCookedPackage* Package = nullptr;
+	FString NormalizedSourcePath;
+	TOptional<FIoBuffer> SourceBuffer;
+	FString DestinationPath;
+	uint64 SourceSize = 0;
+	uint64 IdealOrder = 0;
+	FIoChunkId ChunkId;
+	FIoHash ChunkHash;
+	FBulkDataCookedIndex BulkDataCookedIndex;
+	TArray<uint8> PackageHeaderData;
+	EContainerChunkType ChunkType;
+	bool bForceUncompressed = false;
+
+	TArray<FFileRegion> FileRegions;
+};
+
+struct FContainerTargetSpec
+{
+	FIoContainerId ContainerId;
+	FIoContainerId OptionalSegmentContainerId;
+	FIoContainerHeader Header;
+	FIoContainerHeader OptionalSegmentHeader;
+	FName Name;
+	FGuid EncryptionKeyGuid;
+	FString OutputPath;
+	FString OptionalSegmentOutputPath;
+	FString StageLooseFileRootPath;
+	TSharedPtr<IIoStoreWriter> IoStoreWriter;
+	TSharedPtr<IIoStoreWriter> OptionalSegmentIoStoreWriter;
+	TArray<FContainerTargetFile> TargetFiles;
+	TArray<TUniquePtr<FIoStoreReader>> PatchSourceReaders;
+	EIoContainerFlags ContainerFlags = EIoContainerFlags::None;
+	TArray<FCookedPackage*> Packages;
+	TSet<FShaderInfo*> GlobalShaders;
+	TSet<FShaderInfo*> SharedShaders;
+	TSet<FShaderInfo*> UniqueShaders;
+	TSet<FShaderInfo*> InlineShaders;
+	bool bGenerateDiffPatch = false;
+};
+
+/**
+ * Finds where the extension starts in a path.
+ *
+ * We can't just find the left most '.' in the file name because the path might
+ * contain a bulkdata cooked index (@see FBulkDataCookedIndex) and we don't want
+ * to include the index part. We are attempting to match up the extensions with
+ * those found in the constructor of FCookedFileStatMap.
+ *
+ * For example take the following filenames and the extensions we want to find:
+ * <packagename>.m.ubulk   -> .m.ubulk
+ * <packagename>.o.ubulk   -> .o.ubulk
+ * <packagename>.001.ubulk ->.ubulk
+ *
+ * So we want to find the left most '.' unless that is followed by only numeric
+ * values in which case we have found a cooked index and we want to find the '.'
+ * after those values.
+ */
+int32 GetFullExtensionStartIndex(FStringView Path);
+
+FStringView GetFullExtension(FStringView Path);
+
+} // namespace UE::IoStore::Private
+
