[IoStore] Continue with refactoring of shader library processor.

- Adds `FShaderLibraryFileInfo` struct to parse and split filename into format and platform names.
- Adds `FSerializedShaderArchiveContainer` as subclass of `FSerializedShaderArchive` to store library buffer alongside header;
The base class doesn't actually store the shader code, only the entries  with indices and offsets.
- Moves several steps in `ProcessShaderLibraries()` into smaller functions to organize by processing steps.
- Adds `FShaderGroupIoChunk` and `FShaderGroupIoChunkMap` structs to replace tuples and give its members a name instead of generic `Get<0>()` etc.

#rnx

#Synced-CL 48500140

diff --git a/Engine/Source/Developer/IoStoreUtilities/Private/IoStoreShaderLibraryProcessor.cpp b/Engine/Source/Developer/IoStoreUtilities/Private/IoStoreShaderLibraryProcessor.cpp
--- a/Engine/Source/Developer/IoStoreUtilities/Private/IoStoreShaderLibraryProcessor.cpp
+++ b/Engine/Source/Developer/IoStoreUtilities/Private/IoStoreShaderLibraryProcessor.cpp
@@ -91,74 +91,150 @@ static bool LoadShaderAssetInfo(const FString& Filename, FCookedPackageStore* Pa
 	return true;
 }
 
-static bool ConvertToIoStoreShaderLibrary(
-	const TCHAR* FileName,
-	FIoChunkId InChunkId,
-	FCookedPackageStore* PackageStore,
-	FOodleDataCompression::ECompressor InShaderOodleCompressor,
-	FOodleDataCompression::ECompressionLevel InShaderOodleLevel,
-	TTuple<FIoChunkId, FIoBuffer>& OutLibraryIoChunk,
-	TArray<TTuple<FIoChunkId, FIoBuffer, uint32>>& OutCodeIoChunks,
-	TArray<TTuple<FSHAHash, TArray<FIoChunkId>>>& OutShaderMaps,
-	TArray<TTuple<FSHAHash, TSet<FName>>>& OutShaderMapAssetAssociations)
+struct FShaderLibraryFileInfo
 {
-	IOSTORE_CPU_SCOPE(ConvertShaderLibrary);
-	// ShaderArchive-<MyProject>_Chunk<N>-<ShaderFormat>-<ShaderPlatform>.ushaderbytecode
-	// E.g. "ShaderArchive-Lyra_Chunk0-PCD3D_SM6-PCD3D_SM6.ushaderbytecode
-	FStringView BaseFileNameView = FPathViews::GetBaseFilename(FileName);
-	int32 FormatStartIndex = -1;
-	if (!BaseFileNameView.FindLastChar('-', FormatStartIndex))
+	FString LibraryName;
+	FName FormatName;
+
+	// Initializes library and format name from the specified *.ushaderbytecode file.
+	bool ParseFromFilename(const TCHAR* Filename)
 	{
-		UE_LOG(LogIoStore, Error, TEXT("Invalid shader code library file name '%s'."), FileName);
-		return false;
+		// ShaderArchive-<MyProject>_Chunk<N>-<ShaderFormat>-<ShaderPlatform>.ushaderbytecode
+		// E.g. "ShaderArchive-Lyra_Chunk0-PCD3D_SM6-PCD3D_SM6.ushaderbytecode
+		FStringView BaseFilenameView = FPathViews::GetBaseFilename(Filename);
+		int32 FormatStartIndex = -1;
+		if (!BaseFilenameView.FindLastChar('-', FormatStartIndex))
+		{
+			UE_LOG(LogIoStore, Error, TEXT("Invalid shader code library filename '%s'."), Filename);
+			return false;
+		}
+		int32 LibraryNameStartIndex = -1;
+		if (!BaseFilenameView.FindChar('-', LibraryNameStartIndex) || FormatStartIndex - LibraryNameStartIndex <= 1)
+		{
+			UE_LOG(LogIoStore, Error, TEXT("Invalid shader code library filename '%s'."), Filename);
+			return false;
+		}
+		this->LibraryName = FString(BaseFilenameView.Mid(LibraryNameStartIndex + 1, FormatStartIndex - LibraryNameStartIndex - 1));
+		this->FormatName = FName(BaseFilenameView.RightChop(FormatStartIndex + 1));
+		return true;
 	}
-	int32 LibraryNameStartIndex = -1;
-	if (!BaseFileNameView.FindChar('-', LibraryNameStartIndex) || FormatStartIndex - LibraryNameStartIndex <= 1)
+
+	// Returns the ShaderAssetInfo JSON filename that belongs to the respective *.ushaderbytecode file.
+	FString GetAssetInfoFilename() const
 	{
-		UE_LOG(LogIoStore, Error, TEXT("Invalid shader code library file name '%s'."), FileName);
-		return false;
+		//@todo-lh - Should be merged with GetShaderAssetInfoFilename() in ShaderCodeLibrary.cpp
+		return FString::Printf(TEXT("ShaderAssetInfo-%s-%s.assetinfo.json"), *LibraryName, *FormatName.ToString());
 	}
-	FString LibraryName(BaseFileNameView.Mid(LibraryNameStartIndex + 1, FormatStartIndex - LibraryNameStartIndex - 1));
-	FName FormatName(BaseFileNameView.RightChop(FormatStartIndex + 1));
-
-	TRACE_CPUPROFILER_EVENT_SCOPE_TEXT(*LibraryName);
+};
 
+class FSerializedShaderArchiveContainer : public FSerializedShaderArchive
+{
 	FIoBuffer LibraryBuffer;
-	if (PackageStore && PackageStore->HasZenStoreClient() && InChunkId.IsValid())
+	int64 OffsetToShaderCode;
+
+public:
+	// Loads the shader library buffer either from package store or from local file.
+	bool Load(const TCHAR* InFilename, FIoChunkId InChunkId, FCookedPackageStore* PackageStore)
 	{
-		LibraryBuffer = PackageStore->ReadChunk(InChunkId).ConsumeValueOrDie();
+		if (PackageStore && PackageStore->HasZenStoreClient() && InChunkId.IsValid())
+		{
+			LibraryBuffer = PackageStore->ReadChunk(InChunkId).ConsumeValueOrDie();
+		}
+		else
+		{
+			TUniquePtr<FArchive> LibraryAr(IFileManager::Get().CreateFileReader(InFilename));
+			if (!LibraryAr)
+			{
+				UE_LOG(LogIoStore, Error, TEXT("Missing shader code library file '%s'."), InFilename);
+				return false;
+			}
+
+			LibraryBuffer = FIoBuffer(LibraryAr->TotalSize());
+			LibraryAr->Serialize(LibraryBuffer.GetData(), LibraryBuffer.GetSize());
+		}
+
+		// Serialize shader archive header and store offset to start of shader code payload
+		FLargeMemoryReader LibraryAr(LibraryBuffer.GetData(), LibraryBuffer.GetSize());
+		uint32 Version;
+		LibraryAr << Version;
+
+		LibraryAr << *this;
+		OffsetToShaderCode = LibraryAr.Tell();
+
+		return true;
 	}
-	else
+
+	// Loads and decompresses the specified shader into the destination buffer and returns the advanced pointer of the destination after the new copy.
+	uint8* LoadAndDecompressShader(uint8* Destination, int32 ShaderIndex) const
 	{
-		TUniquePtr<FArchive> LibraryAr(IFileManager::Get().CreateFileReader(FileName));
-		if (!LibraryAr)
+		const FShaderCodeEntry& IndividuallyCompressedShader = ShaderEntries[ShaderIndex];
 		{
-			UE_LOG(LogIoStore, Error, TEXT("Missing shader code library file '%s'."), FileName);
-			return false;
+			// small shaders might be stored without compression, handle them here
+			if (IndividuallyCompressedShader.Size == IndividuallyCompressedShader.UncompressedSize)
+			{
+				FMemory::Memcpy(Destination, GetShaderCodeDataStart() + IndividuallyCompressedShader.Offset, IndividuallyCompressedShader.Size);
+			}
+			else
+			{
+				// This function will crash if decompression fails.
+				ShaderCodeArchive::DecompressShaderWithOodle(Destination, IndividuallyCompressedShader.UncompressedSize, GetShaderCodeDataStart() + IndividuallyCompressedShader.Offset, IndividuallyCompressedShader.Size);
+			}
 		}
+		return Destination + IndividuallyCompressedShader.UncompressedSize;
+	}
 
-		LibraryBuffer = FIoBuffer(LibraryAr->TotalSize());
-		LibraryAr->Serialize(LibraryBuffer.GetData(), LibraryBuffer.GetSize());
+	// Returns a pointer to the start of shader code payload to read from.
+	const uint8* GetShaderCodeDataStart() const
+	{
+		return LibraryBuffer.GetData() + OffsetToShaderCode;
 	}
+};
 
-	FLargeMemoryReader LibraryAr(LibraryBuffer.GetData(), LibraryBuffer.GetSize());
-	uint32 Version;
-	LibraryAr << Version;
+// Shader map hash with its IoChunk Ids of all its shader groups
+struct FShaderGroupIoChunkMap
+{
+	FSHAHash ShaderMapHash;
+	TArray<FIoChunkId> ShaderGroupChunkIds;
+};
 
-	FSerializedShaderArchive SerializedShaders;
-	LibraryAr << SerializedShaders;
-	int64 OffsetToShaderCode = LibraryAr.Tell();
+static bool ConvertToIoStoreShaderLibrary(
+	const TCHAR* InFilename,
+	FIoChunkId InChunkId,
+	FCookedPackageStore* PackageStore,
+	FOodleDataCompression::ECompressor InShaderOodleCompressor,
+	FOodleDataCompression::ECompressionLevel InShaderOodleLevel,
+	TTuple<FIoChunkId, FIoBuffer>& OutLibraryIoChunk,
+	TArray<FShaderGroupIoChunk>& OutCodeIoChunks,
+	TArray<FShaderGroupIoChunkMap>& OutShaderMaps,
+	TArray<TTuple<FSHAHash, FShaderMapAssetPaths>>& OutShaderMapAssetAssociations)
+{
+	IOSTORE_CPU_SCOPE(ConvertShaderLibrary);
 
-	FString AssetInfoFileName = FPaths::GetPath(FileName) / FString::Printf(TEXT("ShaderAssetInfo-%s-%s.assetinfo.json"), *LibraryName, *FormatName.ToString());
+	FShaderLibraryFileInfo LibraryFileInfo;
+	if (!LibraryFileInfo.ParseFromFilename(InFilename))
+	{
+		return false;
+	}
+
+	FSerializedShaderArchiveContainer SerializedShaders;
+	if (!SerializedShaders.Load(InFilename, InChunkId, PackageStore))
+	{
+		return false;
+	}
+
+	const FString AssetInfoFilename = FPaths::GetPath(InFilename) / LibraryFileInfo.GetAssetInfoFilename();
 	TMap<FSHAHash, FShaderMapAssetPaths> ShaderCodeToAssets;
-	if (!LoadShaderAssetInfo(AssetInfoFileName, PackageStore && InChunkId.IsValid() ? PackageStore : nullptr, ShaderCodeToAssets))
+	if (!LoadShaderAssetInfo(AssetInfoFilename, PackageStore && InChunkId.IsValid() ? PackageStore : nullptr, ShaderCodeToAssets))
 	{
-		UE_LOG(LogIoStore, Error, TEXT("Failed loading asset asset info file '%s'"), *AssetInfoFileName);
+		UE_LOG(LogIoStore, Error, TEXT("Failed loading asset asset info file '%s'"), *AssetInfoFilename);
 		return false;
 	}
 
+	TRACE_CPUPROFILER_EVENT_SCOPE_TEXT(*LibraryFileInfo.LibraryName);
+
+	// Create IoStore code archived for all serialized shaders. This will organize the shaders into groups.
 	FIoStoreShaderCodeArchiveHeader IoStoreLibraryHeader;
-	FIoStoreShaderCodeArchive::CreateIoStoreShaderCodeArchiveHeader(FName(FormatName), SerializedShaders, IoStoreLibraryHeader);
+	FIoStoreShaderCodeArchive::CreateIoStoreShaderCodeArchiveHeader(LibraryFileInfo.FormatName, SerializedShaders, IoStoreLibraryHeader);
 	checkf(SerializedShaders.ShaderEntries.Num() == IoStoreLibraryHeader.ShaderEntries.Num(), TEXT("IoStore header has different number of shaders (%d) than the original library (%d)"),
 		IoStoreLibraryHeader.ShaderEntries.Num(), SerializedShaders.ShaderEntries.Num());
 	checkf(SerializedShaders.ShaderMapEntries.Num() == IoStoreLibraryHeader.ShaderMapEntries.Num(), TEXT("IoStore header has different number of shadermaps (%d) than the original library (%d)"),
@@ -170,7 +246,7 @@ static bool ConvertToIoStoreShaderLibrary(
 	OutCodeIoChunks.SetNum(GroupCount);
 	FCriticalSection DiskArchiveAccess;
 	ParallelFor(OutCodeIoChunks.Num(),
-		[&OutCodeIoChunks, &IoStoreLibraryHeader, &DiskArchiveAccess, &LibraryBuffer, &SerializedShaders, &OffsetToShaderCode, InShaderOodleCompressor, InShaderOodleLevel](int32 GroupIndex)
+		[&OutCodeIoChunks, &IoStoreLibraryHeader, &DiskArchiveAccess, &SerializedShaders, InShaderOodleCompressor, InShaderOodleLevel](int32 GroupIndex)
 		{
 			FIoStoreShaderGroupEntry& Group = IoStoreLibraryHeader.ShaderGroupEntries[GroupIndex];
 			uint8* UncompressedGroupMemory = reinterpret_cast<uint8*>(FMemory::Malloc(Group.UncompressedSize));
@@ -181,7 +257,7 @@ static bool ConvertToIoStoreShaderLibrary(
 			// Not worth to special-case for group size of 1 (by converting to copy) - such groups are very fast to decompress and recompress.
 			for (uint32 IdxShaderInGroup = 0; IdxShaderInGroup < Group.NumShaders; ++IdxShaderInGroup)
 			{
-				int32 ShaderIndex = IoStoreLibraryHeader.ShaderIndices[Group.ShaderIndicesOffset + IdxShaderInGroup];
+				const int32 ShaderIndex = IoStoreLibraryHeader.ShaderIndices[Group.ShaderIndicesOffset + IdxShaderInGroup];
 				SmallestShaderInGroupByOrdinal = FMath::Min(SmallestShaderInGroupByOrdinal, ShaderIndex);
 				const FIoStoreShaderCodeEntry& Shader = IoStoreLibraryHeader.ShaderEntries[ShaderIndex];
 
@@ -191,21 +267,7 @@ static bool ConvertToIoStoreShaderLibrary(
 				);
 
 				// load and decompress the shader at the desired offset
-				const FShaderCodeEntry& IndividuallyCompressedShader = SerializedShaders.ShaderEntries[ShaderIndex];
-				{
-					// small shaders might be stored without compression, handle them here
-					if (IndividuallyCompressedShader.Size == IndividuallyCompressedShader.UncompressedSize)
-					{
-						FMemory::Memcpy(ShaderStart, LibraryBuffer.GetData() + OffsetToShaderCode + IndividuallyCompressedShader.Offset, IndividuallyCompressedShader.Size);
-					}
-					else
-					{
-						// This function will crash if decompression fails.
-						ShaderCodeArchive::DecompressShaderWithOodle(ShaderStart, IndividuallyCompressedShader.UncompressedSize, LibraryBuffer.GetData() + OffsetToShaderCode + IndividuallyCompressedShader.Offset, IndividuallyCompressedShader.Size);
-					}
-				}
-
-				ShaderStart += IndividuallyCompressedShader.UncompressedSize;
+				ShaderStart = SerializedShaders.LoadAndDecompressShader(ShaderStart, ShaderIndex);
 			}
 
 			checkf((reinterpret_cast<uint64>(ShaderStart) - reinterpret_cast<uint64>(UncompressedGroupMemory)) == Group.UncompressedSize,
@@ -222,35 +284,40 @@ static bool ConvertToIoStoreShaderLibrary(
 			bool bCompressed = ShaderCodeArchive::CompressShaderWithOodle(CompressedShaderGroupMemory, CompressedGroupSize, UncompressedGroupMemory, Group.UncompressedSize, InShaderOodleCompressor, InShaderOodleLevel);
 			checkf(bCompressed, TEXT("We could not compress the shader group after providing an estimated memory."));
 
-			TTuple<FIoChunkId, FIoBuffer, uint32>& OutCodeIoChunk = OutCodeIoChunks[GroupIndex];
-			OutCodeIoChunk.Get<0>() = IoStoreLibraryHeader.ShaderGroupIoHashes[GroupIndex];
-			// This value is the load order factor for the group (the smaller, the more likely), that IoStore will use to sort the chunks.
-			// We calculate it as the smallest shader index in the group, which shouldn't be bad. Arguably a better way would be to get the lowest-numbered shadermap that references _any_ shader in the group, 
-			// but it's much, much slower to calculate and the resulting order would be likely the same/similar most of the time
-			checkf(SmallestShaderInGroupByOrdinal >= 0 && SmallestShaderInGroupByOrdinal < IoStoreLibraryHeader.ShaderEntries.Num(), TEXT("SmallestShaderInGroupByOrdinal has an invalid value of %d (not within [0, %d) range as expected)"),
-				SmallestShaderInGroupByOrdinal, IoStoreLibraryHeader.ShaderEntries.Num());
-			OutCodeIoChunk.Get<2>() = static_cast<uint32>(SmallestShaderInGroupByOrdinal);
-
-			if (CompressedGroupSize < Group.UncompressedSize)
-			{
-				OutCodeIoChunk.Get<1>() = FIoBuffer(CompressedGroupSize);
-				FMemory::Memcpy(OutCodeIoChunk.Get<1>().GetData(), CompressedShaderGroupMemory, CompressedGroupSize);
-				Group.CompressedSize = CompressedGroupSize;
-			}
-			else
+			// Write out IoChunk for current shader group
 			{
-				// store uncompressed (unlikely, but happens for a 200-byte sized shader that happens to get its own group)
-				OutCodeIoChunk.Get<1>() = FIoBuffer(Group.UncompressedSize);
-				FMemory::Memcpy(OutCodeIoChunk.Get<1>().GetData(), UncompressedGroupMemory, Group.UncompressedSize);
-				Group.CompressedSize = Group.UncompressedSize;
+				FShaderGroupIoChunk& OutCodeIoChunk = OutCodeIoChunks[GroupIndex];
+				OutCodeIoChunk.ChunkId = IoStoreLibraryHeader.ShaderGroupIoHashes[GroupIndex];
+				// This value is the load order factor for the group (the smaller, the more likely), that IoStore will use to sort the chunks.
+				// We calculate it as the smallest shader index in the group, which shouldn't be bad. Arguably a better way would be to get the lowest-numbered shadermap that references _any_ shader in the group, 
+				// but it's much, much slower to calculate and the resulting order would be likely the same/similar most of the time
+				checkf(SmallestShaderInGroupByOrdinal >= 0 && SmallestShaderInGroupByOrdinal < IoStoreLibraryHeader.ShaderEntries.Num(), TEXT("SmallestShaderInGroupByOrdinal has an invalid value of %d (not within [0, %d) range as expected)"),
+					SmallestShaderInGroupByOrdinal, IoStoreLibraryHeader.ShaderEntries.Num());
+				OutCodeIoChunk.LoadOrderFactor = static_cast<uint32>(SmallestShaderInGroupByOrdinal);
+
+				// Copy shader group payload into IoBuffer
+				if (CompressedGroupSize < Group.UncompressedSize)
+				{
+					OutCodeIoChunk.CodeIoBuffer = FIoBuffer(CompressedGroupSize);
+					FMemory::Memcpy(OutCodeIoChunk.CodeIoBuffer.GetData(), CompressedShaderGroupMemory, CompressedGroupSize);
+					Group.CompressedSize = CompressedGroupSize;
+				}
+				else
+				{
+					// store uncompressed (unlikely, but happens for a 200-byte sized shader that happens to get its own group)
+					OutCodeIoChunk.CodeIoBuffer = FIoBuffer(Group.UncompressedSize);
+					FMemory::Memcpy(OutCodeIoChunk.CodeIoBuffer.GetData(), UncompressedGroupMemory, Group.UncompressedSize);
+					Group.CompressedSize = Group.UncompressedSize;
+				}
 			}
+
 			FMemory::Free(UncompressedGroupMemory);
 			FMemory::Free(CompressedShaderGroupMemory);
 		},
 		EParallelForFlags::Unbalanced
 	);
 
-	// calculate and log stats
+	// Calculate and log stats
 	int64 TotalUncompressedSizeViaGroups = 0;	// calculate total uncompressed size twice for a sanity check
 	int64 TotalUncompressedSizeViaShaders = 0;
 	int64 TotalIndividuallyCompressedSize = 0;
@@ -272,44 +339,54 @@ static bool ConvertToIoStoreShaderLibrary(
 		}
 	}
 
-	checkf(TotalUncompressedSizeViaGroups == TotalUncompressedSizeViaShaders, TEXT("Sanity check failure: total uncompressed shader size differs if calculated via shader groups (%lld) or individual shaders (%lld)"),
+	checkf(TotalUncompressedSizeViaGroups == TotalUncompressedSizeViaShaders,
+		TEXT("Sanity check failure: total uncompressed shader size differs if calculated via shader groups (%lld) or individual shaders (%lld)"),
 		TotalUncompressedSizeViaGroups, TotalUncompressedSizeViaShaders
 	);
 
-	UE_LOG(LogIoStore, Display, TEXT("%s(%s): Recompressed %d shaders as %d groups. Library size changed from %lld KB (%.2f : 1 ratio) to %lld KB (%.2f : 1 ratio), %.2f%% of previous."),
-		*LibraryName, *FormatName.ToString(),
-		SerializedShaders.ShaderEntries.Num(), IoStoreLibraryHeader.ShaderGroupEntries.Num(),
-		TotalIndividuallyCompressedSize / 1024, static_cast<double>(TotalUncompressedSizeViaShaders) / static_cast<double>(TotalIndividuallyCompressedSize),
-		TotalGroupCompressedSize / 1024, static_cast<double>(TotalUncompressedSizeViaShaders) / static_cast<double>(TotalGroupCompressedSize),
+	UE_LOG(LogIoStore, Display,
+		TEXT("%s(%s): Recompressed %d shaders as %d groups. Library size changed from %lld KB (%.2f : 1 ratio) to %lld KB (%.2f : 1 ratio), %.2f%% of previous."),
+		*LibraryFileInfo.LibraryName,
+		*LibraryFileInfo.FormatName.ToString(),
+		SerializedShaders.ShaderEntries.Num(),
+		IoStoreLibraryHeader.ShaderGroupEntries.Num(),
+		TotalIndividuallyCompressedSize / 1024,
+		static_cast<double>(TotalUncompressedSizeViaShaders) / static_cast<double>(TotalIndividuallyCompressedSize),
+		TotalGroupCompressedSize / 1024,
+		static_cast<double>(TotalUncompressedSizeViaShaders) / static_cast<double>(TotalGroupCompressedSize),
 		100.0 * static_cast<double>(TotalGroupCompressedSize) / static_cast<double>(TotalIndividuallyCompressedSize)
 	);
 
 	FLargeMemoryWriter IoStoreLibraryAr(0, true);
 	FIoStoreShaderCodeArchive::SaveIoStoreShaderCodeArchive(IoStoreLibraryHeader, IoStoreLibraryAr);
-	OutLibraryIoChunk.Key = FIoStoreShaderCodeArchive::GetShaderCodeArchiveChunkId(LibraryName, FormatName);
+	OutLibraryIoChunk.Key = FIoStoreShaderCodeArchive::GetShaderCodeArchiveChunkId(LibraryFileInfo.LibraryName, LibraryFileInfo.FormatName);
 	int64 TotalSize = IoStoreLibraryAr.TotalSize();
 	OutLibraryIoChunk.Value = FIoBuffer(FIoBuffer::AssumeOwnership, IoStoreLibraryAr.ReleaseOwnership(), TotalSize);
 
-	int32 ShaderMapCount = IoStoreLibraryHeader.ShaderMapHashes.Num();
-	OutShaderMaps.SetNum(ShaderMapCount);
-	ParallelFor(OutShaderMaps.Num(),
-		[&OutShaderMaps, &IoStoreLibraryHeader](int32 ShaderMapIndex)
-		{
-			TTuple<FSHAHash, TArray<FIoChunkId>>& OutShaderMap = OutShaderMaps[ShaderMapIndex];
-			OutShaderMap.Key = IoStoreLibraryHeader.ShaderMapHashes[ShaderMapIndex];
-			const FIoStoreShaderMapEntry& ShaderMapEntry = IoStoreLibraryHeader.ShaderMapEntries[ShaderMapIndex];
-			int32 LookupIndexEnd = ShaderMapEntry.ShaderIndicesOffset + ShaderMapEntry.NumShaders;
-			OutShaderMap.Value.Reserve(ShaderMapEntry.NumShaders);	// worst-case, 1 group == 1 shader
-			for (int32 ShaderLookupIndex = ShaderMapEntry.ShaderIndicesOffset; ShaderLookupIndex < LookupIndexEnd; ++ShaderLookupIndex)
+	// Return shadermaps: own hash and array of shader group IoChunk Ids
+	{
+		int32 ShaderMapCount = IoStoreLibraryHeader.ShaderMapHashes.Num();
+		OutShaderMaps.SetNum(ShaderMapCount);
+		ParallelFor(OutShaderMaps.Num(),
+			[&OutShaderMaps, &IoStoreLibraryHeader](int32 ShaderMapIndex)
 			{
-				int32 ShaderIndex = IoStoreLibraryHeader.ShaderIndices[ShaderLookupIndex];
-				int32 GroupIndex = IoStoreLibraryHeader.ShaderEntries[ShaderIndex].ShaderGroupIndex;
-				OutShaderMap.Value.AddUnique(IoStoreLibraryHeader.ShaderGroupIoHashes[GroupIndex]);
-			}
-		},
-		EParallelForFlags::Unbalanced
-	);
+				FShaderGroupIoChunkMap& OutShaderMap = OutShaderMaps[ShaderMapIndex];
+				OutShaderMap.ShaderMapHash = IoStoreLibraryHeader.ShaderMapHashes[ShaderMapIndex];
+				const FIoStoreShaderMapEntry& ShaderMapEntry = IoStoreLibraryHeader.ShaderMapEntries[ShaderMapIndex];
+				int32 LookupIndexEnd = ShaderMapEntry.ShaderIndicesOffset + ShaderMapEntry.NumShaders;
+				OutShaderMap.ShaderGroupChunkIds.Reserve(ShaderMapEntry.NumShaders);	// worst-case, 1 group == 1 shader
+				for (int32 ShaderLookupIndex = ShaderMapEntry.ShaderIndicesOffset; ShaderLookupIndex < LookupIndexEnd; ++ShaderLookupIndex)
+				{
+					int32 ShaderIndex = IoStoreLibraryHeader.ShaderIndices[ShaderLookupIndex];
+					int32 GroupIndex = IoStoreLibraryHeader.ShaderEntries[ShaderIndex].ShaderGroupIndex;
+					OutShaderMap.ShaderGroupChunkIds.AddUnique(IoStoreLibraryHeader.ShaderGroupIoHashes[GroupIndex]);
+				}
+			},
+			EParallelForFlags::Unbalanced
+		);
+	}
 
+	// Return associations between shader maps and all their asset paths
 	OutShaderMapAssetAssociations.Reserve(ShaderCodeToAssets.Num());
 	for (auto& KV : ShaderCodeToAssets)
 	{
@@ -323,12 +400,7 @@ void FIoStoreShaderLibraryProcessor::ProcessShaderLibraries(const FIoStoreArgume
 {
 	IOSTORE_CPU_SCOPE(ProcessShaderLibraries);
 
-	double LibraryStart = FPlatformTime::Seconds();
-
-	TMap<FIoChunkId, FShaderInfo*> ChunkIdToShaderInfoMap;
-	TMap<FSHAHash, TArray<FIoChunkId>> ShaderChunkIdsByShaderMapHash;
-	TMap<FName, TSet<FSHAHash>> PackageNameToShaderMaps;
-	TMap<FContainerTargetSpec*, TSet<FShaderInfo*>> AllContainerShaderLibraryShadersMap;
+	const double LibraryStart = FPlatformTime::Seconds();
 
 	{
 		IOSTORE_CPU_SCOPE(ConvertShaderLibraries);
@@ -338,225 +410,253 @@ void FIoStoreShaderLibraryProcessor::ProcessShaderLibraries(const FIoStoreArgume
 			{
 				if (TargetFile.ChunkType == EContainerChunkType::ShaderCodeLibrary)
 				{
-					TSet<FShaderInfo*>& ContainerShaderLibraryShaders = AllContainerShaderLibraryShadersMap.FindOrAdd(ContainerTarget);
-
-					TArray<TTuple<FSHAHash, TArray<FIoChunkId>>> ShaderMaps;
-					TTuple<FIoChunkId, FIoBuffer> LibraryChunk;
-					TArray<TTuple<FIoChunkId, FIoBuffer, uint32>> CodeChunks;
-					TArray<TTuple<FSHAHash, TSet<FName>>> ShaderMapAssetAssociations;
-					if (!ConvertToIoStoreShaderLibrary(*TargetFile.NormalizedSourcePath, TargetFile.ChunkId, Arguments.PackageStore.Get(), Arguments.ShaderOodleCompressor, Arguments.ShaderOodleLevel, LibraryChunk, CodeChunks, ShaderMaps, ShaderMapAssetAssociations))
+					if (!ProcessShaderLibraryFromTargetFile(Arguments, TargetFile, Output))
 					{
 						UE_LOG(LogIoStore, Warning, TEXT("Failed converting shader library '%s'"), *TargetFile.NormalizedSourcePath);
-						continue;
 					}
-					TargetFile.ChunkId = LibraryChunk.Key;
-					TargetFile.SourceBuffer.Emplace(LibraryChunk.Value);
-					TargetFile.SourceSize = LibraryChunk.Value.GetSize();
+				}
+			}
+		}
+	}
 
-					const bool bIsGlobalShaderLibrary = FPaths::GetCleanFilename(TargetFile.NormalizedSourcePath).StartsWith(TEXT("ShaderArchive-Global-"));
-					const FShaderInfo::EShaderType ShaderType = bIsGlobalShaderLibrary ? FShaderInfo::Global : FShaderInfo::Normal;
-					for (const TTuple<FIoChunkId, FIoBuffer, uint32>& CodeChunk : CodeChunks)
-					{
-						const FIoChunkId& ShaderChunkId = CodeChunk.Get<0>();
-						FShaderInfo* ShaderInfo = ChunkIdToShaderInfoMap.FindRef(ShaderChunkId);
-						if (!ShaderInfo)
-						{
-							ShaderInfo = new FShaderInfo();
-							ShaderInfo->ChunkId = ShaderChunkId;
-							ShaderInfo->CodeIoBuffer = CodeChunk.Get<1>();
-							ShaderInfo->LoadOrderFactor = CodeChunk.Get<2>();
-							Output.ShaderInfos.Add(ShaderInfo);
-							ChunkIdToShaderInfoMap.Add(ShaderChunkId, ShaderInfo);
-						}
-						else
-						{
-							// first, make sure that the code is exactly the same
-							if (ShaderInfo->CodeIoBuffer != CodeChunk.Get<1>())
-							{
-								UE_LOG(LogIoStore, Error, TEXT("Collision of two shader code chunks, same Id (%s), different code. Packaged game will likely crash, not being able to decompress the shaders."), *LexToString(ShaderChunkId));
-							}
-
-							// If we already exist, then we have two separate LoadOrderFactors,
-							// which one we got first affects build determinism. Take the lower.
-							if (CodeChunk.Get<2>() < ShaderInfo->LoadOrderFactor)
-							{
-								ShaderInfo->LoadOrderFactor = CodeChunk.Get<2>();
-							}
-						}
-
-
-						const FShaderInfo::EShaderType* CurrentShaderTypeInContainer = ShaderInfo->TypeInContainer.Find(ContainerTarget);
-						if (!CurrentShaderTypeInContainer || *CurrentShaderTypeInContainer != FShaderInfo::Global)
-						{
-							// If a shader is both global and shared consider it to be global
-							ShaderInfo->TypeInContainer.Add(ContainerTarget, ShaderType);
-						}
-
-						ContainerShaderLibraryShaders.Add(ShaderInfo);
-					}
-					for (const TTuple<FSHAHash, TSet<FName>>& ShaderMapAssetAssociation : ShaderMapAssetAssociations)
-					{
-						for (const FName& PackageName : ShaderMapAssetAssociation.Value)
-						{
-							TSet<FSHAHash>& PackageShaderMaps = PackageNameToShaderMaps.FindOrAdd(PackageName);
-							PackageShaderMaps.Add(ShaderMapAssetAssociation.Key);
-						}
-					}
+	{
+		IOSTORE_CPU_SCOPE(UpdatePackageStoreShaders);
+		for (FContainerTargetSpec* ContainerTarget : ContainerTargets)
+		{
+			for (FCookedPackage* Package : ContainerTarget->Packages)
+			{
+				UpdatePackageStoreShaders(Package, ContainerTarget->Name);
+			}
+		}
+	}
 
-					for (TTuple<FSHAHash, TArray<FIoChunkId>>& ShaderMap : ShaderMaps)
-					{
-						TArray<FIoChunkId>& ShaderMapChunkIds = ShaderChunkIdsByShaderMapHash.FindOrAdd(ShaderMap.Key);
-						ShaderMapChunkIds.Append(MoveTemp(ShaderMap.Value));
-					}
-				} // end if containerchunktype shadercodelibrary
-			} // end foreach targetfile
-		} // end foreach container
+	{
+		IOSTORE_CPU_SCOPE(AddShaderCodeToContainer);
+		for (FContainerTargetSpec* ContainerTarget : ContainerTargets)
+		{
+			AddAllReferencedShadersToContainerTarget(ContainerTarget, Output);
+		}
+	}
+
+	const double LibraryEnd = FPlatformTime::Seconds();
+	UE_LOG(LogIoStore, Display, TEXT("Shaders processed in %.2f seconds"), LibraryEnd - LibraryStart);
+}
+
+bool FIoStoreShaderLibraryProcessor::ProcessShaderLibraryFromTargetFile(const FIoStoreArguments& Arguments, FContainerTargetFile& TargetFile, FIoStoreShaderLibraryOutput& Output)
+{
+	FContainerTargetSpec* ContainerTarget = TargetFile.ContainerTarget;
+	checkf(ContainerTarget, TEXT("Target file for shader library '%s' is missing its parent node (ContainerTarget is null)"), *TargetFile.NormalizedSourcePath);
+
+	TSet<FShaderInfo*>& ContainerShaderLibraryShaders = ContainerTargetToShaderLibrary.FindOrAdd(ContainerTarget);
+
+	TArray<FShaderGroupIoChunkMap> ShaderMaps;
+	TTuple<FIoChunkId, FIoBuffer> LibraryChunk;
+	TArray<FShaderGroupIoChunk> CodeChunks;
+	TArray<TTuple<FSHAHash, FShaderMapAssetPaths>> ShaderMapAssetAssociations;
+	if (!ConvertToIoStoreShaderLibrary(*TargetFile.NormalizedSourcePath, TargetFile.ChunkId, Arguments.PackageStore.Get(), Arguments.ShaderOodleCompressor, Arguments.ShaderOodleLevel, LibraryChunk, CodeChunks, ShaderMaps, ShaderMapAssetAssociations))
+	{
+		return false;
+	}
+	TargetFile.ChunkId = LibraryChunk.Key;
+	TargetFile.SourceBuffer.Emplace(LibraryChunk.Value);
+	TargetFile.SourceSize = LibraryChunk.Value.GetSize();
+
+	const bool bIsGlobalShaderLibrary = FPaths::GetCleanFilename(TargetFile.NormalizedSourcePath).StartsWith(TEXT("ShaderArchive-Global-"));
+	const FShaderInfo::EShaderType ShaderType = bIsGlobalShaderLibrary ? FShaderInfo::Global : FShaderInfo::Normal;
+	for (const FShaderGroupIoChunk& CodeChunk : CodeChunks)
+	{
+		const FIoChunkId& ShaderChunkId = CodeChunk.ChunkId;
+		FShaderInfo* ShaderInfo = ChunkIdToShaderInfoMap.FindRef(ShaderChunkId);
+		if (!ShaderInfo)
+		{
+			ShaderInfo = new FShaderInfo();
+			ShaderInfo->ShaderGroupChunk = CodeChunk;
+			Output.ShaderInfos.Add(ShaderInfo);
+			ChunkIdToShaderInfoMap.Add(ShaderChunkId, ShaderInfo);
+		}
+		else
+		{
+			// first, make sure that the code is exactly the same
+			if (ShaderInfo->ShaderGroupChunk.CodeIoBuffer != CodeChunk.CodeIoBuffer)
+			{
+				UE_LOG(LogIoStore, Error, TEXT("Collision of two shader code chunks, same Id (%s), different code. Packaged game will likely crash, not being able to decompress the shaders."), *LexToString(ShaderChunkId));
+			}
+
+			// If we already exist, then we have two separate LoadOrderFactors,
+			// which one we got first affects build determinism. Take the lower.
+			if (ShaderInfo->ShaderGroupChunk.LoadOrderFactor > CodeChunk.LoadOrderFactor)
+			{
+				ShaderInfo->ShaderGroupChunk.LoadOrderFactor = CodeChunk.LoadOrderFactor;
+			}
+		}
+
+		const FShaderInfo::EShaderType* CurrentShaderTypeInContainer = ShaderInfo->TypeInContainer.Find(ContainerTarget);
+		if (!CurrentShaderTypeInContainer || *CurrentShaderTypeInContainer != FShaderInfo::Global)
+		{
+			// If a shader is both global and shared consider it to be global
+			ShaderInfo->TypeInContainer.Add(ContainerTarget, ShaderType);
+		}
+
+		ContainerShaderLibraryShaders.Add(ShaderInfo);
+	}
+
+	for (const TTuple<FSHAHash, FShaderMapAssetPaths>& ShaderMapAssetAssociation : ShaderMapAssetAssociations)
+	{
+		for (const FName& PackageName : ShaderMapAssetAssociation.Value)
+		{
+			TSet<FSHAHash>& PackageShaderMaps = PackageNameToShaderMaps.FindOrAdd(PackageName);
+			PackageShaderMaps.Add(ShaderMapAssetAssociation.Key);
+		}
+	}
+
+	for (FShaderGroupIoChunkMap& ShaderMap : ShaderMaps)
+	{
+		TArray<FIoChunkId>& ShaderMapChunkIds = ShaderMapToGroupChunkIds.FindOrAdd(ShaderMap.ShaderMapHash);
+		ShaderMapChunkIds.Append(MoveTemp(ShaderMap.ShaderGroupChunkIds));
 	}
 
+	return true;
+}
+
+// Adds a new entry of type EContainerChunkType::ShaderCode to the specified container target
+static void AddShaderCodeToContainerTarget(FContainerTargetSpec* ContainerTarget, FShaderInfo& ShaderInfo)
+{
+	check(ContainerTarget);
+	FContainerTargetFile& ShaderTargetFile = ContainerTarget->TargetFiles.AddDefaulted_GetRef();
+	ShaderTargetFile.ContainerTarget = ContainerTarget;
+	ShaderTargetFile.ChunkId = ShaderInfo.ShaderGroupChunk.ChunkId;
+	ShaderTargetFile.ChunkType = EContainerChunkType::ShaderCode;
+	ShaderTargetFile.bForceUncompressed = true;
+	ShaderTargetFile.SourceBuffer.Emplace(ShaderInfo.ShaderGroupChunk.CodeIoBuffer);
+	ShaderTargetFile.SourceSize = ShaderInfo.ShaderGroupChunk.CodeIoBuffer.DataSize();
+}
+
+void FIoStoreShaderLibraryProcessor::UpdatePackageStoreShaders(FCookedPackage* Package, const FName& ContainerTargetName)
+{
 	// 1. Update ShaderInfos with which packages we reference.
 	// 2. Add to packages which shaders we use.
 	// 3. Add to PackageStore what shaders we use.
+	const TSet<FSHAHash>* FindShaderMapHashes = PackageNameToShaderMaps.Find(Package->PackageName);
+	if (!FindShaderMapHashes)
 	{
-		IOSTORE_CPU_SCOPE(UpdatePackageStoreShaders);
-		for (FContainerTargetSpec* ContainerTarget : ContainerTargets)
+		// Early exit if package doesn't contain any shader maps
+		return;
+	}
+
+	for (const FSHAHash& ShaderMapHash : *FindShaderMapHashes)
+	{
+		const TArray<FIoChunkId>* FindChunkIds = ShaderMapToGroupChunkIds.Find(ShaderMapHash);
+		if (!FindChunkIds)
 		{
-			for (FCookedPackage* Package : ContainerTarget->Packages)
+			UE_LOG(LogIoStore, Warning, TEXT("Package '%s' in '%s' referencing missing shader map '%s'"), *Package->PackageName.ToString(), *ContainerTargetName.ToString(), *ShaderMapHash.ToString());
+			continue;
+		}
+		Package->ShaderMapHashes.Add(ShaderMapHash);
+
+		for (const FIoChunkId& ShaderChunkId : *FindChunkIds)
+		{
+			FShaderInfo* ShaderInfo = ChunkIdToShaderInfoMap.FindRef(ShaderChunkId);
+			if (!ShaderInfo)
 			{
-				TSet<FSHAHash>* FindShaderMapHashes = PackageNameToShaderMaps.Find(Package->PackageName);
-				if (FindShaderMapHashes)
-				{
-					for (const FSHAHash& ShaderMapHash : *FindShaderMapHashes)
-					{
-						const TArray<FIoChunkId>* FindChunkIds = ShaderChunkIdsByShaderMapHash.Find(ShaderMapHash);
-						if (!FindChunkIds)
-						{
-							UE_LOG(LogIoStore, Warning, TEXT("Package '%s' in '%s' referencing missing shader map '%s'"), *Package->PackageName.ToString(), *ContainerTarget->Name.ToString(), *ShaderMapHash.ToString());
-							continue;
-						}
-						Package->ShaderMapHashes.Add(ShaderMapHash);
-						for (const FIoChunkId& ShaderChunkId : *FindChunkIds)
-						{
-							FShaderInfo* ShaderInfo = ChunkIdToShaderInfoMap.FindRef(ShaderChunkId);
-							if (!ShaderInfo)
-							{
-								UE_LOG(LogIoStore, Warning, TEXT("Package '%s' in '%s' referencing missing shader with chunk id '%s'"), *Package->PackageName.ToString(), *ContainerTarget->Name.ToString(), *BytesToHex(ShaderChunkId.GetData(), ShaderChunkId.GetSize()));
-								continue;
-							}
-
-							check(ShaderInfo);
-							ShaderInfo->ReferencedByPackages.Add(Package);
-							Package->Shaders.AddUnique(ShaderInfo);
-						}
-					}
-				}
-				Algo::Sort(Package->ShaderMapHashes);
+				UE_LOG(LogIoStore, Warning, TEXT("Package '%s' in '%s' referencing missing shader with chunk id '%s'"), *Package->PackageName.ToString(), *ContainerTargetName.ToString(), *LexToString(ShaderChunkId));
+				continue;
 			}
+
+			check(ShaderInfo);
+			ShaderInfo->ReferencedByPackages.Add(Package);
+			Package->Shaders.AddUnique(ShaderInfo);
 		}
 	}
 
+	Algo::Sort(Package->ShaderMapHashes);
+}
+
+void FIoStoreShaderLibraryProcessor::AddAllReferencedShadersToContainerTarget(FContainerTargetSpec* ContainerTarget, FIoStoreShaderLibraryOutput& Output)
+{
+	const TSet<FShaderInfo*>* FindContainerShaderLibraryShaders = ContainerTargetToShaderLibrary.Find(ContainerTarget);
+	if (!FindContainerShaderLibraryShaders)
 	{
-		IOSTORE_CPU_SCOPE(AddShaderChunks);
-		for (FContainerTargetSpec* ContainerTarget : ContainerTargets)
+		// Early exit if package doesn't contain any shaders
+		return;
+	}
+
+	for (FCookedPackage* Package : ContainerTarget->Packages)
+	{
+		for (FShaderInfo* ShaderInfo : Package->Shaders)
 		{
-			auto AddShaderTargetFile = [&ContainerTarget](FShaderInfo* ShaderInfo)
-				{
-					FContainerTargetFile& ShaderTargetFile = ContainerTarget->TargetFiles.AddDefaulted_GetRef();
-					ShaderTargetFile.ContainerTarget = ContainerTarget;
-					ShaderTargetFile.ChunkId = ShaderInfo->ChunkId;
-					ShaderTargetFile.ChunkType = EContainerChunkType::ShaderCode;
-					ShaderTargetFile.bForceUncompressed = true;
-					ShaderTargetFile.SourceBuffer.Emplace(ShaderInfo->CodeIoBuffer);
-					ShaderTargetFile.SourceSize = ShaderInfo->CodeIoBuffer.DataSize();
-				};
-
-			const TSet<FShaderInfo*>* FindContainerShaderLibraryShaders = AllContainerShaderLibraryShadersMap.Find(ContainerTarget);
-			if (FindContainerShaderLibraryShaders)
+			if (ShaderInfo->ReferencedByPackages.Num() == 1)
 			{
-				for (FCookedPackage* Package : ContainerTarget->Packages)
+				FShaderInfo::EShaderType* ShaderType = ShaderInfo->TypeInContainer.Find(ContainerTarget);
+				if (ShaderType && *ShaderType != FShaderInfo::Global)
 				{
-					for (FShaderInfo* ShaderInfo : Package->Shaders)
-					{
-						if (ShaderInfo->ReferencedByPackages.Num() == 1)
-						{
-							FShaderInfo::EShaderType* ShaderType = ShaderInfo->TypeInContainer.Find(ContainerTarget);
-							if (ShaderType && *ShaderType != FShaderInfo::Global)
-							{
-								*ShaderType = FShaderInfo::Inline;
-							}
-						}
-					}
+					*ShaderType = FShaderInfo::Inline;
 				}
+			}
+		}
+	}
 
-				for (FShaderInfo* ShaderInfo : *FindContainerShaderLibraryShaders)
-				{
-					FShaderAssociationInfo::FShaderChunkInfoKey ShaderChunkInfoKey = { ContainerTarget->Name, ShaderInfo->ChunkId };
-					FShaderAssociationInfo::FShaderChunkInfo& ShaderChunkInfo = Output.AssocInfo.ShaderChunkInfos.Add(ShaderChunkInfoKey);
+	for (FShaderInfo* ShaderInfo : *FindContainerShaderLibraryShaders)
+	{
+		FShaderAssociationInfo::FShaderChunkInfoKey ShaderChunkInfoKey = { ContainerTarget->Name, ShaderInfo->ShaderGroupChunk.ChunkId };
+		FShaderAssociationInfo::FShaderChunkInfo& ShaderChunkInfo = Output.AssocInfo.ShaderChunkInfos.Add(ShaderChunkInfoKey);
 
-					FShaderInfo::EShaderType* ShaderType = ShaderInfo->TypeInContainer.Find(ContainerTarget);
-					check(ShaderType);
-					if (*ShaderType == FShaderInfo::Global)
-					{
-						ContainerTarget->GlobalShaders.Add(ShaderInfo);
-						ShaderChunkInfo.Type = FShaderAssociationInfo::FShaderChunkInfo::Global;
-					}
-					else if (*ShaderType == FShaderInfo::Inline)
-					{
-						ContainerTarget->InlineShaders.Add(ShaderInfo);
+		const FShaderInfo::EShaderType* ShaderType = ShaderInfo->TypeInContainer.Find(ContainerTarget);
+		check(ShaderType);
+		if (*ShaderType == FShaderInfo::Global)
+		{
+			ContainerTarget->GlobalShaders.Add(ShaderInfo);
+			ShaderChunkInfo.Type = FShaderAssociationInfo::FShaderChunkInfo::Global;
+		}
+		else if (*ShaderType == FShaderInfo::Inline)
+		{
+			ContainerTarget->InlineShaders.Add(ShaderInfo);
 
-						ShaderChunkInfo.Type = FShaderAssociationInfo::FShaderChunkInfo::Package;
-						checkf(ShaderInfo->ReferencedByPackages.Num() == 1, TEXT("Inline shader chunks must be referenced by 1 package only, but shader chunk %s is referenced by %d"),
-							*LexToString(ShaderInfo->ChunkId), ShaderInfo->ReferencedByPackages.Num());
+			ShaderChunkInfo.Type = FShaderAssociationInfo::FShaderChunkInfo::Package;
+			checkf(ShaderInfo->ReferencedByPackages.Num() == 1, TEXT("Inline shader chunks must be referenced by 1 package only, but shader chunk %s is referenced by %d"),
+				*LexToString(ShaderInfo->ShaderGroupChunk.ChunkId), ShaderInfo->ReferencedByPackages.Num());
 
-						ShaderChunkInfo.ReferencedByPackages.Add(ShaderInfo->ReferencedByPackages[FSetElementId::FromInteger(0)]->PackageName);
-					}
-					else if (ShaderInfo->ReferencedByPackages.Num() > 1)
-					{
-						ContainerTarget->SharedShaders.Add(ShaderInfo);
+			ShaderChunkInfo.ReferencedByPackages.Add(ShaderInfo->ReferencedByPackages[FSetElementId::FromInteger(0)]->PackageName);
+		}
+		else if (ShaderInfo->ReferencedByPackages.Num() > 1)
+		{
+			ContainerTarget->SharedShaders.Add(ShaderInfo);
 
-						ShaderChunkInfo.Type = FShaderAssociationInfo::FShaderChunkInfo::Package;
+			ShaderChunkInfo.Type = FShaderAssociationInfo::FShaderChunkInfo::Package;
 
-						for (FCookedPackage* Package : ShaderInfo->ReferencedByPackages)
-						{
-							ShaderChunkInfo.ReferencedByPackages.Add(Package->PackageName);
+			for (const FCookedPackage* Package : ShaderInfo->ReferencedByPackages)
+			{
+				ShaderChunkInfo.ReferencedByPackages.Add(Package->PackageName);
 
-							TArray<FShaderAssociationInfo::FShaderChunkInfoKey>& PackageShaders = Output.AssocInfo.PackageShaderMap.FindOrAdd(Package->PackageName);
-							PackageShaders.Add(ShaderChunkInfoKey);
-						}
-					}
-					else
-					{
-						//
-						// Note that we can get here with shaders that get split off in to another container (e.g. sm6 shaders). 
-						// Since they are in a different pakChunk they can't get inlined or shared. However, they still "belong" to
-						// the referencing packages for association purposes.
-						//
-						if (ShaderInfo->ReferencedByPackages.Num())
-						{
-							ShaderChunkInfo.Type = FShaderAssociationInfo::FShaderChunkInfo::Package;
-
-							for (FCookedPackage* Package : ShaderInfo->ReferencedByPackages)
-							{
-								ShaderChunkInfo.ReferencedByPackages.Add(Package->PackageName);
-
-								TArray<FShaderAssociationInfo::FShaderChunkInfoKey>& PackageShaders = Output.AssocInfo.PackageShaderMap.FindOrAdd(Package->PackageName);
-								PackageShaders.Add(ShaderChunkInfoKey);
-							}
-						}
-						else
-						{
-							ShaderChunkInfo.Type = FShaderAssociationInfo::FShaderChunkInfo::Orphan;
-						}
-
-						ContainerTarget->UniqueShaders.Add(ShaderInfo);
-					}
-					AddShaderTargetFile(ShaderInfo);
+				Output.AssocInfo.AddPackageShaderInfoKey(Package->PackageName, ShaderChunkInfoKey);
+			}
+		}
+		else
+		{
+			//
+			// Note that we can get here with shaders that get split off in to another container (e.g. sm6 shaders). 
+			// Since they are in a different pakChunk they can't get inlined or shared. However, they still "belong" to
+			// the referencing packages for association purposes.
+			//
+			if (ShaderInfo->ReferencedByPackages.Num())
+			{
+				ShaderChunkInfo.Type = FShaderAssociationInfo::FShaderChunkInfo::Package;
+
+				for (FCookedPackage* Package : ShaderInfo->ReferencedByPackages)
+				{
+					ShaderChunkInfo.ReferencedByPackages.Add(Package->PackageName);
+
+					Output.AssocInfo.AddPackageShaderInfoKey(Package->PackageName, ShaderChunkInfoKey);
 				}
 			}
+			else
+			{
+				ShaderChunkInfo.Type = FShaderAssociationInfo::FShaderChunkInfo::Orphan;
+			}
+
+			ContainerTarget->UniqueShaders.Add(ShaderInfo);
 		}
+		AddShaderCodeToContainerTarget(ContainerTarget, *ShaderInfo);
 	}
-
-	double LibraryEnd = FPlatformTime::Seconds();
-	UE_LOG(LogIoStore, Display, TEXT("Shaders processed in %.2f seconds"), LibraryEnd - LibraryStart);
 }
 
 } // namespace UE::IoStore::Private
diff --git a/Engine/Source/Developer/IoStoreUtilities/Private/IoStoreShaderLibraryProcessor.h b/Engine/Source/Developer/IoStoreUtilities/Private/IoStoreShaderLibraryProcessor.h
--- a/Engine/Source/Developer/IoStoreUtilities/Private/IoStoreShaderLibraryProcessor.h
+++ b/Engine/Source/Developer/IoStoreUtilities/Private/IoStoreShaderLibraryProcessor.h
@@ -6,6 +6,7 @@
 #include "Containers/Array.h"
 #include "Containers/Map.h"
 #include "IO/IoChunkId.h"
+#include "Misc/SecureHash.h"
 #include "Templates/UniquePtr.h"
 #include "UObject/NameTypes.h"
 
@@ -13,6 +14,8 @@ namespace UE::IoStore::Private
 {
 
 struct FContainerTargetSpec;
+struct FContainerTargetFile;
+struct FCookedPackage;
 struct FIoStoreArguments;
 struct FShaderInfo;
 
@@ -62,6 +65,13 @@ struct FShaderAssociationInfo
 
 	// The list of shaders referenced by each package. This can be used to look up in to ContainerShaderInfos.
 	TMap<FName /* PackageName */, TArray<FShaderChunkInfoKey>> PackageShaderMap;
+
+	// Adds a shader chunk info key to the PackageShaderMap for the specified package name.
+	FORCEINLINE void AddPackageShaderInfoKey(const FName& PackageName, const FShaderChunkInfoKey& InfoKey)
+	{
+		TArray<FShaderAssociationInfo::FShaderChunkInfoKey>& PackageShaders = PackageShaderMap.FindOrAdd(PackageName);
+		PackageShaders.Add(InfoKey);
+	}
 };
 
 struct FIoStoreShaderLibraryOutput
@@ -73,9 +83,19 @@ struct FIoStoreShaderLibraryOutput
 // Processor to load cooked shader libraries and package them into IoStore entries
 class FIoStoreShaderLibraryProcessor
 {
+	TMap<FIoChunkId, FShaderInfo*> ChunkIdToShaderInfoMap;
+	TMap<FSHAHash, TArray<FIoChunkId>> ShaderMapToGroupChunkIds; // Maps from shader map hash to its set of shader group IoChunk Ids
+	TMap<FName, TSet<FSHAHash>> PackageNameToShaderMaps;
+	TMap<FContainerTargetSpec*, TSet<FShaderInfo*>> ContainerTargetToShaderLibrary;
+
 public:
 	void ProcessShaderLibraries(const FIoStoreArguments& Arguments, TArray<FContainerTargetSpec*>& ContainerTargets, FIoStoreShaderLibraryOutput& Output);
 
+private:
+	bool ProcessShaderLibraryFromTargetFile(const FIoStoreArguments& Arguments, FContainerTargetFile& TargetFile, FIoStoreShaderLibraryOutput& Output);
+	void UpdatePackageStoreShaders(FCookedPackage* Package, const FName& ContainerTargetName);
+	void AddAllReferencedShadersToContainerTarget(FContainerTargetSpec* ContainerTarget, FIoStoreShaderLibraryOutput& Output);
+
 };
 
 } // namespace UE::IoStore::Private
diff --git a/Engine/Source/Developer/IoStoreUtilities/Private/IoStoreUtilities.cpp b/Engine/Source/Developer/IoStoreUtilities/Private/IoStoreUtilities.cpp
--- a/Engine/Source/Developer/IoStoreUtilities/Private/IoStoreUtilities.cpp
+++ b/Engine/Source/Developer/IoStoreUtilities/Private/IoStoreUtilities.cpp
@@ -1580,7 +1580,7 @@ static void CreateDiskLayout(
 					if (ShaderType && *ShaderType == FShaderInfo::Inline)
 					{
 						check(ContainerTarget->InlineShaders.Contains(Shader));
-						FContainerTargetFile* ShaderTargetFile = ShaderTargetFilesMap.FindRef(Shader->ChunkId);
+						FContainerTargetFile* ShaderTargetFile = ShaderTargetFilesMap.FindRef(Shader->ShaderGroupChunk.ChunkId);
 						check(ShaderTargetFile);
 						PackageInlineShaders.Add(ShaderTargetFile);
 					}
@@ -1610,7 +1610,7 @@ static void CreateDiskLayout(
 				ShaderTargetFiles.Reserve(SortedShaders.Num());
 				for (const FShaderInfo* ShaderInfo : SortedShaders)
 				{
-					FContainerTargetFile* ShaderTargetFile = ShaderTargetFilesMap.FindRef(ShaderInfo->ChunkId);
+					FContainerTargetFile* ShaderTargetFile = ShaderTargetFilesMap.FindRef(ShaderInfo->ShaderGroupChunk.ChunkId);
 					check(ShaderTargetFile);
 					ShaderTargetFiles.Add(ShaderTargetFile);
 				}
@@ -5159,22 +5159,22 @@ static int32 CreateTarget(const FIoStoreArguments& Arguments, const FIoStoreWrit
 		for (const FShaderInfo* ShaderInfo : ContainerTarget->GlobalShaders)
 		{
 			++GlobalShaderCount;
-			GlobalShaderSize += ShaderInfo->CodeIoBuffer.DataSize();
+			GlobalShaderSize += ShaderInfo->ShaderGroupChunk.CodeIoBuffer.DataSize();
 		}
 		for (const FShaderInfo* ShaderInfo : ContainerTarget->SharedShaders)
 		{
 			++SharedShaderCount;
-			SharedShaderSize += ShaderInfo->CodeIoBuffer.DataSize();
+			SharedShaderSize += ShaderInfo->ShaderGroupChunk.CodeIoBuffer.DataSize();
 		}
 		for (const FShaderInfo* ShaderInfo : ContainerTarget->UniqueShaders)
 		{
 			++UniqueShaderCount;
-			UniqueShaderSize += ShaderInfo->CodeIoBuffer.DataSize();
+			UniqueShaderSize += ShaderInfo->ShaderGroupChunk.CodeIoBuffer.DataSize();
 		}
 		for (const FShaderInfo* ShaderInfo : ContainerTarget->InlineShaders)
 		{
 			++InlineShaderCount;
-			InlineShaderSize += ShaderInfo->CodeIoBuffer.DataSize();
+			InlineShaderSize += ShaderInfo->ShaderGroupChunk.CodeIoBuffer.DataSize();
 		}
 	}
 
diff --git a/Engine/Source/Developer/IoStoreUtilities/Private/IoStoreUtilitiesCore.h b/Engine/Source/Developer/IoStoreUtilities/Private/IoStoreUtilitiesCore.h
--- a/Engine/Source/Developer/IoStoreUtilities/Private/IoStoreUtilitiesCore.h
+++ b/Engine/Source/Developer/IoStoreUtilities/Private/IoStoreUtilitiesCore.h
@@ -171,27 +171,36 @@ struct FIoStoreArguments
 	}
 };
 
+// IoChunk per shader group
+struct FShaderGroupIoChunk
+{
+	FIoChunkId ChunkId;
+	FIoBuffer CodeIoBuffer;
+	uint32 LoadOrderFactor = (uint32)MAX_int32; // the smaller the order, the more likely this will be loaded
+
+	static bool Sort(const FShaderGroupIoChunk& A, const FShaderGroupIoChunk& B)
+	{
+		if (A.LoadOrderFactor == B.LoadOrderFactor)
+		{
+			// Shader chunk IDs are the hash of the shader so this is consistent across builds.
+			return FMemory::Memcmp(A.ChunkId.GetData(), B.ChunkId.GetData(), A.ChunkId.GetSize()) < 0;
+		}
+		return A.LoadOrderFactor < B.LoadOrderFactor;
+	}
+};
+
 struct FShaderInfo
 {
 	enum EShaderType { Normal, Global, Inline };
 
-	FIoChunkId ChunkId;
-	FIoBuffer CodeIoBuffer;
-	// the smaller the order, the more likely this will be loaded
-	uint32 LoadOrderFactor = MAX_uint32;
+	FShaderGroupIoChunk ShaderGroupChunk;
 	TSet<struct FCookedPackage*> ReferencedByPackages;
 	TMap<FContainerTargetSpec*, EShaderType> TypeInContainer;
 
-	// This is used to ensure build determinism and such must be stable
-	// across builds.
+	// This is used to ensure build determinism and such must be stable across builds.
 	static bool Sort(const FShaderInfo* A, const FShaderInfo* B)
 	{
-		if (A->LoadOrderFactor == B->LoadOrderFactor)
-		{
-			// Shader chunk IDs are the hash of the shader so this is consistent across builds.
-			return FMemory::Memcmp(A->ChunkId.GetData(), B->ChunkId.GetData(), A->ChunkId.GetSize()) < 0;
-		}
-		return A->LoadOrderFactor < B->LoadOrderFactor;
+		return FShaderGroupIoChunk::Sort(A->ShaderGroupChunk, B->ShaderGroupChunk);
 	}
 };
 
