[ShaderLibrary] Restore previous behavior of always sorting shadermap hashes in `FIoStoreShaderLibraryProcessor::UpdatePackageStoreShaders()`.

This is a potential fix of ZenStore issues showing up since the refactoring of the shader code library (introduced in CL 48448616).

#rnx

#Synced-CL 48943914

diff --git a/Engine/Source/Developer/IoStoreUtilities/Private/IoStoreShaderLibraryProcessor.cpp b/Engine/Source/Developer/IoStoreUtilities/Private/IoStoreShaderLibraryProcessor.cpp
--- a/Engine/Source/Developer/IoStoreUtilities/Private/IoStoreShaderLibraryProcessor.cpp
+++ b/Engine/Source/Developer/IoStoreUtilities/Private/IoStoreShaderLibraryProcessor.cpp
@@ -536,38 +536,33 @@ void FIoStoreShaderLibraryProcessor::UpdatePackageStoreShaders(FCookedPackage* P
 	// 1. Update ShaderInfos with which packages we reference.
 	// 2. Add to packages which shaders we use.
 	// 3. Add to PackageStore what shaders we use.
-	const TSet<FSHAHash>* FindShaderMapHashes = PackageNameToShaderMaps.Find(Package->PackageName);
-	if (!FindShaderMapHashes)
+	if (const TSet<FSHAHash>* FindShaderMapHashes = PackageNameToShaderMaps.Find(Package->PackageName))
 	{
-		// Early exit if package doesn't contain any shader maps
-		return;
-	}
-
-	for (const FSHAHash& ShaderMapHash : *FindShaderMapHashes)
-	{
-		const TArray<FIoChunkId>* FindChunkIds = ShaderMapToGroupChunkIds.Find(ShaderMapHash);
-		if (!FindChunkIds)
-		{
-			UE_LOG(LogIoStore, Warning, TEXT("Package '%s' in '%s' referencing missing shader map '%s'"), *Package->PackageName.ToString(), *ContainerTargetName.ToString(), *ShaderMapHash.ToString());
-			continue;
-		}
-		Package->ShaderMapHashes.Add(ShaderMapHash);
-
-		for (const FIoChunkId& ShaderChunkId : *FindChunkIds)
+		for (const FSHAHash& ShaderMapHash : *FindShaderMapHashes)
 		{
-			FShaderInfo* ShaderInfo = ChunkIdToShaderInfoMap.FindRef(ShaderChunkId);
-			if (!ShaderInfo)
+			const TArray<FIoChunkId>* FindChunkIds = ShaderMapToGroupChunkIds.Find(ShaderMapHash);
+			if (!FindChunkIds)
 			{
-				UE_LOG(LogIoStore, Warning, TEXT("Package '%s' in '%s' referencing missing shader with chunk id '%s'"), *Package->PackageName.ToString(), *ContainerTargetName.ToString(), *LexToString(ShaderChunkId));
+				UE_LOG(LogIoStore, Warning, TEXT("Package '%s' in '%s' referencing missing shader map '%s'"), *Package->PackageName.ToString(), *ContainerTargetName.ToString(), *ShaderMapHash.ToString());
 				continue;
 			}
+			Package->ShaderMapHashes.Add(ShaderMapHash);
 
-			check(ShaderInfo);
-			ShaderInfo->ReferencedByPackages.Add(Package);
-			Package->Shaders.AddUnique(ShaderInfo);
+			for (const FIoChunkId& ShaderChunkId : *FindChunkIds)
+			{
+				FShaderInfo* ShaderInfo = ChunkIdToShaderInfoMap.FindRef(ShaderChunkId);
+				if (!ShaderInfo)
+				{
+					UE_LOG(LogIoStore, Warning, TEXT("Package '%s' in '%s' referencing missing shader with chunk id '%s'"), *Package->PackageName.ToString(), *ContainerTargetName.ToString(), *LexToString(ShaderChunkId));
+					continue;
+				}
+
+				check(ShaderInfo);
+				ShaderInfo->ReferencedByPackages.Add(Package);
+				Package->Shaders.AddUnique(ShaderInfo);
+			}
 		}
 	}
-
 	Algo::Sort(Package->ShaderMapHashes);
 }
 
