UnrealPak IoStore Staging: Shaders: Add log feedback for which chunk is missing in the assertion failure in FSerializedShaderArchiveContainer::Load. Downgrade the severity from assertion to Error to match the severity of local file read failure.
#rnx
#Synced-CL 48939402

diff --git a/Engine/Source/Developer/IoStoreUtilities/Private/IoStoreShaderLibraryProcessor.cpp b/Engine/Source/Developer/IoStoreUtilities/Private/IoStoreShaderLibraryProcessor.cpp
--- a/Engine/Source/Developer/IoStoreUtilities/Private/IoStoreShaderLibraryProcessor.cpp
+++ b/Engine/Source/Developer/IoStoreUtilities/Private/IoStoreShaderLibraryProcessor.cpp
@@ -138,7 +138,16 @@ public:
 	{
 		if (PackageStore && PackageStore->HasZenStoreClient() && InChunkId.IsValid())
 		{
-			LibraryBuffer = PackageStore->ReadChunk(InChunkId).ConsumeValueOrDie();
+			TIoStatusOr<FIoBuffer> ReadChunkResult = PackageStore->ReadChunk(InChunkId);
+			if (!ReadChunkResult.IsOk())
+			{
+				UE_LOG(LogIoStore, Error,
+					TEXT("Missing shader code libary ChunkId. ChunkId = '%s'. InFilename = '%s'. IOStatus = '%s'."),
+					*LexToString(InChunkId), InFilename ? InFilename : TEXT("<nullptr>"),
+					*ReadChunkResult.Status().ToString());
+				return false;
+			}
+			LibraryBuffer = ReadChunkResult.ConsumeValueOrDie();
 		}
 		else
 		{
