[ShaderLibrary] Revert initial value of `FShaderGroupIoChunk::LoadOrderFactor` back to `MAX_uint32`.

This was erroneously changed to `MAX_int32` in CL 48546765.

#rnx

#Synced-CL 48984866

diff --git a/Engine/Source/Developer/IoStoreUtilities/Private/IoStoreUtilitiesCore.h b/Engine/Source/Developer/IoStoreUtilities/Private/IoStoreUtilitiesCore.h
--- a/Engine/Source/Developer/IoStoreUtilities/Private/IoStoreUtilitiesCore.h
+++ b/Engine/Source/Developer/IoStoreUtilities/Private/IoStoreUtilitiesCore.h
@@ -176,7 +176,7 @@ struct FShaderGroupIoChunk
 {
 	FIoChunkId ChunkId;
 	FIoBuffer CodeIoBuffer;
-	uint32 LoadOrderFactor = (uint32)MAX_int32; // the smaller the order, the more likely this will be loaded
+	uint32 LoadOrderFactor = MAX_uint32; // the smaller the order, the more likely this will be loaded
 
 	static bool Sort(const FShaderGroupIoChunk& A, const FShaderGroupIoChunk& B)
 	{
