UnrealPak: Update ListContainer to not require the ucas files for listing.

#rnx
#Synced-CL 48801413

diff --git a/Engine/Source/Developer/IoStoreUtilities/Private/IoStoreUtilities.cpp b/Engine/Source/Developer/IoStoreUtilities/Private/IoStoreUtilities.cpp
--- a/Engine/Source/Developer/IoStoreUtilities/Private/IoStoreUtilities.cpp
+++ b/Engine/Source/Developer/IoStoreUtilities/Private/IoStoreUtilities.cpp
@@ -5603,16 +5603,29 @@ static int32 FindAndExtractPackages(
 
 void IterateContainerFilesInternal(const FKeyChain& KeyChain, const TMap<FPackageId, TPair<FName, FName>>& AllChunksInfoMap, const FString& ContainerFilePath, FChunkEntryCsv* Out)
 {
-	TUniquePtr<FIoStoreReader> Reader = CreateIoStoreReader(*ContainerFilePath, KeyChain);
-	if (!Reader.IsValid())
+	FString TocFileName = FPaths::ChangeExtension(ContainerFilePath, TEXT("utoc"));
+
+	TMap<FGuid, FAES::FAESKey> DecryptionKeys;
+	for (const auto& KV : KeyChain.GetEncryptionKeys())
 	{
-		UE_LOG(LogIoStore, Warning, TEXT("Failed to read container '%s'"), *ContainerFilePath);
-		return;
+		DecryptionKeys.Add(KV.Key, KV.Value.Key);
 	}
 
-	Out->SetPartitionSize(Reader->GetPartitionSize());
+	TUniquePtr<IIoStoreTocReader> TocReader;
+	{
+		IOSTORE_CPU_SCOPE(ReadToc);
+		TIoStatusOr<TUniquePtr<IIoStoreTocReader>> ReadResult = IIoStoreTocReader::ReadFromDisk(*TocFileName, DecryptionKeys);
+		if (!ReadResult.IsOk())
+		{
+			UE_LOG(LogIoStore, Warning, TEXT("Failed to read container toc '%s' for container '%s': %s"), *TocFileName, *ContainerFilePath, *ReadResult.Status().ToString());
+			return;
+		}
 
-	if (!EnumHasAnyFlags(Reader->GetContainerFlags(), EIoContainerFlags::Indexed))
+		TocReader = ReadResult.ConsumeValueOrDie();
+	}
+
+	Out->SetPartitionSize(TocReader->GetTocResource().Header.PartitionSize);
+	if (!EnumHasAnyFlags(TocReader->GetTocResource().Header.ContainerFlags, EIoContainerFlags::Indexed))
 	{
 		UE_LOG(LogIoStore, Display, TEXT("No directory index for container '%s'"), *ContainerFilePath);
 	}
@@ -5620,19 +5633,18 @@ void IterateContainerFilesInternal(const FKeyChain& KeyChain, const TMap<FPackag
 	UE_LOG(LogIoStore, Display, TEXT("Listing container '%s'"), *ContainerFilePath);
 
 	FString ContainerName = FPaths::GetBaseFilename(ContainerFilePath);
+	
 	TArray<FIoStoreTocChunkInfo> Chunks;
-
 	{
 		IOSTORE_CPU_SCOPE(EnumerateChunks);
-		Reader->EnumerateChunks([&Chunks](FIoStoreTocChunkInfo&& ChunkInfo)
-			{
-				Chunks.Add(MoveTemp(ChunkInfo));
-				return true;
-			});
+		for (int32 ChunkIndex = 0; ChunkIndex < TocReader->GetTocResource().ChunkIds.Num(); ++ChunkIndex)
+		{
+			Chunks.Add(TocReader->GetTocChunkInfo(ChunkIndex));
+		}
 	}
 
 	{
-		IOSTORE_CPU_SCOPE(EnumerateChunks);
+		IOSTORE_CPU_SCOPE(SortChunks);
 		auto SortKey = [](const FIoStoreTocChunkInfo& ChunkInfo) { return ChunkInfo.OffsetOnDisk; };
 		Algo::SortBy(Chunks, SortKey);
 	}
diff --git a/Engine/Source/Runtime/Core/Internal/IO/IoStore.h b/Engine/Source/Runtime/Core/Internal/IO/IoStore.h
--- a/Engine/Source/Runtime/Core/Internal/IO/IoStore.h
+++ b/Engine/Source/Runtime/Core/Internal/IO/IoStore.h
@@ -308,3 +308,19 @@ struct FIoStoreTocResource
 
 	[[nodiscard]] CORE_API static uint64 HashChunkIdWithSeed(int32 Seed, const FIoChunkId& ChunkId);
 };
+
+/**
+ * This is the next level up from FIoStoreTocResource - it loads the toc but also resolves the directory index so you get 
+ * filenames. If you don't need that, use FIoStoreTocResource.
+*/
+class IIoStoreTocReader
+{
+public:
+	virtual const FIoStoreTocResourceView& GetTocResource() const = 0;
+	[[nodiscard]] virtual FIoStoreTocChunkInfo GetTocChunkInfo(int32 TocEntryIndex) const = 0;
+	CORE_API virtual ~IIoStoreTocReader();
+
+	// Creates an instance for accessing the TOC. OutTocReader is only valid when FIoStatus is FIoStatus::Ok. The TocFilePath should be a path to the
+	// .utoc file (not the ucas or undecorated). 
+	CORE_API static TIoStatusOr<TUniquePtr<IIoStoreTocReader>> ReadFromDisk(const TCHAR* TocFilePath, const TMap<FGuid, FAES::FAESKey>& DecryptionKeys);
+};
\ No newline at end of file
diff --git a/Engine/Source/Runtime/Core/Private/IO/IoStore.cpp b/Engine/Source/Runtime/Core/Private/IO/IoStore.cpp
--- a/Engine/Source/Runtime/Core/Private/IO/IoStore.cpp
+++ b/Engine/Source/Runtime/Core/Private/IO/IoStore.cpp
@@ -149,13 +149,14 @@ static FIoStatus ValidateContainerSignature(
 	return FIoStatus::Ok;
 }
 
-class FIoStoreTocReader
+class FIoStoreTocReader : public IIoStoreTocReader
 {
 public:
 	FIoStoreTocReader()
 	{
 		FMemory::Memzero(&Toc.Header, sizeof(FIoStoreTocHeader));
 	}
+	virtual ~FIoStoreTocReader() = default;
 
 	[[nodiscard]] FIoStatus Read(const TCHAR* TocFilePath, const TMap<FGuid, FAES::FAESKey>& DecryptionKeys)
 	{
@@ -265,6 +266,19 @@ private:
 	TMap<int32, FString> IndexToFileName;
 };
 
+IIoStoreTocReader::~IIoStoreTocReader() = default;
+
+/* static */ TIoStatusOr<TUniquePtr<IIoStoreTocReader>> IIoStoreTocReader::ReadFromDisk(const TCHAR* TocFilePath, const TMap<FGuid, FAES::FAESKey>& DecryptionKeys)
+{
+	TUniquePtr<IIoStoreTocReader> TocReader = MakeUnique<FIoStoreTocReader>();
+	FIoStatus Result = ((FIoStoreTocReader*)TocReader.Get())->Read(TocFilePath, DecryptionKeys);
+	if (Result != FIoStatus::Ok)
+	{
+		return Result;
+	}
+	return TocReader;
+}
+
 template<typename T>
 static FIoStoreTocChunkInfo GetTocChunkInfoInternal(const T& Toc, int32 TocEntryIndex)
 {
