Fix for unrealpak list containers showing the wrong platform when named pakchunks are used.

#rnx
#Synced-CL 49586837

diff --git a/Engine/Source/Developer/IoStoreUtilities/Private/IoStoreUtilities.cpp b/Engine/Source/Developer/IoStoreUtilities/Private/IoStoreUtilities.cpp
--- a/Engine/Source/Developer/IoStoreUtilities/Private/IoStoreUtilities.cpp
+++ b/Engine/Source/Developer/IoStoreUtilities/Private/IoStoreUtilities.cpp
@@ -949,8 +949,9 @@ public:
 			FString InstallType = TEXT("Base");
 			FString DeliveryType = TEXT("Installed");
 
-			// Extract the PakChunk name and the Platform from the container name
-			ContainerName.Split(TEXT("-"), &PakChunk, &Platform);
+			// Extract the PakChunk name and the Platform from the container name. Have to go from end because
+			// named containers use - as a delimited and we want the last instance for the platform.
+			ContainerName.Split(TEXT("-"), &PakChunk, &Platform, ESearchCase::CaseSensitive, ESearchDir::FromEnd);
 
 			// Determine the Delivery Type Installed/IAS/IAD
 			if (ContainerName.Contains(TEXT("iad")))
