IncrementalCook: Do not record shader files from the previous cook as files to stage in the current cook. We do not need to do that because they are written to disk every cook and we find them during ZenFileSystemManifest->Generate. We need to not add them (or need to prune them) to support the ability of code or asset changes to drop them from the staged output. This change fixes the staging bug "Trying to fetch chunk <ChunkId> using file path <CookOutputDir>/WindowsClient/<ProjectName>/Content/ShaderArchive-<ProjectName>_Chunk<PakFileChunkId>-PCD3D_ES31-PCD3D_ES3_1.ushaderbytecode failed."
#rnx
#Synced-CL 49917028

diff --git a/Engine/Source/Developer/IoStoreUtilities/Private/ZenStoreWriter.cpp b/Engine/Source/Developer/IoStoreUtilities/Private/ZenStoreWriter.cpp
--- a/Engine/Source/Developer/IoStoreUtilities/Private/ZenStoreWriter.cpp
+++ b/Engine/Source/Developer/IoStoreUtilities/Private/ZenStoreWriter.cpp
@@ -3,6 +3,7 @@
 #include "ZenStoreWriter.h"
 
 #include "Algo/BinarySearch.h"
+#include "Algo/Find.h"
 #include "Algo/IsSorted.h"
 #include "Algo/Sort.h"
 #include "AssetRegistry/AssetRegistryState.h"
@@ -496,6 +497,24 @@ void FZenStoreWriter::WriteProjectStoreFile()
 	ProjectStoreWriter.Write(ServiceHostName, ServicePort, IsLocalConnection, *ProjectId, *OplogId, TargetPlatformFName);
 }
 
+// TODO: Move this into FPathViews and change it to take flags instead of bools
+FStringView FPathViewsGetExtension(const FStringView& InPath, bool bIncludeDot, bool bIncludeCompoundExtension)
+{
+	if (!bIncludeCompoundExtension)
+	{
+		return FPathViews::GetExtension(InPath, bIncludeDot);
+	}
+	else
+	{
+		if (const TCHAR* Dot = Algo::Find(FPathViews::GetCleanFilename(InPath), TEXT('.')))
+		{
+			const TCHAR* Extension = bIncludeDot ? Dot : Dot + 1;
+			return FStringView(Extension, UE_PTRDIFF_TO_INT32(InPath.GetData() + InPath.Len() - Extension));
+		}
+		return FStringView();
+	}
+}
+
 void FZenStoreWriter::Initialize(const FCookInfo& Info)
 {
 	CookMode = Info.CookMode;
@@ -586,6 +605,9 @@ void FZenStoreWriter::Initialize(const FCookInfo& Info)
 
 			UE_LOG(LogZenStoreWriter, Display, TEXT("Fetching file manifest..."), *ProjectId, *OplogId);
 
+			TSet<FString> CopyFromPreviousCookDenyListExtensions;
+			CopyFromPreviousCookDenyListExtensions.Append({ TEXT("ushaderbytecode"),TEXT("assetinfo.json"), TEXT("stinfo"), TEXT("ushaderpipelines") });
+
 			TIoStatusOr<FCbObject> FileStatus = HttpClient->GetFiles().Get();
 			if (FileStatus.IsOk())
 			{
@@ -605,7 +627,20 @@ void FZenStoreWriter::Initialize(const FCookInfo& Info)
 							FIoChunkId FileChunkId;
 							FileChunkId.Set(FileId.GetView());
 
-							ZenFileSystemManifest->AddManifestEntry(FileChunkId, MoveTemp(ServerPath), MoveTemp(ClientPath));
+							// Do not allow the addition of some extensions that we know are rewritten every cook
+							// and that we want to be able to remove from the cook when they are removed due to
+							// asset or code changes. TODO: Remove this copy-through of the files from the previous
+							// cook and use the ZenCookArtifactReader instead, and allow the cooker to delete files
+							// from the ZenCookArtifactReader when it deletes them from local disk, see UE-351803
+							// and UE-352816.
+							FStringView Extension = FPathViewsGetExtension(ClientPath, false /* bIncludeDot */,
+								true /* bIncludeCompoundExtension */);
+							if (!CopyFromPreviousCookDenyListExtensions.ContainsByHash(
+								GetTypeHash(Extension), Extension))
+							{
+								ZenFileSystemManifest->AddManifestEntry(FileChunkId, MoveTemp(ServerPath),
+									MoveTemp(ClientPath));
+							}
 						}
 					}
 				}
