UnrealPak initial support for patch-in-place storage containers layout

#rnx
#Synced-CL 50050818

diff --git a/Engine/Source/Developer/IoStoreUtilities/Private/IoStoreUtilities.cpp b/Engine/Source/Developer/IoStoreUtilities/Private/IoStoreUtilities.cpp
--- a/Engine/Source/Developer/IoStoreUtilities/Private/IoStoreUtilities.cpp
+++ b/Engine/Source/Developer/IoStoreUtilities/Private/IoStoreUtilities.cpp
@@ -349,6 +349,18 @@ public:
 		}
 	}
 
+	virtual bool GetContainerFilePath(const FIoContainerId& InContainerId, FString& OutContainerFilePath) override
+	{
+		TUniquePtr<FReaderChunks>* ReaderChunksPtr = ChunkDatabase.Find(InContainerId);
+		if (ReaderChunksPtr == nullptr)
+		{
+			return false;
+		}
+
+		OutContainerFilePath = (*ReaderChunksPtr)->ContainerPathAndFileName;
+		return true;
+	}
+
 	virtual uint32 GetCompressionBlockSize() const override
 	{
 		return CompressionBlockSize;
@@ -2278,7 +2290,7 @@ static void InitializeContainerTargetsAndPackages(
 	});
 };
 
-void LogWriterResults(const TArray<FIoStoreWriterResult>& Results)
+void LogWriterResults(const FIoStoreWriterSettings& GeneralIoWriterSettings, const TArray<FIoStoreWriterResult>& Results)
 {
 	struct FContainerStats
 	{
@@ -2287,6 +2299,8 @@ void LogWriterResults(const TArray<FIoStoreWriterResult>& Results)
 		uint64 UncompressedContainerSize = 0;
 		uint64 CompressedContainerSize = 0;
 		uint64 PaddingSize = 0;
+		uint64 PatchInPlacePinnedSize = 0;
+		int64  PatchInPlaceContainerSizeChange = 0;
 	};
 
 	// Container name length can vary wildly, so find the longest one so that we can correctly align the output grids
@@ -2302,9 +2316,14 @@ void LogWriterResults(const TArray<FIoStoreWriterResult>& Results)
 	UE_LOG(LogIoStore, Display, TEXT("Container Summary"));
 	UE_LOG(LogIoStore, Display, TEXT("=================="));
 	UE_LOG(LogIoStore, Display, TEXT(""));
-	UE_LOG(LogIoStore, Display, TEXT("%-*s %10s %15s %15s %20s %20s"), ColumnWidth, TEXT("Container"), TEXT("Flags"), TEXT("Chunk(s) #"), TEXT("TOC (KiB)"), TEXT("Raw Size (MiB)"), TEXT("Size (MiB)"));
-	
-	const int32 SummaryBorderLength = ColumnWidth + 5 + 10 + 15 + 15 + 20 + 20; // +5 to account for spaces between columns
+	FString PatchInPlaceHeader = TEXT("");
+	if (GeneralIoWriterSettings.bUsePatchInPlaceLayout)
+	{
+		PatchInPlaceHeader = FString::Printf(TEXT("%20s %20s"), TEXT("PIP Pinned (MiB)"), TEXT("PIP Overhead (MiB)"));
+	}
+	UE_LOG(LogIoStore, Display, TEXT("%-*s %10s %15s %15s %20s %20s%s"), ColumnWidth, TEXT("Container"), TEXT("Flags"), TEXT("Chunk(s) #"), TEXT("TOC (KiB)"), TEXT("Raw Size (MiB)"), TEXT("Size (MiB)"), *PatchInPlaceHeader);
+
+	const int32 SummaryBorderLength = ColumnWidth + 5 + 10 + 15 + 15 + 20 + 20 + (GeneralIoWriterSettings.bUsePatchInPlaceLayout ? 2 + 20 + 20 : 0); // +5 to account for spaces between columns
 	UE_LOG(LogIoStore, Display, TEXT("%s"), *FString::ChrN(SummaryBorderLength, '-'));
 	
 	FContainerStats TotalStats;
@@ -2321,6 +2340,14 @@ void LogWriterResults(const TArray<FIoStoreWriterResult>& Results)
 				*Result.CompressionMethod.ToString());
 		}
 
+		FString PatchInPlaceInfo = TEXT("");
+		if (GeneralIoWriterSettings.bUsePatchInPlaceLayout)
+		{
+			PatchInPlaceInfo = FString::Printf(TEXT(" %20.2lf %20.2lf"),
+				(double)Result.PatchInPlacePinnedSize / 1024.0 / 1024.0,
+				(double)Result.PatchInPlaceContainerSizeChange / 1024.0 / 1024.0);
+		}
+
 		FString ContainerSettings = FString::Printf(TEXT("%s/%s/%s/%s/%s"),
 			EnumHasAnyFlags(Result.ContainerFlags, EIoContainerFlags::Compressed) ? TEXT("C") : TEXT("-"),
 			EnumHasAnyFlags(Result.ContainerFlags, EIoContainerFlags::Encrypted) ? TEXT("E") : TEXT("-"),
@@ -2328,13 +2355,14 @@ void LogWriterResults(const TArray<FIoStoreWriterResult>& Results)
 			EnumHasAnyFlags(Result.ContainerFlags, EIoContainerFlags::Indexed) ? TEXT("I") : TEXT("-"),
 			EnumHasAnyFlags(Result.ContainerFlags, EIoContainerFlags::OnDemand) ? TEXT("O") : TEXT("-"));
 
-		UE_LOG(LogIoStore, Display, TEXT("%-*s %10s %15llu %15.2lf %20.2lf %20.2lf %s"), ColumnWidth,
+		UE_LOG(LogIoStore, Display, TEXT("%-*s %10s %15llu %15.2lf %20.2lf %20.2lf%s %s"), ColumnWidth,
 			*Result.ContainerName,
 			*ContainerSettings,
 			Result.TocEntryCount,
 			(double)Result.TocSize / 1024.0,
 			(double)Result.UncompressedContainerSize / 1024.0 / 1024.0,
 			(double)Result.CompressedContainerSize / 1024.0 / 1024.0,
+			*PatchInPlaceInfo,
 			*CompressionInfo);
 
 		if (EnumHasAnyFlags(Result.ContainerFlags, EIoContainerFlags::OnDemand))
@@ -2350,6 +2378,8 @@ void LogWriterResults(const TArray<FIoStoreWriterResult>& Results)
 		TotalStats.UncompressedContainerSize += Result.UncompressedContainerSize;
 		TotalStats.CompressedContainerSize += Result.CompressedContainerSize;
 		TotalStats.PaddingSize += Result.PaddingSize;
+		TotalStats.PatchInPlacePinnedSize += Result.PatchInPlacePinnedSize;
+		TotalStats.PatchInPlaceContainerSizeChange += Result.PatchInPlaceContainerSizeChange;
 	}
 
 	UE_LOG(LogIoStore, Display, TEXT("%s"), *FString::ChrN(SummaryBorderLength, '-'));
@@ -2365,13 +2395,22 @@ void LogWriterResults(const TArray<FIoStoreWriterResult>& Results)
 			(double)OnDemandStats.CompressedContainerSize / 1024.0 / 1024.0);
 	}
 
-	UE_LOG(LogIoStore, Display, TEXT("%-*s %10s %15llu %15.2lf %20.2lf %20.2lf"), ColumnWidth,
+	FString PatchInPlaceTotal = TEXT("");
+	if (GeneralIoWriterSettings.bUsePatchInPlaceLayout)
+	{
+		PatchInPlaceTotal = FString::Printf(TEXT(" %20.2lf %20.2lf"),
+			(double)TotalStats.PatchInPlacePinnedSize / 1024.0 / 1024.0,
+			(double)TotalStats.PatchInPlaceContainerSizeChange / 1024.0 / 1024.0);
+	}
+
+	UE_LOG(LogIoStore, Display, TEXT("%-*s %10s %15llu %15.2lf %20.2lf %20.2lf%s"), ColumnWidth,
 		TEXT("Total"),
 		TEXT(""),
 		TotalStats.TocCount,
 		(double)TotalStats.TocSize / 1024.0,
 		(double)TotalStats.UncompressedContainerSize / 1024.0 / 1024.0,
-		(double)TotalStats.CompressedContainerSize / 1024.0 / 1024.0);
+		(double)TotalStats.CompressedContainerSize / 1024.0 / 1024.0,
+		*PatchInPlaceTotal);
 
 	UE_LOG(LogIoStore, Display, TEXT(""));
 	UE_LOG(LogIoStore, Display, TEXT("** Flags: (C)ompressed / (E)ncrypted / (S)igned) / (I)ndexed) / (O)nDemand **"));
@@ -4539,6 +4578,7 @@ static int32 CreateTarget(const FIoStoreArguments& Arguments, const FIoStoreWrit
 				GlobalContainerSettings.ContainerFlags |= EIoContainerFlags::Signed;
 			}
 			GlobalIoStoreWriter = IoStoreWriterContext->CreateContainer(*Arguments.GlobalContainerPath, GlobalContainerSettings);
+			GlobalIoStoreWriter->SetReferenceChunkDatabase(ChunkDatabase);
 			IoStoreWriters.Add(GlobalIoStoreWriter);
 			IoStoreWriterInfos.Add({UE::Cook::EPluginSizeTypes::Installed, "global"});
 		}
@@ -5179,7 +5219,7 @@ static int32 CreateTarget(const FIoStoreArguments& Arguments, const FIoStoreWrit
 		}
 	}
 
-	LogWriterResults(IoStoreWriterResults);
+	LogWriterResults(GeneralIoWriterSettings, IoStoreWriterResults);
 	LogContainerPackageInfo(ContainerTargets);
 
 	UE_LOG(LogIoStore, Display, TEXT("Input:  %8d Packages"), Packages.Num());
@@ -5516,7 +5556,7 @@ int32 CreateContentPatch(const FIoStoreArguments& Arguments, const FIoStoreWrite
 		Results.Emplace(IoStoreWriter->GetResult().ConsumeValueOrDie());
 	}
 
-	LogWriterResults(Results);
+	LogWriterResults(GeneralIoWriterSettings, Results);
 
 	return 0;
 }
@@ -9051,6 +9091,8 @@ bool ParseContainerGenerationArguments(FIoStoreArguments& Arguments, FIoStoreWri
 	WriterSettings.bContainerMetaEnabled = FParse::Param(FCommandLine::Get(), TEXT("ContainerMeta"));
 	WriterSettings.bPerContainerMeta = Arguments.IsDLC();
 
+	WriterSettings.bUsePatchInPlaceLayout = FParse::Param(FCommandLine::Get(), TEXT("UsePatchInPlaceLayout"));
+
 	UE_LOG(LogIoStore, Display, TEXT("Using memory mapping alignment '%ld'"), WriterSettings.MemoryMappingAlignment);
 	UE_LOG(LogIoStore, Display, TEXT("Using compression block size '%ld'"), WriterSettings.CompressionBlockSize);
 	UE_LOG(LogIoStore, Display, TEXT("Using compression block alignment '%ld'"), WriterSettings.CompressionBlockAlignment);
@@ -9087,6 +9129,11 @@ bool ParseContainerGenerationArguments(FIoStoreArguments& Arguments, FIoStoreWri
 		UE_LOG(LogIoStore, Display, TEXT("Container metadata disabled"));
 	}
 
+	if (WriterSettings.bUsePatchInPlaceLayout)
+	{
+		UE_LOG(LogIoStore, Display, TEXT("Using patch-in-place ordering"));
+	}
+
 	FString CommandListFile;
 	if (FParse::Value(FCommandLine::Get(), TEXT("Commands="), CommandListFile))
 	{
diff --git a/Engine/Source/Developer/IoStoreUtilities/Private/IoStoreWriter.cpp b/Engine/Source/Developer/IoStoreUtilities/Private/IoStoreWriter.cpp
--- a/Engine/Source/Developer/IoStoreUtilities/Private/IoStoreWriter.cpp
+++ b/Engine/Source/Developer/IoStoreUtilities/Private/IoStoreWriter.cpp
@@ -14,6 +14,7 @@
 #include "DerivedDataCacheInterface.h"
 #include "DerivedDataCacheRecord.h"
 #include "DerivedDataRequestOwner.h"
+#include "Algo/IsSorted.h"
 #include "HAL/FileManager.h"
 #include "HAL/PlatformFileManager.h"
 #include "IO/IoContainerMeta.h"
@@ -27,6 +28,7 @@
 #include "Serialization/MemoryWriter.h"
 #include "Tasks/Task.h"
 #include "Templates/UniquePtr.h"
+#include "Templates/Greater.h"
 
 TRACE_DECLARE_MEMORY_COUNTER(IoStoreCompressionMemoryUsed, TEXT("IoStoreWriter/CompressionMemoryUsed"));
 TRACE_DECLARE_MEMORY_COUNTER(IoStoreCompressionMemoryScheduled, TEXT("IoStoreWriter/CompressionMemoryScheduled"));
@@ -776,6 +778,614 @@ private:
 	TMap<FIoChunkId, FString> ChunkIdToFileName;
 };
 
+class FPatchInPlaceLayout
+{
+	static constexpr bool bParanoidChecks = false;
+
+public:
+	FPatchInPlaceLayout(
+		const FString& InReferenceContainerPath,
+		const FString& InSourceContainerPath,
+		const FString& InTargetContainerPath,
+		const FIoStoreTocResourceView& InSourceToc,
+		uint64 InMemoryMappingAlignment)
+		: ReferenceContainerPath(InReferenceContainerPath)
+		, SourceContainerPath(InSourceContainerPath)
+		, TargetContainerPath(InTargetContainerPath)
+		, SourceToc(InSourceToc)
+		, MemoryMappingAlignment(InMemoryMappingAlignment)
+		, ReferenceToc()
+	{
+	}
+
+	bool PrepareForLayout()
+	{
+		TRACE_CPUPROFILER_EVENT_SCOPE(PrepareForLayout);
+
+		const FIoStatus ReferenceTocValid = FIoStoreTocResourceView::Read(*ReferenceContainerPath, EIoStoreTocReadOptions::ReadAll, ReferenceToc, ReferenceTocStorage);
+		if (!ReferenceTocValid.IsOk())
+		{
+			return false;
+		}
+
+		BuildTocMap(ReferenceToc, ReferenceTocChunkIdToIndex);
+		BuildTocMap(SourceToc, SourceTocChunkIdToIndex);
+
+		// TODO this algo needs a 0%-100% progressive tuning
+		// e.g. 0% would mean nothing is stays in place, 100% would mean everything stays in place
+		// this is needed to let data move around to shrink containers between big releases
+
+		PinChunksEntries();
+
+		// only do the layout if there are entries to be kept in place
+		return !EntriesToKeep.IsEmpty();
+	}
+
+	void RunLayout(FIoStoreTocResource& TargetToc, int32& OutTargetPartitionNum, uint64& OutPinnedSize, uint64& OutTotalPadding)
+	{
+		TRACE_CPUPROFILER_EVENT_SCOPE(PatchInPlaceLayout);
+
+		UE_LOG(LogIoStore, Display, TEXT("Patch-in-place '%s' -> '%s'"), *ReferenceContainerPath, *TargetContainerPath);
+
+		OutPinnedSize = 0;
+		for (const FTocEntry& Entry : EntriesToKeep)
+		{
+			OutPinnedSize += Entry.Size;
+		}
+
+		uint64 InsertedSize = 0;
+		for (const FTocEntry& Entry : EntriesToInsert)
+		{
+			InsertedSize += Entry.Size;
+		}
+
+		int32 SourcePartitionNum = 0;
+
+		TArray<FCopyOp> CopyOps = {};
+
+		BestFitDecreasing(EntriesToKeep, EntriesToInsert, MemoryMappingAlignment);
+		FitIntoPartitions(EntriesToKeep, SourceToc.Header.PartitionSize, OutTotalPadding);
+		BuildTargetToc(TargetToc, CopyOps, EntriesToKeep, SourcePartitionNum, OutTargetPartitionNum);
+		WriteChunksToTarget(EntriesToKeep, CopyOps, SourcePartitionNum, OutTargetPartitionNum);
+
+		for (int32 Index = 0; Index < OutTargetPartitionNum; ++Index)
+		{
+			OverwriteFile(GetPartitionFilePath(SourceContainerPath, Index), GetPartitionFilePath(TargetContainerPath, Index));
+		}
+	}
+
+private:
+	const FString ReferenceContainerPath;
+	const FString SourceContainerPath;
+	const FString TargetContainerPath;
+	const FIoStoreTocResourceView SourceToc;
+	const uint64 MemoryMappingAlignment;
+
+	FIoStoreTocResourceView ReferenceToc;
+	FIoStoreTocResourceStorage ReferenceTocStorage;
+
+	TMap<FIoChunkId, int32> ReferenceTocChunkIdToIndex;
+	TMap<FIoChunkId, int32> SourceTocChunkIdToIndex;
+
+	struct FTocEntry;
+	TArray<FTocEntry> EntriesToKeep;
+	TArray<FTocEntry> EntriesToInsert;
+
+	static void BuildTocMap(const FIoStoreTocResourceView& Toc, TMap<FIoChunkId, int32>& OutResult)
+	{
+		OutResult.Reserve(Toc.ChunkIds.Num());
+		for (int32 Index = 0; Index < Toc.ChunkIds.Num(); ++Index)
+		{
+			OutResult.Add(Toc.ChunkIds[Index], Index);
+		}
+	}
+
+	struct FTocEntry
+	{
+		FIoChunkId Id;
+		uint64 Offset;
+		uint64 Size;
+
+		FTocEntry(FIoChunkId InId, uint64 InOffset, uint64 InSize)
+			: Id(InId), Offset(InOffset), Size(InSize)
+		{
+		}
+
+		uint64 End() const
+		{
+			return Offset + Size;
+		}
+	};
+
+	struct FLayoutGap
+	{
+		uint64 Begin;
+		uint64 End;
+
+		FLayoutGap(uint64 InBegin, uint64 InEnd)
+			: Begin(InBegin), End(InEnd)
+		{
+		}
+		
+		uint64 Size() const
+		{
+			return End - Begin;
+		}
+	};
+
+	struct FCopyOp
+	{
+		uint64 SourceOffset;
+		uint64 TargetOffset;
+		uint64 Size;
+
+		FCopyOp(uint64 InSourceOffset, uint64 InTargetOffset, uint64 InSize)
+			: SourceOffset(InSourceOffset)
+			, TargetOffset(InTargetOffset)
+			, Size(InSize)
+		{
+		}
+
+		int32 SourcePartitionIndex(const uint64 PartitionSize) const
+		{
+			return PartitionSize > 0 ? SourceOffset / PartitionSize : 0;
+		}
+
+		int32 TargetPartitionIndex(const uint64 PartitionSize) const
+		{
+			return PartitionSize > 0 ? TargetOffset / PartitionSize : 0;
+		}
+
+		uint64 SourcePartitionOffset(const uint64 PartitionSize) const
+		{
+			return PartitionSize > 0 ? SourceOffset % PartitionSize : SourceOffset;
+		}
+
+		uint64 TargetPartitionOffset(const uint64 PartitionSize) const
+		{
+			return PartitionSize > 0 ? TargetOffset % PartitionSize : TargetOffset;
+		}
+	};
+
+	void PinChunksEntries()
+	{
+		for (int32 ChunkIndex = 0; ChunkIndex < SourceToc.ChunkIds.Num(); ChunkIndex++)
+		{
+			const FIoStoreTocChunkInfo& SourceChunk = SourceToc.GetTocChunkInfo(ChunkIndex);
+
+			const int32* ReferenceTocIndex = ReferenceTocChunkIdToIndex.Find(SourceChunk.Id);
+			if (!ReferenceTocIndex)
+			{
+				EntriesToInsert.Add(FTocEntry(SourceChunk.Id, UINT64_MAX, SourceChunk.SizeOnDisk));
+				continue;
+			}
+
+			const FIoStoreTocChunkInfo ReferenceChunk = ReferenceToc.GetTocChunkInfo(*ReferenceTocIndex);
+			if (SourceChunk.ChunkHash == ReferenceChunk.ChunkHash &&
+				SourceChunk.SizeOnDisk == ReferenceChunk.SizeOnDisk)
+			{
+				EntriesToKeep.Add(FTocEntry(SourceChunk.Id, ReferenceChunk.OffsetOnDisk, SourceChunk.SizeOnDisk));
+			}
+			else
+			{
+				EntriesToInsert.Add(FTocEntry(SourceChunk.Id, UINT64_MAX, SourceChunk.SizeOnDisk));
+			}
+		}
+	}
+
+	static uint64 GapsBinarySearchProjection(const FLayoutGap& Gap)
+	{
+		return Gap.Size();
+	}
+
+	static uint64 AlignOffset(const uint64 Offset, const uint64 MemoryMappingAlignment)
+	{
+		uint64 Result = Offset;
+		if (MemoryMappingAlignment > 0)
+		{
+			Result = Align(Result, MemoryMappingAlignment);
+		}
+		return Result;
+	}
+	
+	static void EnsureSortedAndNonOverlapping(TConstArrayView<FTocEntry> Entries)
+	{
+		for (int32 Index = 1; Index < Entries.Num(); ++Index)
+		{
+			verify(Entries[Index].Offset > Entries[Index - 1].Offset);
+		}
+
+		for (int32 I = 0; I < Entries.Num(); ++I)
+		{
+			for (int32 J = 0; J < Entries.Num(); ++J)
+			{
+				if (I == J)
+				{
+					continue;
+				}
+
+				const uint64 L1 = Entries[I].Offset;
+				const uint64 R1 = Entries[I].End() - 1;
+
+				const uint64 L2 = Entries[J].Offset;
+				const uint64 R2 = Entries[J].End() - 1;
+
+				verify(R1 < L2 || L1 > R2);
+			}
+		}
+	}
+
+	static void BestFitDecreasing(TArray<FTocEntry>& EntriesToKeep, TArray<FTocEntry>& EntriesToInsert, const uint64 MemoryMappingAlignment)
+	{
+		TRACE_CPUPROFILER_EVENT_SCOPE(BestFitDecreasing);
+
+		Algo::SortBy(EntriesToKeep, [](const FTocEntry& X){return X.Offset;});
+		Algo::SortBy(EntriesToInsert, [](const FTocEntry& X){return X.Size;}, TGreater<uint64>());
+
+		if (bParanoidChecks)
+		{
+			EnsureSortedAndNonOverlapping(EntriesToKeep);
+		}
+
+		TArray<FLayoutGap> Gaps;
+		if (!EntriesToKeep.IsEmpty() && EntriesToKeep[0].Offset > 0)
+		{
+			Gaps.Add(FLayoutGap(0, EntriesToKeep[0].Offset));
+		}
+		
+		for (int32 Index = 1; Index < EntriesToKeep.Num(); ++Index)
+		{
+			const uint64 PrevEnd = AlignOffset(EntriesToKeep[Index - 1].End(), MemoryMappingAlignment);
+			const uint64 CurrentBegin = EntriesToKeep[Index].Offset;
+			if (CurrentBegin > PrevEnd + 1)
+			{
+				Gaps.Add(FLayoutGap(PrevEnd, CurrentBegin - 1));
+			}
+		}
+
+		Algo::SortBy(Gaps, GapsBinarySearchProjection);
+
+		uint64 AppendOffset = !EntriesToKeep.IsEmpty() ? EntriesToKeep.Last().End() : 0;
+
+		// TODO support compression block alignment?
+		for (FTocEntry& Entry : EntriesToInsert)
+		{
+			const int32 BestGapIndex = Algo::UpperBoundBy(Gaps, Entry.Size, GapsBinarySearchProjection);
+
+			// TODO make it faster
+			if (BestGapIndex != Gaps.Num())
+			{
+				FLayoutGap Gap = Gaps[BestGapIndex];
+				if (bParanoidChecks)
+				{
+					verify(Gap.Size() >= Entry.Size);
+				}
+
+				Gaps.RemoveAt(BestGapIndex, EAllowShrinking::No);
+
+				Entry.Offset = Gap.Begin;
+				Gap.Begin = AlignOffset(Gap.Begin + Entry.Size, MemoryMappingAlignment);
+
+				if (Gap.Begin < Gap.End)
+				{
+					Gaps.Insert(Gap, Algo::UpperBoundBy(Gaps, Gap.Size(), GapsBinarySearchProjection));
+
+					if (bParanoidChecks)
+					{
+						verify(Algo::IsSortedBy(Gaps, GapsBinarySearchProjection));
+					}
+				}
+			}
+			else
+			{
+				Entry.Offset = AppendOffset;
+				AppendOffset = AlignOffset(AppendOffset + Entry.Size, MemoryMappingAlignment);
+			}
+		}
+
+		EntriesToKeep.Append(MoveTemp(EntriesToInsert));
+
+		Algo::SortBy(EntriesToKeep, [](const FTocEntry& X) {return X.Offset;});
+
+		if (bParanoidChecks)
+		{
+			EnsureSortedAndNonOverlapping(EntriesToKeep);
+		}
+	}
+
+	static void FitIntoPartitions(TArray<FTocEntry>& Chunks, const uint64 PartitionSize, uint64& OutTotalPadding)
+	{
+		OutTotalPadding = 0;
+
+		uint64 Padding = 0;
+		uint64 PreviousEnd = 0;
+
+		for (FTocEntry& Chunk : Chunks)
+		{
+			ensure(Chunk.Size < PartitionSize);
+			const uint64 BeginIndex = (Chunk.Offset + Padding) / PartitionSize;
+			const uint64 EndIndex = (Chunk.Offset + Chunk.Size + Padding) / PartitionSize;
+
+			if (EndIndex != BeginIndex)
+			{
+				// bump chunk to the next partition
+				Padding += EndIndex * PartitionSize - (Chunk.Offset + Padding);
+			}
+
+			Chunk.Offset += Padding;
+
+			OutTotalPadding = Chunk.Offset - PreviousEnd;
+			PreviousEnd = Chunk.Offset + Chunk.Size;
+		}
+	}
+
+	void BuildTargetToc(FIoStoreTocResource& TargetToc, TArray<FCopyOp>& OutCopyOps, TConstArrayView<FTocEntry> Chunks, int32& OutSourcePartitionNum, int32& OutTargetPartitionNum)
+	{
+		TRACE_CPUPROFILER_EVENT_SCOPE(BuildTargetToc);
+
+		if (Chunks.Num() != SourceToc.ChunkIds.Num())
+		{
+			UE_LOG(LogIoStore, Error, TEXT("Broken ToC of '%s'"), *SourceContainerPath);
+			return;
+		}
+
+		const int32 ChunksNum = Chunks.Num();
+
+		TargetToc.ChunkIds.SetNum(ChunksNum);
+		TargetToc.ChunkOffsetLengths.SetNum(ChunksNum);
+		TargetToc.ChunkMetas.SetNum(ChunksNum);
+		TargetToc.CompressionBlocks.SetNum(SourceToc.CompressionBlocks.Num());
+		TargetToc.ChunkBlockSignatures.SetNum(SourceToc.ChunkBlockSignatures.Num());
+		TargetToc.CompressionMethods = SourceToc.CompressionMethods;
+		TargetToc.DirectoryIndexBuffer = SourceToc.DirectoryIndexBuffer;
+
+		const uint32 CompressionBlockSize = SourceToc.Header.CompressionBlockSize;
+		const uint64 MaxPartitionSize = SourceToc.Header.PartitionSize;
+
+		uint64 UncompressedOffset = 0;
+		TUniquePtr<FArchive> TargetHandle;
+
+		int32 TargetChunkIndex = 0;
+		int32 TargetBlockIndex = 0;
+
+		OutCopyOps.Reserve(Chunks.Num());
+
+		int32 SourcePartitionNum = 0;
+		int32 TargetPartitionNum = 0;
+
+		const uint64 MaxPossibleOffset = FIoStoreTocCompressedBlockEntry::OffsetMask;
+
+		for (const FTocEntry& Chunk : Chunks)
+		{
+			if (!ensure(SourceTocChunkIdToIndex.Contains(Chunk.Id)))
+			{
+				continue;
+			}
+			const int32 SourceIndex = *SourceTocChunkIdToIndex.Find(Chunk.Id);
+
+			TargetToc.ChunkIds[TargetChunkIndex] = Chunk.Id;
+			TargetToc.ChunkMetas[TargetChunkIndex] = SourceToc.ChunkMetas[SourceIndex];
+
+			const uint64 SourceOffset = SourceToc.ChunkOffsetLengths[SourceIndex].GetOffset();
+			const uint64 SourceSize = SourceToc.ChunkOffsetLengths[SourceIndex].GetLength();
+			const int32 SourceFirstBlockIndex = int32(SourceOffset / CompressionBlockSize);
+			const int32 SourceLastBlockIndex = int32((Align(SourceOffset + SourceSize, CompressionBlockSize) - 1) / CompressionBlockSize);
+
+			const uint64 SourceReadOffset = SourceToc.CompressionBlocks[SourceFirstBlockIndex].GetOffset();
+			uint64 UncompressedSize = 0;
+			uint64 DiskSize = 0;
+
+			for (int32 SourceBlockIndex = SourceFirstBlockIndex; SourceBlockIndex <= SourceLastBlockIndex; ++SourceBlockIndex, ++TargetBlockIndex)
+			{
+				const FIoStoreTocCompressedBlockEntry& SourceBlock = SourceToc.CompressionBlocks[SourceBlockIndex];
+				FIoStoreTocCompressedBlockEntry& TargetBlock = TargetToc.CompressionBlocks[TargetBlockIndex];
+
+				const uint64 TargetBlockOffset = Chunk.Offset + DiskSize;
+				verify(TargetBlockOffset <= MaxPossibleOffset);
+
+				TargetBlock.SetOffset(Chunk.Offset + DiskSize); // this is disk offset between all partitions
+				TargetBlock.SetCompressedSize(SourceBlock.GetCompressedSize());
+				TargetBlock.SetUncompressedSize(SourceBlock.GetUncompressedSize());
+				TargetBlock.SetCompressionMethodIndex(SourceBlock.GetCompressionMethodIndex());
+
+				if (SourceBlockIndex < SourceToc.ChunkBlockSignatures.Num())
+				{
+					TargetToc.ChunkBlockSignatures[TargetBlockIndex] = SourceToc.ChunkBlockSignatures[SourceBlockIndex];
+				}
+
+				UncompressedSize += SourceBlock.GetUncompressedSize();
+				DiskSize += Align(SourceBlock.GetCompressedSize(), FAES::AESBlockSize);
+			}
+
+			verify(UncompressedOffset <= MaxPossibleOffset);
+			verify(UncompressedSize <= MaxPossibleOffset);
+
+			TargetToc.ChunkOffsetLengths[TargetChunkIndex].SetOffset(UncompressedOffset);
+			TargetToc.ChunkOffsetLengths[TargetChunkIndex].SetLength(UncompressedSize);
+
+			// try to coalesce with previous operation
+			if (OutCopyOps.Num() > 0 && 
+				OutCopyOps.Last().SourceOffset + OutCopyOps.Last().Size == SourceReadOffset &&
+				OutCopyOps.Last().TargetOffset + OutCopyOps.Last().Size == Chunk.Offset &&
+				// only coalesce with-in same partition
+				(MaxPartitionSize > 0 ? (
+					OutCopyOps.Last().SourceOffset / MaxPartitionSize == SourceReadOffset / MaxPartitionSize &&
+					OutCopyOps.Last().TargetOffset / MaxPartitionSize == Chunk.Offset / MaxPartitionSize
+					) : true)
+				)
+			{
+				OutCopyOps.Last().Size += Chunk.Size;
+			}
+			else
+			{
+				const FCopyOp& Op = OutCopyOps.Add_GetRef(FCopyOp(SourceReadOffset, Chunk.Offset, Chunk.Size));
+				SourcePartitionNum = FMath::Max(Op.SourcePartitionIndex(MaxPartitionSize) + 1, SourcePartitionNum);
+				TargetPartitionNum = FMath::Max(Op.TargetPartitionIndex(MaxPartitionSize) + 1, TargetPartitionNum);
+			}
+
+			verify(DiskSize == Chunk.Size);
+
+			UncompressedOffset += (SourceLastBlockIndex - SourceFirstBlockIndex + 1) * CompressionBlockSize;
+			TargetChunkIndex++;
+		}
+
+		Algo::SortBy(OutCopyOps, [](const FCopyOp& X) {return X.SourceOffset;});
+
+		OutSourcePartitionNum = SourcePartitionNum;
+		OutTargetPartitionNum = TargetPartitionNum;
+	}
+
+	void WriteChunksToTarget(TConstArrayView<FTocEntry> Chunks, TConstArrayView<FCopyOp> CopyOps, const int32 SourcePartitionNum, const int32 TargetPartitionNum)
+	{
+		TRACE_CPUPROFILER_EVENT_SCOPE(WriteChunksToTarget);
+
+		if (Chunks.IsEmpty() || SourcePartitionNum == 0 || TargetPartitionNum == 0 || CopyOps.IsEmpty())
+		{
+			return;
+		}
+
+		const uint64 MaxPartitionSize = SourceToc.Header.PartitionSize;
+		const uint64 LastPartitionSizeLogical = Chunks.Last().Offset + Chunks.Last().Size;
+		const uint64 LastPartitionSize = MaxPartitionSize > 0 ? LastPartitionSizeLogical % MaxPartitionSize : LastPartitionSizeLogical;
+
+		if (LastPartitionSizeLogical == 0)
+		{
+			return;
+		}
+
+		IPlatformFile& Ipf = IPlatformFile::GetPlatformPhysical();
+		Ipf.CreateDirectoryTree(*FPaths::GetPath(TargetContainerPath));
+
+		TArray<TUniquePtr<IFileHandle>> SourcePartitions;
+		SourcePartitions.SetNumZeroed(SourcePartitionNum);
+		for (int32 I = 0; I < SourcePartitionNum; ++I)
+		{
+			const FString PartitionFilePath = GetPartitionFilePath(SourceContainerPath, I);
+			FFileOpenResult Result = Ipf.OpenRead(*PartitionFilePath, IPlatformFile::EOpenReadFlags::None);
+			if (!Result.IsValid())
+			{
+				UE_LOG(LogIoStore, Fatal, TEXT("Failed to read partition '%s'"), *PartitionFilePath);
+				return;
+			}
+			SourcePartitions[I] = MoveTemp(Result.GetValue());
+		}
+
+		TArray<TUniquePtr<IFileHandle>> TargetPartitions;
+		TargetPartitions.SetNumZeroed(TargetPartitionNum);
+		for (int32 I = 0; I < TargetPartitionNum; ++I)
+		{
+			const FString PartitionFilePath = GetPartitionFilePath(TargetContainerPath, I);
+			FFileOpenResult Result = Ipf.OpenWrite(*PartitionFilePath, IPlatformFile::EOpenWriteFlags::AllowRead | IPlatformFile::EOpenWriteFlags::Append);
+			if (!Result.IsValid())
+			{
+				UE_LOG(LogIoStore, Fatal, TEXT("Failed to read partition '%s'"), *PartitionFilePath);
+				return;
+			}
+			TargetPartitions[I] = MoveTemp(Result.GetValue());
+			if (!TargetPartitions[I]->Truncate(I + 1 < TargetPartitionNum ? MaxPartitionSize : LastPartitionSize))
+			{
+				UE_LOG(LogIoStore, Error, TEXT("Failed to resize partition '%s'"), *PartitionFilePath);
+			}
+		}
+
+		constexpr uint64 TempBufferSize = 4 * 1024 * 1024; 
+		FIoBuffer TempBuffer(TempBufferSize);
+
+		for (const FCopyOp& Op : CopyOps) 
+		{
+			IFileHandle* Source = SourcePartitions[Op.SourcePartitionIndex(MaxPartitionSize)].Get();
+			IFileHandle* Target = TargetPartitions[Op.TargetPartitionIndex(MaxPartitionSize)].Get();
+
+			const uint64 SourceOffset = Op.SourcePartitionOffset(MaxPartitionSize);
+			const uint64 TargetOffset = Op.TargetPartitionOffset(MaxPartitionSize);
+
+			for (uint64 Offset = 0; Offset < Op.Size; Offset += TempBufferSize)
+			{
+				const uint64 Size = FMath::Min(Op.Size - Offset, TempBufferSize);
+
+				if (bParanoidChecks)
+				{
+					verify(SourceOffset + Offset + Size <= (uint64)Source->Size());
+					verify(TargetOffset + Offset + Size <= (uint64)Target->Size());
+				}
+
+				Source->ReadAt(TempBuffer.GetData(), Size, SourceOffset + Offset);
+
+				Target->Seek(TargetOffset + Offset);
+				Target->Write(TempBuffer.GetData(), Size);
+			}
+		}
+
+		if (bParanoidChecks)
+		{
+			FIoBuffer TempBuffer2(TempBufferSize);
+
+			for (const FCopyOp& Op : CopyOps) 
+			{
+				IFileHandle* Source = SourcePartitions[Op.SourcePartitionIndex(MaxPartitionSize)].Get();
+				IFileHandle* Target = TargetPartitions[Op.TargetPartitionIndex(MaxPartitionSize)].Get();
+
+				const uint64 SourceOffset = Op.SourcePartitionOffset(MaxPartitionSize);
+				const uint64 TargetOffset = Op.TargetPartitionOffset(MaxPartitionSize);
+
+				for (uint64 Offset = 0; Offset < Op.Size; Offset += TempBufferSize)
+				{
+					const uint64 Size = FMath::Min(Op.Size - Offset, TempBufferSize);
+
+					Source->ReadAt(TempBuffer.GetData(), Size, SourceOffset + Offset);
+					Source->ReadAt(TempBuffer2.GetData(), Size, TargetOffset + Offset);
+					
+					verify(FMemory::Memcmp(TempBuffer.GetData(), TempBuffer2.GetData(), Size) == 0);
+				}
+			}
+		}
+	}
+
+	static FString GetPartitionFilePath(FString ContainerFilePath, int32 PartitionIndex)
+	{
+		FString BasePath = FPaths::GetBaseFilename(ContainerFilePath, false);
+		if (PartitionIndex > 0)
+		{
+			BasePath += FString::Printf(TEXT("_s%d"), PartitionIndex);
+		}
+		return BasePath + TEXT(".ucas");
+	}
+
+	static void OverwriteFile(const FString& DstPath, const FString& SrcPath)
+	{
+		TRACE_CPUPROFILER_EVENT_SCOPE(OverwriteFile);
+
+		IPlatformFile& Ipf = IPlatformFile::GetPlatformPhysical();
+		if (!Ipf.FileExists(*SrcPath))
+		{
+			UE_LOG(LogIoStore, Error, TEXT("Source file '%s' doesn't exist"), *SrcPath);
+			return;
+		}
+
+		if (bParanoidChecks)
+		{
+			Ipf.MoveFile(*(DstPath + TEXT(".backup")), *DstPath);
+		}
+
+		Ipf.DeleteFile(*DstPath);
+
+		if (Ipf.MoveFile(*DstPath, *SrcPath))
+		{
+			UE_LOG(LogIoStore, Display, TEXT("Moved '%s' -> '%s'"), *SrcPath, *DstPath);
+			return;
+		}
+
+		if (Ipf.CopyFile(*DstPath, *SrcPath))
+		{
+			UE_LOG(LogIoStore, Display, TEXT("Copied '%s' -> '%s'"), *SrcPath, *DstPath);
+			Ipf.DeleteFile(*SrcPath);
+			return;
+		}
+
+		UE_LOG(LogIoStore, Error, TEXT("Failed to move '%s' -> '%s'"), *SrcPath, *DstPath);
+	}
+};
+
 class FIoStoreWriter
 	: public IIoStoreWriter
 {
@@ -1282,13 +1892,116 @@ public:
 		return true;
 	}
 
+	void DoPatchInPlaceLayout()
+	{
+		if (!ReferenceChunkDatabase.IsValid())
+		{
+			UE_LOG(LogIoStore, Warning, TEXT("Trying to use patch-in-place without valid reference database"));
+			return;
+		}
+
+		FString ReferenceContainerPath;
+		if (!ReferenceChunkDatabase->GetContainerFilePath(ContainerSettings.ContainerId, ReferenceContainerPath))
+		{
+			UE_LOG(LogIoStore, Display, TEXT("Skipping patch-in-place for container '%s' because reference container was not found"), *ContainerPathAndBaseFileName);
+			return;
+		}
+
+		const FIoStoreWriterSettings& WriterSettings = WriterContext->GetSettings();
+
+		if (WriterSettings.CompressionBlockAlignment)
+		{
+			UE_LOG(LogIoStore, Warning, TEXT("Ignoring CompressionBlockAlignment in for container '%s' due to patch-in-place enabled"), *ContainerPathAndBaseFileName);
+		}
+
+		FIoStoreTocResource& SourceToc = TocBuilder.GetTocResource();
+		// needed for FIoStoreTocResource::GetTocChunkInfo to work 
+		SourceToc.Header.CompressionBlockSize = WriterSettings.CompressionBlockSize;
+		SourceToc.Header.PartitionSize = MaxPartitionSize > 0 ? MaxPartitionSize : MAX_uint64;
+
+		FIoStoreTocResource TargetToc;
+		int32 TargetPartitionNum = 0;
+
+		FPatchInPlaceLayout Layout(
+			ReferenceContainerPath,
+			TocFilePath,
+			ContainerPathAndBaseFileName + TEXT("_patchinplace.utoc"),
+			TocBuilder.GetTocResource().AsView(),
+			WriterSettings.MemoryMappingAlignment
+		);
+
+		if (!Layout.PrepareForLayout())
+		{
+			return;
+		}
+
+		int64 ContainerSizeBeforePatchInPlace = 0;
+		for (FPartition& Partition : Partitions)
+		{
+			ContainerSizeBeforePatchInPlace += Partition.Offset;
+		}
+
+		// release file handles for patch-in-place
+		Partitions.Empty();
+
+		Layout.RunLayout(TargetToc, TargetPartitionNum, PatchInPlacePinnedSize, TotalPaddingSize);
+
+		SourceToc = MoveTemp(TargetToc);
+
+		// recreate file handles
+		int64 ContainerSizeAfterPatchInPlace = 0;
+		for (int32 i = 0; i < TargetPartitionNum; ++i)
+		{
+			FPartition& Partition = Partitions.AddDefaulted_GetRef();
+			Partition.Index = Partitions.Num() - 1;
+			CreatePartitionContainerFile(Partition, true);
+			Partition.Offset = Partition.ContainerFileHandle->Tell();
+			ContainerSizeAfterPatchInPlace += Partition.Offset;
+		}
+
+		PatchInPlaceContainerSizeChange = ContainerSizeAfterPatchInPlace - ContainerSizeBeforePatchInPlace;
+
+		if (WriterContext->GetSettings().bEnableFileRegions)
+		{
+			int32 PreviousPartitionIndex = -1;
+
+			for (int32 ChunkIndex = 0; ChunkIndex < SourceToc.ChunkIds.Num(); ++ChunkIndex)
+			{
+				const FIoStoreTocChunkInfo& SourceChunk = SourceToc.GetTocChunkInfo(ChunkIndex);
+
+				ensure(SourceChunk.PartitionIndex >= PreviousPartitionIndex);
+				PreviousPartitionIndex = SourceChunk.PartitionIndex;
+
+				const TArray<FFileRegion>* FileRegionsPtr = PatchInPlaceFileRegions.Find(SourceChunk.Id);
+				FFileRegion::AccumulateFileRegions(
+					Partitions[SourceChunk.PartitionIndex].AllFileRegions,
+					SourceChunk.OffsetOnDisk,
+					SourceChunk.OffsetOnDisk,
+					SourceChunk.OffsetOnDisk + SourceChunk.SizeOnDisk,
+					FileRegionsPtr ? *FileRegionsPtr : TConstArrayView<FFileRegion>()
+				);
+			}
+
+			PatchInPlaceFileRegions.Empty();
+		}
+
+		CurrentPartitionIndex = TargetPartitionNum - 1;
+	}
+
 	void Finalize()
 	{
 		check(bHasFlushed);
 
+		const FIoStoreWriterSettings& WriterSettings = WriterContext->GetSettings();
+
+		if (WriterSettings.bUsePatchInPlaceLayout)
+		{
+			DoPatchInPlaceLayout();
+		}
+
 		UncompressedContainerSize = TotalEntryUncompressedSize + TotalPaddingSize;
 		CompressedContainerSize = 0;
-		const FIoStoreWriterSettings& WriterSettings = WriterContext->GetSettings();
+
 		for (FPartition& Partition : Partitions)
 		{
 			CompressedContainerSize += Partition.Offset;
@@ -1394,7 +2107,7 @@ public:
 
 		TIoStatusOr<uint64> TocSize = FIoStoreTocResource::Write(*TocFilePath, TocResource, static_cast<uint32>(WriterSettings.CompressionBlockSize), MaxPartitionSize, ContainerSettings);
 		checkf(TocSize.IsOk(), TEXT("FIoStoreTocResource::Write failed with error %s"), *TocSize.Status().ToString());
-		
+
 		Result.ContainerId = ContainerSettings.ContainerId;
 		Result.ContainerName = FPaths::GetBaseFilename(TocFilePath);
 		Result.ContainerFlags = ContainerSettings.ContainerFlags;
@@ -1409,6 +2122,8 @@ public:
 		Result.CompressionMethod = EnumHasAnyFlags(ContainerSettings.ContainerFlags, EIoContainerFlags::Compressed)
 			? WriterSettings.CompressionMethod
 			: NAME_None;
+		Result.PatchInPlacePinnedSize = PatchInPlacePinnedSize;
+		Result.PatchInPlaceContainerSizeChange = PatchInPlaceContainerSizeChange;
 		Result.ModifiedChunksCount = 0;
 		Result.AddedChunksCount = 0;
 		Result.ModifiedChunksSize= 0;
@@ -1583,8 +2298,10 @@ private:
 		LayoutEntries.Empty();
 	}
 
-	FIoStatus CreatePartitionContainerFile(FPartition& Partition)
+	FIoStatus CreatePartitionContainerFile(FPartition& Partition, const bool bUseExistingData = false)
 	{
+		const uint32 WriteFlags = bUseExistingData ? FILEWRITE_Append | FILEWRITE_AllowRead : 0;
+
 		check(!Partition.ContainerFileHandle);
 		FString ContainerFilePath = ContainerPathAndBaseFileName;
 		if (Partition.Index > 0)
@@ -1592,8 +2309,8 @@ private:
 			ContainerFilePath += FString::Printf(TEXT("_s%d"), Partition.Index);
 		}
 		ContainerFilePath += TEXT(".ucas");
-		
-		Partition.ContainerFileHandle.Reset(IFileManager::Get().CreateFileWriter(*ContainerFilePath));
+
+		Partition.ContainerFileHandle.Reset(IFileManager::Get().CreateFileWriter(*ContainerFilePath, WriteFlags));
 		if (!Partition.ContainerFileHandle)
 		{
 			return FIoStatusBuilder(EIoErrorCode::FileOpenFailed) << TEXT("Failed to open IoStore container file '") << *ContainerFilePath << TEXT("'");
@@ -1601,7 +2318,7 @@ private:
 		if (WriterContext->GetSettings().bEnableFileRegions)
 		{
 			FString RegionsFilePath = ContainerFilePath + FFileRegion::RegionsFileExtension;
-			Partition.RegionsArchive.Reset(IFileManager::Get().CreateFileWriter(*RegionsFilePath));
+			Partition.RegionsArchive.Reset(IFileManager::Get().CreateFileWriter(*RegionsFilePath, WriteFlags));
 			if (!Partition.RegionsArchive)
 			{
 				return FIoStatusBuilder(EIoErrorCode::FileOpenFailed) << TEXT("Failed to open IoStore regions file '") << *RegionsFilePath << TEXT("'");
@@ -2133,7 +2850,19 @@ private:
 
 		if (WriterSettings.bEnableFileRegions)
 		{
-			FFileRegion::AccumulateFileRegions(TargetPartition->AllFileRegions, RegionStartOffset, RegionStartOffset, TargetPartition->Offset, Entry->Request->GetRegions());
+			if (WriterSettings.bUsePatchInPlaceLayout)
+			{
+				TConstArrayView<FFileRegion> FileRegions = Entry->Request->GetRegions();
+				if (!FileRegions.IsEmpty())
+				{
+					// make a copy as entry is deleted after this method returns
+					PatchInPlaceFileRegions.Emplace(Entry->ChunkId, FileRegions);
+				}
+			}
+			else
+			{
+				FFileRegion::AccumulateFileRegions(TargetPartition->AllFileRegions, RegionStartOffset, RegionStartOffset, TargetPartition->Offset, Entry->Request->GetRegions());
+			}
 		}
 		uint64 WriteStartCycles = FPlatformTime::Cycles64();
 		uint64 WriteBytes = 0;
@@ -2174,6 +2903,7 @@ private:
 	FLayoutEntry*				LayoutEntriesHead = nullptr;
 	FLayoutEntry*				LayoutEntriesTail = nullptr;
 	TMap<FIoChunkId, FLayoutEntry*> PreviousBuildLayoutEntryByChunkId;
+	TMap<FIoChunkId, TArray<FFileRegion>> PatchInPlaceFileRegions;
 	TUniquePtr<FArchive>		CsvArchive;
 	FIoStoreWriterResult		Result;
 	uint64						UncompressedFileOffset = 0;
@@ -2183,7 +2913,9 @@ private:
 	uint64						TotalPaddingSize = 0;
 	uint64						UncompressedContainerSize = 0; // this is the size the container would be if it were uncompressed.
 	uint64						CompressedContainerSize = 0; // this is the size of the container with the given compression (which may be none).
-	uint64 						MaxPartitionSize = 0;
+	uint64						PatchInPlacePinnedSize = 0;
+	int64						PatchInPlaceContainerSizeChange = 0;
+	uint64						MaxPartitionSize = 0;
 	int32						CurrentPartitionIndex = 0;
 	bool						bHasMemoryMappedEntry = false;
 	bool						bHasFlushed = false;
@@ -2266,7 +2998,7 @@ void FIoStoreWriterContextImpl::Flush()
 	ParallelFor(TEXT("IoStoreWriter::Finalize.PF"), IoStoreWriters.Num(), 1, [this](int Index)
 	{ 
 		IoStoreWriters[Index]->Finalize(); 
-	});
+	}, EParallelForFlags::Unbalanced);
 	double FinalizeEnd = FPlatformTime::Seconds();
 	int64 TotalTocSize = 0;
 	for (TSharedPtr<IIoStoreWriter> Writer : IoStoreWriters)
diff --git a/Engine/Source/Developer/IoStoreUtilities/Private/IoStoreWriter.h b/Engine/Source/Developer/IoStoreUtilities/Private/IoStoreWriter.h
--- a/Engine/Source/Developer/IoStoreUtilities/Private/IoStoreWriter.h
+++ b/Engine/Source/Developer/IoStoreUtilities/Private/IoStoreWriter.h
@@ -28,6 +28,7 @@ struct FIoStoreWriterSettings
 	bool bPerContainerMeta = false;
 	bool bCompressionEnableDDC = false;
 	bool bValidateChunkHashes = false;
+	bool bUsePatchInPlaceLayout = false;
 };
 
 struct FIoStoreWriterResult
@@ -42,6 +43,8 @@ struct FIoStoreWriterResult
 	int64 DirectoryIndexSize = 0;
 	uint64 TotalEntryCompressedSize = 0; // sum of the compressed size of entries excluding encryption alignment.
 	uint64 ReferenceCacheMissBytes = 0; // number of compressed bytes excluding alignment that could have been from refcache but weren't.
+	uint64 PatchInPlacePinnedSize = 0; // compressed size of all entries which offset is matching reference db offset.
+	int64 PatchInPlaceContainerSizeChange = 0; // difference of value of CompressedContainerSize before and after patch-in-place layout.
 	uint64 AddedChunksCount = 0;
 	uint64 AddedChunksSize = 0;
 	uint64 ModifiedChunksCount = 0;
@@ -171,6 +174,11 @@ public:
 	* Called by an iostore writer implementation to notify the ref cache it's been added
 	*/
 	virtual void NotifyAddedToWriter(const FIoContainerId& InContainerId, const FString& InContainerName) = 0;
+
+	/*
+	* Returns true if ContainerId was found and OutContainerFilePath is set.
+	*/
+	virtual bool GetContainerFilePath(const FIoContainerId& InContainerId, FString& OutContainerFilePath) = 0;
 };
 
 class IIoStoreWriter
diff --git a/Engine/Source/Programs/AutomationTool/AutomationUtils/ProjectParams.cs b/Engine/Source/Programs/AutomationTool/AutomationUtils/ProjectParams.cs
--- a/Engine/Source/Programs/AutomationTool/AutomationUtils/ProjectParams.cs
+++ b/Engine/Source/Programs/AutomationTool/AutomationUtils/ProjectParams.cs
@@ -307,6 +307,10 @@ namespace AutomationTool
 			this.OriginalReleaseVersion = InParams.OriginalReleaseVersion;
 			this.BasedOnReleaseVersion = InParams.BasedOnReleaseVersion;
             this.CreateReleaseVersion = InParams.CreateReleaseVersion;
+			this.ReferenceContainerGlobalFileName = InParams.ReferenceContainerGlobalFileName;
+			this.ReferenceContainerAdditionalPath = InParams.ReferenceContainerAdditionalPath;
+			this.ReferenceContainerChangesCSVFileName = InParams.ReferenceContainerChangesCSVFileName;
+			this.ReferenceContainerCryptoKeys = InParams.ReferenceContainerCryptoKeys;
             this.GeneratePatch = InParams.GeneratePatch;
 			this.AddPatchLevel = InParams.AddPatchLevel;
 			this.DLCFile = InParams.DLCFile;
diff --git a/Engine/Source/Programs/AutomationTool/Scripts/CopyBuildToStagingDirectory.Automation.cs b/Engine/Source/Programs/AutomationTool/Scripts/CopyBuildToStagingDirectory.Automation.cs
--- a/Engine/Source/Programs/AutomationTool/Scripts/CopyBuildToStagingDirectory.Automation.cs
+++ b/Engine/Source/Programs/AutomationTool/Scripts/CopyBuildToStagingDirectory.Automation.cs
@@ -4613,6 +4613,11 @@ namespace AutomationScripts
 					{
 						AdditionalArgs += " -ReferenceContainerChangesCSVFileName=\"" + Params.ReferenceContainerChangesCSVFileName + "\"";
 					}
+					bool bUsePatchInPlaceLayout = false;
+					if (PlatformGameConfig.GetBool("/Script/UnrealEd.ProjectPackagingSettings", "bUsePatchInPlaceLayout", out bUsePatchInPlaceLayout) && bUsePatchInPlaceLayout)
+					{
+						AdditionalArgs += " -UsePatchInPlaceLayout";
+					}
 				}
 
 				if (PlatformGameConfig.GetString("/Script/UnrealEd.ProjectPackagingSettings", "FallbackOrdering", out string FallbackOrdering))
@@ -6038,18 +6043,18 @@ namespace AutomationScripts
 
 			if (!string.IsNullOrEmpty(Params.ReferenceContainerGlobalFileName))
 			{
-				CommandLineBuilder.AppendFormat("-ReferenceContainerGlobalFileName={0}", Params.ReferenceContainerGlobalFileName);
+				CommandLineBuilder.AppendFormat("-ReferenceContainerGlobalFileName={0} ", Params.ReferenceContainerGlobalFileName);
 				if (!string.IsNullOrEmpty(Params.ReferenceContainerAdditionalPath))
 				{
-					CommandLineBuilder.AppendFormat("-ReferenceContainerAdditionalPath={0}", Params.ReferenceContainerAdditionalPath);
+					CommandLineBuilder.AppendFormat("-ReferenceContainerAdditionalPath={0} ", Params.ReferenceContainerAdditionalPath);
 				}
 				if (!string.IsNullOrEmpty(Params.ReferenceContainerCryptoKeys))
 				{
-					CommandLineBuilder.AppendFormat("-ReferenceContainerCryptoKeys={0}", Params.ReferenceContainerCryptoKeys);
+					CommandLineBuilder.AppendFormat("-ReferenceContainerCryptoKeys={0} ", Params.ReferenceContainerCryptoKeys);
 				}
 				if (!string.IsNullOrEmpty(Params.ReferenceContainerChangesCSVFileName))
 				{
-					CommandLineBuilder.AppendFormat("-ReferenceContainerChangesCSVFileName={0}", Params.ReferenceContainerChangesCSVFileName);
+					CommandLineBuilder.AppendFormat("-ReferenceContainerChangesCSVFileName={0} ", Params.ReferenceContainerChangesCSVFileName);
 				}
 			}
 
diff --git a/Engine/Source/Runtime/Core/Internal/IO/IoStore.h b/Engine/Source/Runtime/Core/Internal/IO/IoStore.h
--- a/Engine/Source/Runtime/Core/Internal/IO/IoStore.h
+++ b/Engine/Source/Runtime/Core/Internal/IO/IoStore.h
@@ -307,6 +307,8 @@ struct FIoStoreTocResource
 	[[nodiscard]] CORE_API static TIoStatusOr<uint64> Write(const TCHAR* TocFilePath, FIoStoreTocResource& TocResource, uint32 CompressionBlockSize, uint64 MaxPartitionSize, const FIoContainerSettings& ContainerSettings);
 
 	[[nodiscard]] CORE_API static uint64 HashChunkIdWithSeed(int32 Seed, const FIoChunkId& ChunkId);
+
+	[[nodiscard]] CORE_API FIoStoreTocResourceView AsView() const;
 };
 
 /**
diff --git a/Engine/Source/Runtime/Core/Private/IO/IoStore.cpp b/Engine/Source/Runtime/Core/Private/IO/IoStore.cpp
--- a/Engine/Source/Runtime/Core/Private/IO/IoStore.cpp
+++ b/Engine/Source/Runtime/Core/Private/IO/IoStore.cpp
@@ -309,10 +309,12 @@ static FIoStoreTocChunkInfo GetTocChunkInfoInternal(const T& Toc, int32 TocEntry
 	ChunkInfo.OffsetOnDisk = Toc.CompressionBlocks[FirstBlockIndex].GetOffset();
 	ChunkInfo.CompressedSize = 0;
 	ChunkInfo.PartitionIndex = -1;
+	ChunkInfo.SizeOnDisk = 0;
 	for (int32 BlockIndex = FirstBlockIndex; BlockIndex <= LastBlockIndex; ++BlockIndex)
 	{
 		const FIoStoreTocCompressedBlockEntry& CompressionBlock = Toc.CompressionBlocks[BlockIndex];
 		ChunkInfo.CompressedSize += CompressionBlock.GetCompressedSize();
+		ChunkInfo.SizeOnDisk += Align(CompressionBlock.GetCompressedSize(), FAES::AESBlockSize);
 		if (ChunkInfo.PartitionIndex < 0)
 		{
 			ChunkInfo.PartitionIndex = int32(CompressionBlock.GetOffset() / Toc.Header.PartitionSize);
@@ -1676,6 +1678,24 @@ TIoStatusOr<uint64> FIoStoreTocResource::Write(
 	return TocFileHandle->Tell();
 }
 
+FIoStoreTocResourceView FIoStoreTocResource::AsView() const
+{
+	FIoStoreTocResourceView Result = {};
+	Result.Header = Header;
+	Result.ChunkIds = ChunkIds;
+	Result.ChunkOffsetLengths = ChunkOffsetLengths;
+	Result.ChunkPerfectHashSeeds = ChunkPerfectHashSeeds;
+	Result.ChunkIndicesWithoutPerfectHash = ChunkIndicesWithoutPerfectHash;
+	Result.CompressionBlocks = CompressionBlocks;
+	Result.CompressionMethods = CompressionMethods;
+	Result.SignatureHash = SignatureHash;
+	Result.ChunkBlockSignatures = ChunkBlockSignatures;
+	Result.DirectoryIndexBuffer = DirectoryIndexBuffer;
+	Result.ChunkMetas = ChunkMetas;
+	Result.LegacyChunkMetas = {};
+	return Result;
+}
+
 uint64 FIoStoreTocResource::HashChunkIdWithSeed(int32 Seed, const FIoChunkId& ChunkId)
 {
 	const uint8* Data = ChunkId.GetData();
diff --git a/Engine/Source/Runtime/Core/Public/IO/IoDispatcher.h b/Engine/Source/Runtime/Core/Public/IO/IoDispatcher.h
--- a/Engine/Source/Runtime/Core/Public/IO/IoDispatcher.h
+++ b/Engine/Source/Runtime/Core/Public/IO/IoDispatcher.h
@@ -562,6 +562,7 @@ struct FIoStoreTocChunkInfo
 	uint64 Offset;
 	uint64 OffsetOnDisk;
 	uint64 Size;
+	uint64 SizeOnDisk;
 	uint64 CompressedSize;
 	uint32 NumCompressedBlocks;
 	int32 PartitionIndex;
