Zen manifest generator path fix

FManifestGenerator::AddSingleFile should be able to deal with files that are on different drives

#Synced-CL 50258247

diff --git a/Engine/Source/Developer/IoStoreUtilities/Private/ZenFileSystemManifest.cpp b/Engine/Source/Developer/IoStoreUtilities/Private/ZenFileSystemManifest.cpp
--- a/Engine/Source/Developer/IoStoreUtilities/Private/ZenFileSystemManifest.cpp
+++ b/Engine/Source/Developer/IoStoreUtilities/Private/ZenFileSystemManifest.cpp
@@ -585,8 +585,13 @@ void FManifestGenerator::AddSingleFile(const FString& ClientRoot, const FString&
 	ClientPath = FPaths::Combine(ClientRoot, ClientPath);
 
 	FString ServerRelativePath(LocalFilePath);
-	FPaths::MakePathRelativeTo(ServerRelativePath, *(UnrealEngineRootDir / TEXT("")));
-	ServerRelativePath = FString(TEXT("/")) + ServerRelativePath;
+
+	// MakePathRelative returns false if Target and Source are on different drives. If that's the case 
+	// we use the full path. Otherwise, we make the path relative to the root dir. 
+	if (FPaths::MakePathRelativeTo(ServerRelativePath, *(UnrealEngineRootDir / TEXT(""))))
+	{
+		ServerRelativePath = FString(TEXT("/")) + ServerRelativePath;
+	}
 
 	const FIoChunkId FileChunkId = CreateExternalFileChunkId(ClientPath);
 	Manifest.AddManifestEntry(
diff --git a/Engine/Source/Developer/IoStoreUtilities/Private/ZenStoreWriter.cpp b/Engine/Source/Developer/IoStoreUtilities/Private/ZenStoreWriter.cpp
--- a/Engine/Source/Developer/IoStoreUtilities/Private/ZenStoreWriter.cpp
+++ b/Engine/Source/Developer/IoStoreUtilities/Private/ZenStoreWriter.cpp
@@ -1861,7 +1861,14 @@ void FZenStoreWriter::WriteManifestEntryToPackageWriter(FCbPackage& Pkg, FCbWrit
 	}
 	else
 	{
-		const FString AbsPath = ZenFileSystemManifest->ServerRootPath() / Entry.ServerPath;
+		FString AbsPath = Entry.ServerPath;
+		// The path can be a full path if the file was on another drive than engine's working directory,
+		// if that's the  case ServerPath will be a fullpath, hence we don't prepend server root. 
+		if (FPaths::IsRelative(Entry.ServerPath))
+		{
+			AbsPath = ZenFileSystemManifest->ServerRootPath() / AbsPath;
+		}
+
 		TUniquePtr<FArchive> Reader(IFileManager::Get().CreateFileReader(*AbsPath, 0));
 		if (Reader)
 		{
