comment update

#Synced-CL 50258247

diff --git a/Engine/Source/Developer/IoStoreUtilities/Private/ZenFileSystemManifest.cpp b/Engine/Source/Developer/IoStoreUtilities/Private/ZenFileSystemManifest.cpp
--- a/Engine/Source/Developer/IoStoreUtilities/Private/ZenFileSystemManifest.cpp
+++ b/Engine/Source/Developer/IoStoreUtilities/Private/ZenFileSystemManifest.cpp
@@ -588,6 +588,9 @@ void FManifestGenerator::AddSingleFile(const FString& ClientRoot, const FString&
 
 	// MakePathRelative returns false if Target and Source are on different drives. If that's the case 
 	// we use the full path. Otherwise, we make the path relative to the root dir. 
+	// TODO: Absolute paths are not valid for use in paths on the zenserver , because they are supposed to be 
+	// transferable. We will need to add an ExternalProjectRoot token to ZenServer paths, which will require 
+	// some changes in zenserver to support.
 	if (FPaths::MakePathRelativeTo(ServerRelativePath, *(UnrealEngineRootDir / TEXT(""))))
 	{
 		ServerRelativePath = FString(TEXT("/")) + ServerRelativePath;
diff --git a/Engine/Source/Runtime/StorageServerClient/Private/BuiltInHttpClient/BuiltInHttpClientUDSSocket.cpp b/Engine/Source/Runtime/StorageServerClient/Private/BuiltInHttpClient/BuiltInHttpClientUDSSocket.cpp
--- a/Engine/Source/Runtime/StorageServerClient/Private/BuiltInHttpClient/BuiltInHttpClientUDSSocket.cpp
+++ b/Engine/Source/Runtime/StorageServerClient/Private/BuiltInHttpClient/BuiltInHttpClientUDSSocket.cpp
@@ -10,6 +10,7 @@
 #include <afunix.h>
 
 // MSG_WAITALL is 0x8 on Windows, but 0x100 on other platforms.
+// NOTE: copies low-level TranslateFlags from Sockets\Private\SocketSubsystemBSDPrivate.h.
 inline int TranslateFlags(ESocketReceiveFlags::Type Flags)
 {
 	int TranslatedFlags = 0;
