[ShaderLibrary] Allow shader chunk re-assignment during packaging.

This adds support to assign shaders to their pakchunks during packaging depending on their asset associations.
Previously, this was only handled during cooking, but PakFile rules can override those assignments afterwards.
To support loading individual shaders on demand (and the OnDemandInstallCache backend of the IoStore interface),
shaders need to be assigned to their respective pakchunks that are designated for OnDemand content.

This is implemented by loading all cooked shader archives into one monolithic archive per shader format,
create re-assignd subsets of them for each container based on their asset associations from that container,
and then convert that archive to an IoStoreShaderCodeArchive as before (i.e. grouping of shaders to build IoChunks).

- Add CVar `r.ShaderCodeLibrary.PakChunkOverride` and command line argument `-OverrideShaderChunks` to toggle shader library chunk override;
also added CVar `r.ShaderCodeLibrary.PakChunkOverride.DumpDebugInfo` to enable debug information about shader library pakchunk differences.
- Remove `GShaderPipelineArchiveVersion` and `GetPipelinesArchiveFilename()`:
Shader pipelines are no longer saved separately in the shader code library.
- Wrap file writers into TUniquePtr<FArchive> to reduce manual memory releases.
- Add static function `FSerializedShaderArchive::DecompressShaderEntry()` to generalize shader compression.
- Make `FShaderCodeStats` output optional:
We don't need this during packaging; also `LoadExistingShaderCodeLibrary()` and `CreatePatchLibrary()` didn't use it.
- Refactor `FShaderLibraryFileInfo` to use three components for its name and renamed it to `FShaderLibraryNameInfo`:
This used a two-component "LibraryName" and a single component "PlatformName"; now all three components are separated (ChunkName, ShaderFormatName, ShaderPlatformName).
- Rename `FEditorShaderCodeArchive::FormatName` to `FormatAndPlatformName` to make it clear this is a combined name of shader format and shader platform, e.g. "SF_VULKAN_SM6-VULKAN_SM6".
- Remove `bNeedsDeterministicOrder` from `FEditorShaderCodeArchive` since its no longer used;
Also remove INI entry "[/Script/UnrealEd.ProjectPackagingSettings]:bDeterministicShaderCodeOrder" (UE-198898).
- Move `FEditorShaderCodeArchive` from Private/ into Internal/ folder and replace `WITH_EDITOR` condition with `WITH_EDITORONLY_DATA`;
This is needed in UnrealPak. This does not change module dependency, since UnrealPak already included RenderCore module.
- Move `FShaderMapAssetAssociations` into new file ShaderMapAssetAssociation.h and put it into global scope;
Reduces boilerplate code and doesn't need to be nested into `FSerializedShaderArchive` since it's independent of that class.
- Also move `FShaderMapAssetPaths` type-alias into ShaderMapAssetAssociation.h; Had duplicated definition in Shader.h and ShaderCodeLibrary.h.
- Add `FShaderMapAssetAssociations::Merge()` so that we can replace `PackageNameToShaderMaps` field in `FIoStoreShaderLibraryProcessor`.
- Fix implementation of `FShaderCodeArchive::GetAllShaderIndices()`; Output parameter name shadowing rendered this function pointless.
- Remove unused function `FShaderLibrariesCollection::PackageNativeShaderLibrary()`.
- Don't treat missing target formats in TPM as error anymore; UnrealPak doesn't need audio formats. TPM should not report errors when only shader formats are needed, but both are always initialized.
- Add optional command line argument `-DumpShaderLibraryHashes` to dump all shader maps and shader groups with all their associated shader hashes at runtime when mounting the shader libraries.

#rnx

#Synced-CL 50342528

diff --git a/Engine/Source/Developer/IoStoreUtilities/Private/IoStoreShaderLibraryProcessor.cpp b/Engine/Source/Developer/IoStoreUtilities/Private/IoStoreShaderLibraryProcessor.cpp
--- a/Engine/Source/Developer/IoStoreUtilities/Private/IoStoreShaderLibraryProcessor.cpp
+++ b/Engine/Source/Developer/IoStoreUtilities/Private/IoStoreShaderLibraryProcessor.cpp
@@ -13,6 +13,7 @@
 #include "Serialization/LargeMemoryWriter.h"
 #include "Serialization/MemoryReader.h"
 #include "ShaderCodeArchive.h"
+#include "ShaderLibrary/EditorShaderCodeArchive.h"
 
 
 #define IOSTORE_CPU_SCOPE(NAME) TRACE_CPUPROFILER_EVENT_SCOPE(IoStore##NAME);
@@ -21,7 +22,33 @@
 namespace UE::IoStore::Private
 {
 
-static bool LoadShaderAssetInfo(const FString& Filename, FCookedPackageStore* PackageStore, TMap<FSHAHash, TSet<FName>>& OutShaderCodeToAssets)
+static TAutoConsoleVariable<bool> CVarShaderLibraryPakChunkOverride(
+	TEXT("r.ShaderCodeLibrary.PakChunkOverride"),
+	false,
+	TEXT("If true, shader library chunks are re-assigned with their asset-associations during packaging. Disabled by default.")
+);
+
+static TAutoConsoleVariable<bool> CVarShaderLibraryPakChunkOverrideDumpDebugInfo(
+	TEXT("r.ShaderCodeLibrary.PakChunkOverride.DumpDebugInfo"),
+	true, // @lh-todo - Keep this debug information on initially until the chunk override has been tested enough
+	TEXT("If true, dumps a debug file for each pakchunk with information what shaders have been moved to a different pakchunk.\n")
+	TEXT("Also requires `r.ShaderCodeLibrary.PakChunkOverride=true` or command line argument `-OverrideShaderChunks`.")
+);
+
+// Returns true if shader code chunk override is enabled (CVar 'r.ShaderCodeLibrary.PakChunkOverride').
+static bool IsChunkOverrideEnabled()
+{
+	static bool bIsChunkOverrideEnabledByCmdLine = FParse::Param(FCommandLine::Get(), TEXT("OverrideShaderChunks"));
+	return bIsChunkOverrideEnabledByCmdLine || CVarShaderLibraryPakChunkOverride.GetValueOnAnyThread();
+}
+
+static bool IsChunkOverrideDebugEnabled()
+{
+	return CVarShaderLibraryPakChunkOverrideDumpDebugInfo.GetValueOnAnyThread();
+}
+
+// Loads an *.assetinfo.json text file either from local file system or Zen store.
+static bool LoadShaderAssetInfo(const FString& Filename, FCookedPackageStore* PackageStore, FShaderMapAssetAssociations& OutShaderCodeToAssets)
 {
 	FString JsonText;
 	if (PackageStore && PackageStore->HasZenStoreClient())
@@ -80,77 +107,104 @@ static bool LoadShaderAssetInfo(const FString& Filename, FCookedPackageStore* Pa
 				continue;
 			}
 
-			TSet<FName>& Assets = OutShaderCodeToAssets.Add(ShaderMapHash);
 			TArray<TSharedPtr<FJsonValue>> AssetPathsArray = AssetPathsArrayValue->AsArray();
 			for (int32 IdxAsset = 0, NumAssets = AssetPathsArray.Num(); IdxAsset < NumAssets; ++IdxAsset)
 			{
-				Assets.Add(FName(*AssetPathsArray[IdxAsset]->AsString()));
+				const FName AssetName = FName(*AssetPathsArray[IdxAsset]->AsString());
+				OutShaderCodeToAssets.FindOrAddAsset(AssetName, ShaderMapHash);
 			}
 		}
 	}
 	return true;
 }
 
-struct FShaderLibraryFileInfo
-{
-	FString LibraryName;
-	FName FormatName;
+#if WITH_EDITORONLY_DATA
 
-	// Initializes library and format name from the specified *.ushaderbytecode file.
-	bool ParseFromFilename(const TCHAR* Filename)
+// Loads a *.stinfo binary file either from local file system or Zen store.
+static bool LoadShaderTypeInfo(const FString& Filename, FCookedPackageStore* PackageStore, TArray<FSerializedShaderArchive::FShaderTypeHashes>& OutShaderTypes)
+{
+	TArray<uint8> FileData;
+	if (PackageStore && PackageStore->HasZenStoreClient())
 	{
-		// ShaderArchive-<MyProject>_Chunk<N>-<ShaderFormat>-<ShaderPlatform>.ushaderbytecode
-		// E.g. "ShaderArchive-Lyra_Chunk0-PCD3D_SM6-PCD3D_SM6.ushaderbytecode
-		FStringView BaseFilenameView = FPathViews::GetBaseFilename(Filename);
-		int32 FormatStartIndex = -1;
-		if (!BaseFilenameView.FindLastChar('-', FormatStartIndex))
+		FIoChunkId ChunkId = PackageStore->GetChunkIdFromFileName(Filename);
+		if (!ChunkId.IsValid())
 		{
-			UE_LOG(LogIoStore, Error, TEXT("Invalid shader code library filename '%s'."), Filename);
 			return false;
 		}
-		int32 LibraryNameStartIndex = -1;
-		if (!BaseFilenameView.FindChar('-', LibraryNameStartIndex) || FormatStartIndex - LibraryNameStartIndex <= 1)
+		TIoStatusOr<FIoBuffer> Buffer = PackageStore->ReadChunk(ChunkId);
+		if (!Buffer.IsOk())
 		{
-			UE_LOG(LogIoStore, Error, TEXT("Invalid shader code library filename '%s'."), Filename);
 			return false;
 		}
-		this->LibraryName = FString(BaseFilenameView.Mid(LibraryNameStartIndex + 1, FormatStartIndex - LibraryNameStartIndex - 1));
-		this->FormatName = FName(BaseFilenameView.RightChop(FormatStartIndex + 1));
-		return true;
+		FileData.Append(Buffer.ValueOrDie().GetData(), (int32)Buffer.ValueOrDie().GetSize());
 	}
-
-	// Returns the ShaderAssetInfo JSON filename that belongs to the respective *.ushaderbytecode file.
-	FString GetAssetInfoFilename() const
+	else
 	{
-		//@todo-lh - Should be merged with GetShaderAssetInfoFilename() in ShaderCodeLibrary.cpp
-		return FString::Printf(TEXT("ShaderAssetInfo-%s-%s.assetinfo.json"), *LibraryName, *FormatName.ToString());
+		if (!FFileHelper::LoadFileToArray(FileData, *Filename))
+		{
+			return false;
+		}
 	}
-};
 
-class FSerializedShaderArchiveContainer : public FSerializedShaderArchive
+	FMemoryReaderView MemoryReader(FMemoryView(FileData.GetData(), FileData.NumBytes()));
+	MemoryReader << OutShaderTypes;
+
+	return true;
+}
+
+#endif // WITH_EDITORONLY_DATA
+
+// Stores the shader code of all serialized shaders as a single flat buffer.
+// This is the final format of serialized shaders when they are written to or read from a *.ushaderbytecode file.
+class FSerializedShaderArchiveBuffer : public FSerializedShaderArchive
 {
+	// Pointer to the original target file. This can be re-assigned when the IoStore shader library is generated during packaging.
+	const FContainerTargetFile* TargetFile = nullptr;
 	FIoBuffer LibraryBuffer;
-	int64 OffsetToShaderCode;
+	int64 OffsetToShaderCode = 0;
 
 public:
+	bool LoadFromIoBuffer(FIoBuffer&& InBuffer)
+	{
+		LibraryBuffer = MoveTemp(InBuffer);
+
+		// Serialize shader archive header and store offset to start of shader code payload
+		FLargeMemoryReader LibraryAr(LibraryBuffer.GetData(), LibraryBuffer.GetSize());
+		if (!FSerializedShaderArchive::SerializeHeaderVersion(LibraryAr))
+		{
+			UE_LOG(LogIoStore, Error, TEXT("Failed to serialize header version from ShaderCodeArchive"));
+			return false;
+		}
+		LibraryAr << *this;
+		OffsetToShaderCode = LibraryAr.Tell();
+
+		return true;
+	}
+
 	// Loads the shader library buffer either from package store or from local file.
-	bool Load(const TCHAR* InFilename, FIoChunkId InChunkId, FCookedPackageStore* PackageStore)
+	bool LoadFromTargetFile(const TCHAR* InFilename, const FContainerTargetFile* InTargetFile, FCookedPackageStore* PackageStore)
 	{
-		if (PackageStore && PackageStore->HasZenStoreClient() && InChunkId.IsValid())
+		check(InTargetFile);
+		TargetFile = InTargetFile;
+
+		if (PackageStore && PackageStore->HasZenStoreClient() && InTargetFile->ChunkId.IsValid())
 		{
-			TIoStatusOr<FIoBuffer> ReadChunkResult = PackageStore->ReadChunk(InChunkId);
+			// Load serialized shaders from IoChunk with Id that was generated from an incremental cook.
+			// We don't have to deal with generating this ChunkId here, it's just an opaque token that was assigned to the target file.
+			TIoStatusOr<FIoBuffer> ReadChunkResult = PackageStore->ReadChunk(InTargetFile->ChunkId);
 			if (!ReadChunkResult.IsOk())
 			{
 				UE_LOG(LogIoStore, Error,
 					TEXT("Missing shader code libary ChunkId. ChunkId = '%s'. InFilename = '%s'. IOStatus = '%s'."),
-					*LexToString(InChunkId), InFilename ? InFilename : TEXT("<nullptr>"),
+					*LexToString(InTargetFile->ChunkId), InFilename ? InFilename : TEXT("<nullptr>"),
 					*ReadChunkResult.Status().ToString());
 				return false;
 			}
-			LibraryBuffer = ReadChunkResult.ConsumeValueOrDie();
+			return LoadFromIoBuffer(ReadChunkResult.ConsumeValueOrDie());
 		}
 		else
 		{
+			// Load serialized shaders from local file.
 			TUniquePtr<FArchive> LibraryAr(IFileManager::Get().CreateFileReader(InFilename));
 			if (!LibraryAr)
 			{
@@ -158,39 +212,45 @@ public:
 				return false;
 			}
 
-			LibraryBuffer = FIoBuffer(LibraryAr->TotalSize());
-			LibraryAr->Serialize(LibraryBuffer.GetData(), LibraryBuffer.GetSize());
+			FIoBuffer IntermediateLibraryBuffer(LibraryAr->TotalSize());
+			LibraryAr->Serialize(IntermediateLibraryBuffer.GetData(), IntermediateLibraryBuffer.GetSize());
+			return LoadFromIoBuffer(MoveTemp(IntermediateLibraryBuffer));
 		}
+	}
 
-		// Serialize shader archive header and store offset to start of shader code payload
-		FLargeMemoryReader LibraryAr(LibraryBuffer.GetData(), LibraryBuffer.GetSize());
-		uint32 Version;
-		LibraryAr << Version;
-
-		LibraryAr << *this;
-		OffsetToShaderCode = LibraryAr.Tell();
+#if WITH_EDITORONLY_DATA
+	bool LoadFromEditorArchive(FEditorShaderCodeArchive& ShaderCodeArchive)
+	{
+		// Copy editor archive to this archive
+		TArray<uint8> FlatShaderCode;
+		bool bRanOutOfRoom = false;
+		constexpr int64 kNoShaderSizeLimit = -1;
+		constexpr int64 kNoShaderCountLimit = -1;
+		ShaderCodeArchive.CopyToArchiveAndClear(*this, FlatShaderCode, bRanOutOfRoom, kNoShaderSizeLimit, kNoShaderCountLimit, ECookShaderLibrarySource::CurrentCook);
+		check(!bRanOutOfRoom);
+
+		// Write header
+		FLargeMemoryWriter HeaderWriter;
+		const bool bSerializeVersionSucceeded = FSerializedShaderArchive::SerializeHeaderVersion(HeaderWriter);
+		check(bSerializeVersionSucceeded);
+		HeaderWriter << *this;
+		
+		// Convert to flat IoBuffer
+		LibraryBuffer = FIoBuffer(HeaderWriter.TotalSize() + FlatShaderCode.NumBytes());
+		OffsetToShaderCode = HeaderWriter.TotalSize();
+
+		FMemory::Memcpy(LibraryBuffer.GetData(), HeaderWriter.GetData(), HeaderWriter.TotalSize());
+		FMemory::Memcpy(LibraryBuffer.GetData() + OffsetToShaderCode, FlatShaderCode.GetData(), FlatShaderCode.NumBytes());
 
 		return true;
 	}
+#endif
 
-	// Loads and decompresses the specified shader into the destination buffer and returns the advanced pointer of the destination after the new copy.
-	uint8* LoadAndDecompressShader(uint8* Destination, int32 ShaderIndex) const
+	// Loads and decompresses the specified shader and appends it to the end of the output container.
+	void DecompressShaderFromLibraryBuffer(int32 ShaderIndex, TArray<uint8>& OutDecompressedShader) const
 	{
-		const FShaderCodeEntry& IndividuallyCompressedShader = ShaderEntries[ShaderIndex];
-		uint32 UncompressedSize = IndividuallyCompressedShader.GetUncompressedSize();
-		{
-			// small shaders might be stored without compression, handle them here
-			if (IndividuallyCompressedShader.Size == UncompressedSize)
-			{
-				FMemory::Memcpy(Destination, GetShaderCodeDataStart() + IndividuallyCompressedShader.Offset, IndividuallyCompressedShader.Size);
-			}
-			else
-			{
-				// This function will crash if decompression fails.
-				ShaderCodeArchive::DecompressShaderWithOodle(Destination, UncompressedSize, GetShaderCodeDataStart() + IndividuallyCompressedShader.Offset, IndividuallyCompressedShader.Size);
-			}
-		}
-		return Destination + UncompressedSize;
+		const FShaderCodeEntry& Entry = ShaderEntries[ShaderIndex];
+		FSerializedShaderArchive::DecompressShaderEntryAndAppend(Entry, GetShaderCodeDataStart() + Entry.Offset, OutDecompressedShader);
 	}
 
 	// Returns a pointer to the start of shader code payload to read from.
@@ -198,53 +258,217 @@ public:
 	{
 		return LibraryBuffer.GetData() + OffsetToShaderCode;
 	}
+
+	// Returns an array view of the shader code payload.
+	TConstArrayView64<uint8> GetShaderCodeView() const
+	{
+		return TConstArrayView64<uint8>(GetShaderCodeDataStart(), static_cast<int64>(LibraryBuffer.GetSize()) - OffsetToShaderCode);
+	}
+
+	// Returns the container target file this archive was loaded from.
+	FORCEINLINE const FContainerTargetFile* GetTargetFile() const
+	{
+		return TargetFile;
+	}
 };
 
-// Shader map hash with its IoChunk Ids of all its shader groups
-struct FShaderGroupIoChunkMap
+static void ValidateAndLogShaderGroupStats(
+	const FShaderLibraryNameInfo& LibraryNameInfo, const FIoStoreShaderCodeArchiveHeader& IoStoreLibraryHeader, const FSerializedShaderArchive& SerializedShaders)
 {
-	FSHAHash ShaderMapHash;
-	TArray<FIoChunkId> ShaderGroupChunkIds;
-};
+	int64 TotalUncompressedSizeViaGroups = 0; // calculate total uncompressed size twice for a sanity check
+	int64 TotalUncompressedSizeViaShaders = 0;
+	int64 TotalIndividuallyCompressedSize = 0;
+	int64 TotalGroupCompressedSize = 0;
 
-static bool ConvertToIoStoreShaderLibrary(
-	const TCHAR* InFilename,
-	FIoChunkId InChunkId,
-	FCookedPackageStore* PackageStore,
-	FOodleDataCompression::ECompressor InShaderOodleCompressor,
-	FOodleDataCompression::ECompressionLevel InShaderOodleLevel,
-	TTuple<FIoChunkId, FIoBuffer>& OutLibraryIoChunk,
-	TArray<FShaderGroupIoChunk>& OutCodeIoChunks,
-	TArray<FShaderGroupIoChunkMap>& OutShaderMaps,
-	TArray<TTuple<FSHAHash, FShaderMapAssetPaths>>& OutShaderMapAssetAssociations)
+	for (int32 IdxGroup = 0, NumGroups = IoStoreLibraryHeader.ShaderGroupEntries.Num(); IdxGroup < NumGroups; ++IdxGroup)
+	{
+		const FIoStoreShaderGroupEntry& Group = IoStoreLibraryHeader.ShaderGroupEntries[IdxGroup];
+		TotalGroupCompressedSize += Group.CompressedSize;
+		TotalUncompressedSizeViaGroups += Group.UncompressedSize;
+
+		// now go via shaders
+		for (uint32 IdxShaderInGroup = 0; IdxShaderInGroup < Group.NumShaders; ++IdxShaderInGroup)
+		{
+			int32 ShaderIndex = IoStoreLibraryHeader.ShaderIndices[Group.ShaderIndicesOffset + IdxShaderInGroup];
+			const FShaderCodeEntry& IndividuallyCompressedShader = SerializedShaders.ShaderEntries[ShaderIndex];
+
+			TotalIndividuallyCompressedSize += IndividuallyCompressedShader.Size;
+			TotalUncompressedSizeViaShaders += IndividuallyCompressedShader.GetUncompressedSize();
+		}
+	}
+
+	checkf(TotalUncompressedSizeViaGroups == TotalUncompressedSizeViaShaders,
+		TEXT("Sanity check failure: total uncompressed shader size differs if calculated via shader groups (%lld) or individual shaders (%lld)"),
+		TotalUncompressedSizeViaGroups, TotalUncompressedSizeViaShaders
+	);
+
+	UE_LOG(LogIoStore, Display,
+		TEXT("%s(%s): Recompressed %d shaders as %d groups. Library size changed from %lld KB (%.2f : 1 ratio) to %lld KB (%.2f : 1 ratio), %.2f%% of previous."),
+		*LibraryNameInfo.GetLibraryName(),
+		*LibraryNameInfo.ShaderPlatformName.ToString(),
+		SerializedShaders.ShaderEntries.Num(),
+		IoStoreLibraryHeader.ShaderGroupEntries.Num(),
+		TotalIndividuallyCompressedSize / 1024,
+		TotalIndividuallyCompressedSize > 0 ? static_cast<double>(TotalUncompressedSizeViaShaders) / static_cast<double>(TotalIndividuallyCompressedSize) : 0.0,
+		TotalGroupCompressedSize / 1024,
+		TotalGroupCompressedSize > 0 ? static_cast<double>(TotalUncompressedSizeViaShaders) / static_cast<double>(TotalGroupCompressedSize) : 0.0,
+		TotalIndividuallyCompressedSize > 0 ? 100.0 * static_cast<double>(TotalGroupCompressedSize) / static_cast<double>(TotalIndividuallyCompressedSize) : 0.0
+	);
+}
+
+static bool ParseIoStoreShaderLibraryName(FShaderLibraryNameInfo& OutLibraryNameInfo, const TCHAR* InFilename)
 {
-	IOSTORE_CPU_SCOPE(ConvertShaderLibrary);
+	if (!OutLibraryNameInfo.ParseFromFilename(InFilename))
+	{
+		UE_LOG(LogIoStore, Error, TEXT("Invalid shader code library filename '%s'."), InFilename);
+		return false;
+	}
+	return true;
+}
 
-	FShaderLibraryFileInfo LibraryFileInfo;
-	if (!LibraryFileInfo.ParseFromFilename(InFilename))
+bool FIoStoreShaderLibraryProcessor::LoadSerializedShaders(
+	const FContainerTargetFile& TargetFile, FSerializedShaderArchiveBuffer& OutSerializedShaders, bool bIncludeTypeInfo, FShaderLibraryNameInfo* OutNameInfo)
+{
+	const TCHAR* InFilename = *TargetFile.NormalizedSourcePath;
+
+	// Convert shader format and platform name from *.ushaderbytecode filename
+	FShaderLibraryNameInfo LibraryNameInfo;
+	if (!ParseIoStoreShaderLibraryName(LibraryNameInfo, InFilename))
 	{
 		return false;
 	}
 
-	FSerializedShaderArchiveContainer SerializedShaders;
-	if (!SerializedShaders.Load(InFilename, InChunkId, PackageStore))
+	if (OutNameInfo)
+	{
+		*OutNameInfo = LibraryNameInfo;
+	}
+
+	// If this is a proxy target file, don't load anything and just return the output library name
+	const bool bIsProxyTargetFile = FPathViews::GetPath(InFilename).IsEmpty();
+	if (bIsProxyTargetFile)
+	{
+		return true;
+	}
+
+	// Load serialized shaders from *.ushaderbytecode file
+	FCookedPackageStore* ActivePackageStore = PackageStore;
+	if (!OutSerializedShaders.LoadFromTargetFile(InFilename, &TargetFile, PackageStore))
 	{
 		return false;
 	}
 
-	const FString AssetInfoFilename = FPaths::GetPath(InFilename) / LibraryFileInfo.GetAssetInfoFilename();
-	TMap<FSHAHash, FShaderMapAssetPaths> ShaderCodeToAssets;
-	if (!LoadShaderAssetInfo(AssetInfoFilename, PackageStore && InChunkId.IsValid() ? PackageStore : nullptr, ShaderCodeToAssets))
+	if (PackageStore && !TargetFile.ChunkId.IsValid())
+	{
+		PackageStore = nullptr;
+	}
+
+	// Load *.assetinfo.json file for the specified input shader library and store its content in the ShaderMapAssetAssociations container
+	{
+		const FString AssetInfoFilename = FPaths::GetPath(InFilename) / LibraryNameInfo.GetAssetInfoFilename();
+		if (!LoadShaderAssetInfo(AssetInfoFilename, PackageStore, OutSerializedShaders.GetShaderMapAssetAssociations()))
+		{
+			UE_LOG(LogIoStore, Error, TEXT("Failed loading shader asset info file '%s'"), *AssetInfoFilename);
+			return false;
+		}
+	}
+
+#if WITH_EDITORONLY_DATA
+	// Load *.stinfo file for the specified input shader library
+	if (bIncludeTypeInfo)
+	{
+		const FString TypeInfoFilename = FPaths::GetPath(InFilename) / LibraryNameInfo.GetTypeInfoFilename();
+		if (!LoadShaderTypeInfo(TypeInfoFilename, PackageStore, OutSerializedShaders.ShaderTypes))
+		{
+			UE_LOG(LogIoStore, Error, TEXT("Failed loading shader type info file '%s'"), *TypeInfoFilename);
+			return false;
+		}
+	}
+#endif
+
+	return true;
+}
+
+static TArray<uint8> DecompressGroupedShaders(
+	const FIoStoreShaderCodeArchiveHeader& IoStoreLibraryHeader,
+	const FIoStoreShaderGroupEntry& Group,
+	const FSerializedShaderArchiveBuffer& SerializedShaders,
+	int32& OutSmallestShaderInGroupByOrdinal)
+{
+	TArray<uint8> UncompressedGroupMemory;
+	UncompressedGroupMemory.Reserve(Group.UncompressedSize);
+	const uint64 UncompressedGroupMemoryPtr = reinterpret_cast<uint64>(UncompressedGroupMemory.GetData());
+
+	int32 SmallestShaderInGroupByOrdinal = MAX_int32;
+
+	{
+		// Not worth to special-case for group size of 1 (by converting to copy) - such groups are very fast to decompress and recompress.
+		const uint8* UncompressedShaderStart = UncompressedGroupMemory.GetData();
+		for (uint32 IdxShaderInGroup = 0; IdxShaderInGroup < Group.NumShaders; ++IdxShaderInGroup)
+		{
+			const int32 ShaderIndex = IoStoreLibraryHeader.ShaderIndices[Group.ShaderIndicesOffset + IdxShaderInGroup];
+			SmallestShaderInGroupByOrdinal = FMath::Min(SmallestShaderInGroupByOrdinal, ShaderIndex);
+			const FIoStoreShaderCodeEntry& Shader = IoStoreLibraryHeader.ShaderEntries[ShaderIndex];
+
+			checkf((reinterpret_cast<uint64>(UncompressedShaderStart) - UncompressedGroupMemoryPtr) == Shader.UncompressedOffsetInGroup,
+				TEXT("Shader uncompressed offset in group does not agree with its actual placement: Shader.UncompressedOffsetInGroup=%llu, actual=%llu"),
+				Shader.UncompressedOffsetInGroup, (reinterpret_cast<uint64>(UncompressedShaderStart) - UncompressedGroupMemoryPtr)
+			);
+
+			// Load and decompress the shader at the desired offset
+			SerializedShaders.DecompressShaderFromLibraryBuffer(ShaderIndex, UncompressedGroupMemory);
+			UncompressedShaderStart = UncompressedGroupMemory.GetData() + UncompressedGroupMemory.Num();
+		}
+
+		checkf((reinterpret_cast<uint64>(UncompressedShaderStart) - UncompressedGroupMemoryPtr) == Group.UncompressedSize,
+			TEXT("Uncompressed shader group size does not agree with the actual results (Group.UncompressedSize=%llu, actual=%llu)"),
+			Group.UncompressedSize, (reinterpret_cast<uint64>(UncompressedShaderStart) - UncompressedGroupMemoryPtr));
+	}
+
+	OutSmallestShaderInGroupByOrdinal = SmallestShaderInGroupByOrdinal;
+
+	return UncompressedGroupMemory;
+}
+
+static TArray<uint8> CompressShaderGroup(
+	const uint8* UncompressedGroupMemory,
+	const FIoStoreShaderGroupEntry& Group,
+	FOodleDataCompression::ECompressor InOodleCompressor,
+	FOodleDataCompression::ECompressionLevel InOodleLevel)
+{
+	// Get size estimate for compressed group to allocate memory
+	int64 CompressedGroupSize = 0;
+	ShaderCodeArchive::CompressShaderWithOodle(nullptr, CompressedGroupSize, UncompressedGroupMemory, Group.UncompressedSize, InOodleCompressor, InOodleLevel);
+	checkf(CompressedGroupSize > 0, TEXT("CompressedGroupSize estimate seems wrong (%lld)"), CompressedGroupSize);
+	checkf(CompressedGroupSize <= MAX_int32, TEXT("CompressedGroupSize exceeded limit of signed integer (%lld)"), CompressedGroupSize);
+
+	// Now compress shader group
+	TArray<uint8> CompressedGroupMemory;
+	CompressedGroupMemory.SetNumUninitialized(static_cast<int32>(CompressedGroupSize));
+	const bool bCompressed = ShaderCodeArchive::CompressShaderWithOodle(CompressedGroupMemory.GetData(), CompressedGroupSize, UncompressedGroupMemory, Group.UncompressedSize, InOodleCompressor, InOodleLevel);
+	checkf(bCompressed, TEXT("We could not compress the shader group after providing an estimated memory."));
+
+	// Now reduce the output memory to its actual compressed size, since we had to over estimate the allocation. We don't have to shrink the capacity to avoid re-allocating memory.
+	CompressedGroupMemory.SetNum(CompressedGroupSize, EAllowShrinking::No);
+
+	return CompressedGroupMemory;
+}
+
+bool FIoStoreShaderLibraryProcessor::ConvertToIoStoreShaderLibrary(FContainerTargetFile& TargetFile, const FSerializedShaderArchiveBuffer& SerializedShaders)
+{
+	const TCHAR* InFilename = *TargetFile.NormalizedSourcePath;
+
+	// Convert shader format and platform name from *.ushaderbytecode filename
+	FShaderLibraryNameInfo LibraryNameInfo;
+	if (!ParseIoStoreShaderLibraryName(LibraryNameInfo, InFilename))
 	{
-		UE_LOG(LogIoStore, Error, TEXT("Failed loading asset asset info file '%s'"), *AssetInfoFilename);
 		return false;
 	}
 
-	TRACE_CPUPROFILER_EVENT_SCOPE_TEXT(*LibraryFileInfo.LibraryName);
+	TRACE_CPUPROFILER_EVENT_SCOPE_TEXT(*LibraryNameInfo.GetLibraryName());
 
 	// Create IoStore code archived for all serialized shaders. This will organize the shaders into groups.
 	FIoStoreShaderCodeArchiveHeader IoStoreLibraryHeader;
-	FIoStoreShaderCodeArchive::CreateIoStoreShaderCodeArchiveHeader(LibraryFileInfo.FormatName, SerializedShaders, IoStoreLibraryHeader);
+	FIoStoreShaderCodeArchive::CreateIoStoreShaderCodeArchiveHeader(LibraryNameInfo.ShaderPlatformName, SerializedShaders, IoStoreLibraryHeader);
 	checkf(SerializedShaders.ShaderEntries.Num() == IoStoreLibraryHeader.ShaderEntries.Num(), TEXT("IoStore header has different number of shaders (%d) than the original library (%d)"),
 		IoStoreLibraryHeader.ShaderEntries.Num(), SerializedShaders.ShaderEntries.Num());
 	checkf(SerializedShaders.ShaderMapEntries.Num() == IoStoreLibraryHeader.ShaderMapEntries.Num(), TEXT("IoStore header has different number of shadermaps (%d) than the original library (%d)"),
@@ -252,51 +476,25 @@ static bool ConvertToIoStoreShaderLibrary(
 
 	// Load all the shaders, decompress them and recompress into groups. Each code chunk is shader group chunk now.
 	// This also updates group's compressed size, the only value left to calculate.
-	int32 GroupCount = IoStoreLibraryHeader.ShaderGroupEntries.Num();
-	OutCodeIoChunks.SetNum(GroupCount);
-	FCriticalSection DiskArchiveAccess;
-	ParallelFor(OutCodeIoChunks.Num(),
-		[&OutCodeIoChunks, &IoStoreLibraryHeader, &DiskArchiveAccess, &SerializedShaders, InShaderOodleCompressor, InShaderOodleLevel](int32 GroupIndex)
+	const int32 GroupCount = IoStoreLibraryHeader.ShaderGroupEntries.Num();
+	CodeIoChunks.SetNum(GroupCount);
+
+	ParallelFor(CodeIoChunks.Num(),
+		[this, &IoStoreLibraryHeader, &SerializedShaders](int32 GroupIndex)
 		{
 			FIoStoreShaderGroupEntry& Group = IoStoreLibraryHeader.ShaderGroupEntries[GroupIndex];
-			uint8* UncompressedGroupMemory = reinterpret_cast<uint8*>(FMemory::Malloc(Group.UncompressedSize));
-			check(UncompressedGroupMemory);
 
+			// Decompress all individual shaders of the current group
 			int32 SmallestShaderInGroupByOrdinal = MAX_int32;
-			uint8* ShaderStart = UncompressedGroupMemory;
-			// Not worth to special-case for group size of 1 (by converting to copy) - such groups are very fast to decompress and recompress.
-			for (uint32 IdxShaderInGroup = 0; IdxShaderInGroup < Group.NumShaders; ++IdxShaderInGroup)
-			{
-				const int32 ShaderIndex = IoStoreLibraryHeader.ShaderIndices[Group.ShaderIndicesOffset + IdxShaderInGroup];
-				SmallestShaderInGroupByOrdinal = FMath::Min(SmallestShaderInGroupByOrdinal, ShaderIndex);
-				const FIoStoreShaderCodeEntry& Shader = IoStoreLibraryHeader.ShaderEntries[ShaderIndex];
+			const TArray<uint8> UncompressedGroupMemory = DecompressGroupedShaders(IoStoreLibraryHeader, Group, SerializedShaders, SmallestShaderInGroupByOrdinal);
 
-				checkf((reinterpret_cast<uint64>(ShaderStart) - reinterpret_cast<uint64>(UncompressedGroupMemory)) == Shader.UncompressedOffsetInGroup,
-					TEXT("Shader uncompressed offset in group does not agree with its actual placement: Shader.UncompressedOffsetInGroup=%llu, actual=%llu"),
-					Shader.UncompressedOffsetInGroup, (reinterpret_cast<uint64>(ShaderStart) - reinterpret_cast<uint64>(UncompressedGroupMemory))
-				);
-
-				// load and decompress the shader at the desired offset
-				ShaderStart = SerializedShaders.LoadAndDecompressShader(ShaderStart, ShaderIndex);
-			}
-
-			checkf((reinterpret_cast<uint64>(ShaderStart) - reinterpret_cast<uint64>(UncompressedGroupMemory)) == Group.UncompressedSize,
-				TEXT("Uncompressed shader group size does not agree with the actual results (Group.UncompressedSize=%llu, actual=%llu)"),
-				Group.UncompressedSize, (reinterpret_cast<uint64>(ShaderStart) - reinterpret_cast<uint64>(UncompressedGroupMemory)));
-
-			// now compress the whole group
-			int64 CompressedGroupSize = 0;
-			uint8* CompressedShaderGroupMemory = nullptr;
-			ShaderCodeArchive::CompressShaderWithOodle(nullptr, CompressedGroupSize, UncompressedGroupMemory, Group.UncompressedSize, InShaderOodleCompressor, InShaderOodleLevel);
-			checkf(CompressedGroupSize > 0, TEXT("CompressedGroupSize estimate seems wrong (%lld)"), CompressedGroupSize);
-
-			CompressedShaderGroupMemory = reinterpret_cast<uint8*>(FMemory::Malloc(CompressedGroupSize));
-			bool bCompressed = ShaderCodeArchive::CompressShaderWithOodle(CompressedShaderGroupMemory, CompressedGroupSize, UncompressedGroupMemory, Group.UncompressedSize, InShaderOodleCompressor, InShaderOodleLevel);
-			checkf(bCompressed, TEXT("We could not compress the shader group after providing an estimated memory."));
+			// Now compress the whole group
+			const TArray<uint8> CompressedGroupMemory = CompressShaderGroup(UncompressedGroupMemory.GetData(), Group, this->ShaderOodleCompressor, this->ShaderOodleLevel);
+			const uint32 CompressedGroupMemorySize = static_cast<uint32>(CompressedGroupMemory.Num());
 
 			// Write out IoChunk for current shader group
 			{
-				FShaderGroupIoChunk& OutCodeIoChunk = OutCodeIoChunks[GroupIndex];
+				FShaderGroupIoChunk& OutCodeIoChunk = this->CodeIoChunks[GroupIndex];
 				OutCodeIoChunk.ChunkId = IoStoreLibraryHeader.ShaderGroupIoHashes[GroupIndex];
 				// This value is the load order factor for the group (the smaller, the more likely), that IoStore will use to sort the chunks.
 				// We calculate it as the smallest shader index in the group, which shouldn't be bad. Arguably a better way would be to get the lowest-numbered shadermap that references _any_ shader in the group, 
@@ -306,85 +504,55 @@ static bool ConvertToIoStoreShaderLibrary(
 				OutCodeIoChunk.LoadOrderFactor = static_cast<uint32>(SmallestShaderInGroupByOrdinal);
 
 				// Copy shader group payload into IoBuffer
-				if (CompressedGroupSize < Group.UncompressedSize)
+				if (CompressedGroupMemorySize < Group.UncompressedSize)
 				{
-					OutCodeIoChunk.CodeIoBuffer = FIoBuffer(CompressedGroupSize);
-					FMemory::Memcpy(OutCodeIoChunk.CodeIoBuffer.GetData(), CompressedShaderGroupMemory, CompressedGroupSize);
-					Group.CompressedSize = CompressedGroupSize;
+					OutCodeIoChunk.CodeIoBuffer = FIoBuffer(CompressedGroupMemorySize);
+					FMemory::Memcpy(OutCodeIoChunk.CodeIoBuffer.GetData(), CompressedGroupMemory.GetData(), CompressedGroupMemorySize);
+					Group.CompressedSize = CompressedGroupMemorySize;
 				}
 				else
 				{
 					// store uncompressed (unlikely, but happens for a 200-byte sized shader that happens to get its own group)
 					OutCodeIoChunk.CodeIoBuffer = FIoBuffer(Group.UncompressedSize);
-					FMemory::Memcpy(OutCodeIoChunk.CodeIoBuffer.GetData(), UncompressedGroupMemory, Group.UncompressedSize);
+					FMemory::Memcpy(OutCodeIoChunk.CodeIoBuffer.GetData(), UncompressedGroupMemory.GetData(), Group.UncompressedSize);
 					Group.CompressedSize = Group.UncompressedSize;
 				}
 			}
-
-			FMemory::Free(UncompressedGroupMemory);
-			FMemory::Free(CompressedShaderGroupMemory);
 		},
 		EParallelForFlags::Unbalanced
 	);
 
 	// Calculate and log stats
-	int64 TotalUncompressedSizeViaGroups = 0;	// calculate total uncompressed size twice for a sanity check
-	int64 TotalUncompressedSizeViaShaders = 0;
-	int64 TotalIndividuallyCompressedSize = 0;
-	int64 TotalGroupCompressedSize = 0;
-	for (int32 IdxGroup = 0, NumGroups = IoStoreLibraryHeader.ShaderGroupEntries.Num(); IdxGroup < NumGroups; ++IdxGroup)
-	{
-		FIoStoreShaderGroupEntry& Group = IoStoreLibraryHeader.ShaderGroupEntries[IdxGroup];
-		TotalGroupCompressedSize += Group.CompressedSize;
-		TotalUncompressedSizeViaGroups += Group.UncompressedSize;
+	ValidateAndLogShaderGroupStats(LibraryNameInfo, IoStoreLibraryHeader, SerializedShaders);
 
-		// now go via shaders
-		for (uint32 IdxShaderInGroup = 0; IdxShaderInGroup < Group.NumShaders; ++IdxShaderInGroup)
-		{
-			int32 ShaderIndex = IoStoreLibraryHeader.ShaderIndices[Group.ShaderIndicesOffset + IdxShaderInGroup];
-			const FShaderCodeEntry& IndividuallyCompressedShader = SerializedShaders.ShaderEntries[ShaderIndex];
-
-			TotalIndividuallyCompressedSize += IndividuallyCompressedShader.Size;
-			TotalUncompressedSizeViaShaders += IndividuallyCompressedShader.GetUncompressedSize();
-		}
+	// Override IoChunk for shader library header.
+	// The old ChunkId was for serialized shaders, but now we store FIoStoreShaderCodeArchive for this target file, so use a new ChunkId.
+	// The name that generates the ChunkId (a hash rather) needs to match an exact name pattern that is later reconstructed when the pakchunk is mounted.
+	// See FNamedShaderLibrary::OnPakFileMounted()
+	{
+		FLargeMemoryWriter IoStoreLibraryAr(0, true);
+		FIoStoreShaderCodeArchive::SaveIoStoreShaderCodeArchive(IoStoreLibraryHeader, IoStoreLibraryAr);
+		TargetFile.ChunkId = FIoStoreShaderCodeArchive::GetShaderCodeArchiveChunkId(LibraryNameInfo.ChunkName, FName(LibraryNameInfo.GetFormatAndPlatformName()));
+		TargetFile.SourceSize = static_cast<uint64>(IoStoreLibraryAr.TotalSize());
+		TargetFile.SourceBuffer.Emplace(FIoBuffer(FIoBuffer::AssumeOwnership, IoStoreLibraryAr.ReleaseOwnership(), TargetFile.SourceSize));
+
+		// Log the shader code archive name with its ChunkId and its target filename, so we can match it with the shader library at runtime if needed.
+		UE_LOG(LogIoStore, Display, TEXT("Save IoStoreShaderCodeArchive '%s-%s' with ChunkId '%s' for target file '%s'"),
+			*LibraryNameInfo.ChunkName, *LibraryNameInfo.GetFormatAndPlatformName(), *LexToString(TargetFile.ChunkId), *TargetFile.NormalizedSourcePath);
 	}
 
-	checkf(TotalUncompressedSizeViaGroups == TotalUncompressedSizeViaShaders,
-		TEXT("Sanity check failure: total uncompressed shader size differs if calculated via shader groups (%lld) or individual shaders (%lld)"),
-		TotalUncompressedSizeViaGroups, TotalUncompressedSizeViaShaders
-	);
-
-	UE_LOG(LogIoStore, Display,
-		TEXT("%s(%s): Recompressed %d shaders as %d groups. Library size changed from %lld KB (%.2f : 1 ratio) to %lld KB (%.2f : 1 ratio), %.2f%% of previous."),
-		*LibraryFileInfo.LibraryName,
-		*LibraryFileInfo.FormatName.ToString(),
-		SerializedShaders.ShaderEntries.Num(),
-		IoStoreLibraryHeader.ShaderGroupEntries.Num(),
-		TotalIndividuallyCompressedSize / 1024,
-		static_cast<double>(TotalUncompressedSizeViaShaders) / static_cast<double>(TotalIndividuallyCompressedSize),
-		TotalGroupCompressedSize / 1024,
-		static_cast<double>(TotalUncompressedSizeViaShaders) / static_cast<double>(TotalGroupCompressedSize),
-		100.0 * static_cast<double>(TotalGroupCompressedSize) / static_cast<double>(TotalIndividuallyCompressedSize)
-	);
-
-	FLargeMemoryWriter IoStoreLibraryAr(0, true);
-	FIoStoreShaderCodeArchive::SaveIoStoreShaderCodeArchive(IoStoreLibraryHeader, IoStoreLibraryAr);
-	OutLibraryIoChunk.Key = FIoStoreShaderCodeArchive::GetShaderCodeArchiveChunkId(LibraryFileInfo.LibraryName, LibraryFileInfo.FormatName);
-	int64 TotalSize = IoStoreLibraryAr.TotalSize();
-	OutLibraryIoChunk.Value = FIoBuffer(FIoBuffer::AssumeOwnership, IoStoreLibraryAr.ReleaseOwnership(), TotalSize);
-
 	// Return shadermaps: own hash and array of shader group IoChunk Ids
 	{
-		int32 ShaderMapCount = IoStoreLibraryHeader.ShaderMapHashes.Num();
-		OutShaderMaps.SetNum(ShaderMapCount);
-		ParallelFor(OutShaderMaps.Num(),
-			[&OutShaderMaps, &IoStoreLibraryHeader](int32 ShaderMapIndex)
+		const int32 ShaderMapCount = IoStoreLibraryHeader.ShaderMapHashes.Num();
+		CodeIoChunkMaps.SetNum(ShaderMapCount);
+		ParallelFor(CodeIoChunkMaps.Num(),
+			[this, &IoStoreLibraryHeader](int32 ShaderMapIndex)
 			{
-				FShaderGroupIoChunkMap& OutShaderMap = OutShaderMaps[ShaderMapIndex];
+				FShaderGroupIoChunkMap& OutShaderMap = this->CodeIoChunkMaps[ShaderMapIndex];
 				OutShaderMap.ShaderMapHash = IoStoreLibraryHeader.ShaderMapHashes[ShaderMapIndex];
 				const FIoStoreShaderMapEntry& ShaderMapEntry = IoStoreLibraryHeader.ShaderMapEntries[ShaderMapIndex];
 				int32 LookupIndexEnd = ShaderMapEntry.ShaderIndicesOffset + ShaderMapEntry.NumShaders;
-				OutShaderMap.ShaderGroupChunkIds.Reserve(ShaderMapEntry.NumShaders);	// worst-case, 1 group == 1 shader
+				OutShaderMap.ShaderGroupChunkIds.Reserve(ShaderMapEntry.NumShaders); // worst-case, 1 group == 1 shader
 				for (int32 ShaderLookupIndex = ShaderMapEntry.ShaderIndicesOffset; ShaderLookupIndex < LookupIndexEnd; ++ShaderLookupIndex)
 				{
 					int32 ShaderIndex = IoStoreLibraryHeader.ShaderIndices[ShaderLookupIndex];
@@ -396,37 +564,212 @@ static bool ConvertToIoStoreShaderLibrary(
 		);
 	}
 
-	// Return associations between shader maps and all their asset paths
-	OutShaderMapAssetAssociations.Reserve(ShaderCodeToAssets.Num());
-	for (auto& KV : ShaderCodeToAssets)
+	return true;
+}
+
+FIoStoreShaderLibraryProcessor::~FIoStoreShaderLibraryProcessor()
+{
+#if WITH_EDITORONLY_DATA
+	for (TPair<FString, FEditorShaderCodeArchive*> ShaderArchive : AllMonolithicShaderArchives)
+	{
+		delete ShaderArchive.Value;
+	}
+	for (TPair<const FContainerTargetFile*, FSerializedShaderArchiveBuffer*> ShaderArchive : AllAssignedSerializedShaders)
+	{
+		delete ShaderArchive.Value;
+	}
+#endif
+}
+
+// Invokes TargetFileCallback on each ShaderCodeLibrary entry in all container targets.
+static void ForEachShaderCodeLibraryTargetFile(const TArrayView<FContainerTargetSpec*>& ContainerTargets, const TFunction<void(FContainerTargetFile&)>& TargetFileCallback)
+{
+	for (FContainerTargetSpec* ContainerTarget : ContainerTargets)
 	{
-		OutShaderMapAssetAssociations.Emplace(KV.Key, MoveTemp(KV.Value));
+		for (FContainerTargetFile& TargetFile : ContainerTarget->TargetFiles)
+		{
+			if (TargetFile.ChunkType == EContainerChunkType::ShaderCodeLibrary)
+			{
+				TargetFileCallback(TargetFile);
+			}
+		}
 	}
+}
 
-	return true;
+// Returns true if the specified target file refers to a Global shadermap library
+static bool IsTargetFileGlobalShaderLibrary(const FContainerTargetFile& TargetFile)
+{
+	return FPathViews::GetCleanFilename(TargetFile.NormalizedSourcePath).StartsWith(TEXT("ShaderArchive-Global-"));
+}
+
+// Returns true if the specified target file must not be moved during chunk re-assignment.
+// This applies to the global Global shader library. If needed, further conditions can be added here to also include Chunk0 for instance.
+static bool IsTargetFileExcludedFromChunkReassignment(const FContainerTargetFile& TargetFile)
+{
+	return IsTargetFileGlobalShaderLibrary(TargetFile);
+}
+
+// Generate shader chunk name from pakchunk.
+// This needs to match the shader library name that will be loaded through IoStore when pakchunks are mounted at runtime.
+// Format is "ShaderArchive-<PROJECT>_Chunk<ID>-<SHADERFORMAT>-<SHADERPLATFORM>", e.g. "ShaderArchive_Lyra_Chunk0-SF_VULKAN_SM6-VULKAN_SM6".
+static FString GetShaderChunkNameFromPakChunk(const FStringView& PakChunkName, const TCHAR* ProjectName, const TCHAR* InFormatAndPlatformName)
+{
+	int32 ChunkId = 0;
+	FString ChunkIdSuffix;
+	(void)FShaderLibraryNameInfo::ParsePakChunkId(PakChunkName, ChunkId, ChunkIdSuffix);
+	return FString::Printf(TEXT("ShaderArchive-%s_Chunk%d%s-%s"), ProjectName, ChunkId, *ChunkIdSuffix, InFormatAndPlatformName);
+}
+
+// Adds a new container target file that acts as a proxy for shader libraries that are re-assigned to a new container target.
+static FContainerTargetFile* AddShaderLibraryHeaderToContainerTarget(FContainerTargetSpec* ContainerTarget, const FStringView& PakChunkName, const TCHAR* ProjectName, const TCHAR* InFormatAndPlatformName)
+{
+	check(ContainerTarget);
+	FContainerTargetFile& ShaderLibraryTargetFile = ContainerTarget->TargetFiles.AddDefaulted_GetRef();
+	ShaderLibraryTargetFile.ContainerTarget = ContainerTarget;
+	ShaderLibraryTargetFile.NormalizedSourcePath = GetShaderChunkNameFromPakChunk(PakChunkName, ProjectName, InFormatAndPlatformName);
+	ShaderLibraryTargetFile.ChunkType = EContainerChunkType::ShaderCodeLibrary;
+	ShaderLibraryTargetFile.SourceSize = 0llu;
+	ShaderLibraryTargetFile.ChunkId = FIoChunkId::InvalidChunkId;
+	return &ShaderLibraryTargetFile;
 }
 
-void FIoStoreShaderLibraryProcessor::ProcessShaderLibraries(const FIoStoreArguments& Arguments, TArray<FContainerTargetSpec*>& ContainerTargets, FIoStoreShaderLibraryOutput& Output)
+// Returns true if the container target has any shader library target file of the specified format.
+static bool HasContainerShaderLibraryFormat(const FContainerTargetSpec& InContainerTarget, FStringView InFormatAndPlatformName)
+{
+	for (const FContainerTargetFile& TargetFile : InContainerTarget.TargetFiles)
+	{
+		if (TargetFile.ChunkType == EContainerChunkType::ShaderCodeLibrary &&
+			TargetFile.NormalizedSourcePath.Contains(InFormatAndPlatformName))
+		{
+			return true;
+		}
+	}
+	return false;
+}
+
+// Returns the name of the project that is being packaged and all target shader-platform formats.
+static bool GetPackagedProjectNameAndTargetFormats(TArrayView<FContainerTargetSpec*> InContainerTargets, FString& OutProjectName, TArray<FString>& OutFormatAndPlatformNames)
+{
+	OutProjectName.Empty();
+	OutFormatAndPlatformNames.Empty();
+
+	// Read project name from command line, since FApp::GetProjectName() returns "UnrealPak" instead of the project that is being packaged and `-ProjectName=` is not a required command line argument.
+	// The first command line argument is always the *.uproject filename.
+	{
+		FString ProjectFilename;
+		const TCHAR* CmdLine = FCommandLine::GetOriginal();
+		if (FParse::Token(CmdLine, ProjectFilename, false, TEXT(' ')))
+		{
+			OutProjectName = FPaths::GetBaseFilename(ProjectFilename);
+		}
+		else
+		{
+			return false;
+		}
+	}
+
+	// Heuristic to find all shader formats relevant for the container targets.
+	ForEachShaderCodeLibraryTargetFile(InContainerTargets, [&OutProjectName, &OutFormatAndPlatformNames](FContainerTargetFile& TargetFile) -> void {
+		FShaderLibraryNameInfo LibraryNameInfo;
+		LibraryNameInfo.ParseFromFilename(*TargetFile.NormalizedSourcePath);
+		OutFormatAndPlatformNames.AddUnique(LibraryNameInfo.GetFormatAndPlatformName());
+	});
+
+	return !OutProjectName.IsEmpty() && !OutFormatAndPlatformNames.IsEmpty();
+}
+
+void FIoStoreShaderLibraryProcessor::ProcessShaderLibraries(const FIoStoreArguments& Arguments, TArrayView<FContainerTargetSpec*> ContainerTargets, FIoStoreShaderLibraryOutput& Output)
 {
 	IOSTORE_CPU_SCOPE(ProcessShaderLibraries);
 
 	const double LibraryStart = FPlatformTime::Seconds();
+	TArray<FString> FormatAndPlatformNames;
 
+	// Cache input arguments
+	PackageStore = Arguments.PackageStore.Get();
+	ShaderOodleCompressor = Arguments.ShaderOodleCompressor;
+	ShaderOodleLevel = Arguments.ShaderOodleLevel;
+
+	if (IsChunkOverrideEnabled())
 	{
-		IOSTORE_CPU_SCOPE(ConvertShaderLibraries);
-		for (FContainerTargetSpec* ContainerTarget : ContainerTargets)
+#if WITH_EDITORONLY_DATA
+		UE_LOG(LogIoStore, Display, TEXT("Shader library chunk override is enabled (r.ShaderCodeLibrary.PakChunkOverride=true)"));
+
+		//@lh-todo:
+		//  We need to generate all target files for ShaderCodeLibrary chunks. The shader library headers for both IAD and non-IAD
+		//  must be put into non-IAD containers, so they are always mounted and immediately installed.
+		//  We shouldn't put those headers into the IAD containers just to move them back out into non-IAD containers.
+		//  Clean up this code so that ProcessShaderLibraryFromTargetFile() can emplace shader library headers
+		//  and shader code entries into separate containers; Also see 'MoveShaderLibraryTOCsToInstalledContainers' below.
+		// Proxy shader library targets
 		{
-			for (FContainerTargetFile& TargetFile : ContainerTarget->TargetFiles)
+			IOSTORE_CPU_SCOPE(ProxyShaderLibraryTargets);
+
+			// Get project name that is being packaged and target formats to generate specific names for the proxy shader library target files.
+			FString ProjectName;
+			if (GetPackagedProjectNameAndTargetFormats(ContainerTargets, ProjectName, FormatAndPlatformNames))
 			{
-				if (TargetFile.ChunkType == EContainerChunkType::ShaderCodeLibrary)
+				// Add proxy target files for all containers that don't contain a shader library yet.
+				// This provides the shader library targets for pakchunks that were overridden but didn't have an assigned shader library previously.
+				for (FContainerTargetSpec* ContainerTarget : ContainerTargets)
 				{
-					if (!ProcessShaderLibraryFromTargetFile(Arguments, TargetFile, Output))
+					for (const FString& FormatAndPlatformName : FormatAndPlatformNames)
 					{
-						UE_LOG(LogIoStore, Warning, TEXT("Failed converting shader library '%s'"), *TargetFile.NormalizedSourcePath);
+						if (!HasContainerShaderLibraryFormat(*ContainerTarget, FormatAndPlatformName))
+						{
+							// Add new proxy target file for shader library, since it doesn't have one in the current container target for this format.
+							FContainerTargetFile* ProxyTargetFile = AddShaderLibraryHeaderToContainerTarget(ContainerTarget, ContainerTarget->Name.ToString(), *ProjectName, *FormatAndPlatformName);
+							UE_LOG(LogIoStore, Display, TEXT("Add shader library proxy target file '%s' to container '%s'"), *ProxyTargetFile->NormalizedSourcePath, *ContainerTarget->Name.ToString());
+						}
 					}
 				}
 			}
+			else
+			{
+				// If no other container had any shader libraries we could obtain the shader formats from,
+				// skip creating proxy targets as there is nothing to re-assign shaders to (e.g. when packaging a Server).
+				UE_LOG(LogIoStore, Display, TEXT("Skipped shader library proxy targets due to missing shader format information from other containers (%d total)"), ContainerTargets.Num());
+			}
 		}
+
+		// Preload all shader archives, so we can re-assign shaders to their targets from monolithic libraries
+		{
+			IOSTORE_CPU_SCOPE(PreloadShaderLibraries);
+			ForEachShaderCodeLibraryTargetFile(ContainerTargets, [this, &Arguments, &Output](FContainerTargetFile& TargetFile) -> void {
+				if (!IsTargetFileExcludedFromChunkReassignment(TargetFile) && !PreloadShaderArchive(TargetFile))
+				{
+					UE_LOG(LogIoStore, Warning, TEXT("Failed preloading shader library '%s'"), *TargetFile.NormalizedSourcePath);
+				}
+			});
+		}
+
+		// Now assign shader archives to their final target files
+		{
+			IOSTORE_CPU_SCOPE(AssignShaderArchives);
+			ForEachShaderCodeLibraryTargetFile(ContainerTargets, [this, &Arguments, &Output](FContainerTargetFile& TargetFile) -> void {
+				if (!IsTargetFileExcludedFromChunkReassignment(TargetFile) && !AssignShaderArchives(TargetFile))
+				{
+					UE_LOG(LogIoStore, Warning, TEXT("Failed assigning shader archives for target library '%s'"), *TargetFile.NormalizedSourcePath);
+				}
+			});
+		}
+#else
+		UE_LOG(LogIoStore, Warning, TEXT("Shader library chunk override is enabled (r.ShaderCodeLibrary.PakChunkOverride=true), but not supported (WITH_EDITORONLY_DATA=0)"));
+#endif // WITH_EDITORONLY_DATA
+	}
+	else
+	{
+		UE_LOG(LogIoStore, Display, TEXT("Shader library chunk override is disabled (r.ShaderCodeLibrary.PakChunkOverride=false)"));
+	}
+
+	{
+		IOSTORE_CPU_SCOPE(ConvertShaderLibraries);
+		ForEachShaderCodeLibraryTargetFile(ContainerTargets, [this, &Arguments, &Output](FContainerTargetFile& TargetFile) -> void {
+			if (!ProcessShaderLibraryFromTargetFile(Arguments, TargetFile, Output))
+			{
+				UE_LOG(LogIoStore, Warning, TEXT("Failed converting shader library '%s'"), *TargetFile.NormalizedSourcePath);
+			}
+		});
 	}
 
 	{
@@ -440,6 +783,7 @@ void FIoStoreShaderLibraryProcessor::ProcessShaderLibraries(const FIoStoreArgume
 		}
 	}
 
+	// Transfer ownership of shader information (all previously allocated FShaderInfo entries) to the container targets
 	{
 		IOSTORE_CPU_SCOPE(AddShaderCodeToContainer);
 		for (FContainerTargetSpec* ContainerTarget : ContainerTargets)
@@ -448,32 +792,340 @@ void FIoStoreShaderLibraryProcessor::ProcessShaderLibraries(const FIoStoreArgume
 		}
 	}
 
+	// Finally, move all shader library header chunks from OnDemand containers back into their equivalent standard containers.
+	// This ensures the header (or table-of-content rather) is installed before any shaders are loaded on-demand.
+	// The alternative would be to install the header from the OnDemand container as soon as its mounted,
+	// but that eliminates the purpose of OnDemand data and would unnecessarily increase download sizes.
+	if (IsChunkOverrideEnabled())
+	{
+		IOSTORE_CPU_SCOPE(MoveShaderLibraryTOCsToInstalledContainers);
+
+		auto GetConditionalChunkId = [](const FContainerTargetSpec& InContainerTarget, bool bOnDemandOnly) -> int32
+			{
+				int32 ChunkId = INDEX_NONE;
+				FString ChunkIdSuffix;
+				(void)FShaderLibraryNameInfo::ParsePakChunkId(InContainerTarget.Name.ToString(), ChunkId, ChunkIdSuffix);
+				const bool bIsOnDemandChunk = !ChunkIdSuffix.IsEmpty();
+				return bIsOnDemandChunk == bOnDemandOnly ? ChunkId : INDEX_NONE;
+			};
+
+		auto FindNonOnDemandContainerTarget = [&ContainerTargets, &GetConditionalChunkId](const FContainerTargetSpec& InContainerTarget, int32 InChunkId) -> FContainerTargetSpec*
+			{
+				// Try to find non-OnDemand container with same chunk ID (e.g. "Chunk100" for input container "Chunk100iad")
+				for (FContainerTargetSpec* ContainerTarget : ContainerTargets)
+				{
+					if (GetConditionalChunkId(*ContainerTarget, false) == InChunkId)
+					{
+						return ContainerTarget;
+					}
+				}
+
+				// Otherwise, try to find container for "Chunk0" since that is always mounted
+				constexpr int32 kAlwaysInstalledChunkId = 0;
+				for (FContainerTargetSpec* ContainerTarget : ContainerTargets)
+				{
+					if (GetConditionalChunkId(*ContainerTarget, false) == kAlwaysInstalledChunkId)
+					{
+						return ContainerTarget;
+					}
+				}
+
+				// Failed to find a suitable container
+				return nullptr;
+			};
+
+		for (FContainerTargetSpec* ContainerTarget : ContainerTargets)
+		{
+			if (const int32 ChunkId = GetConditionalChunkId(*ContainerTarget, true); ChunkId != INDEX_NONE)
+			{
+				for (auto Iter = ContainerTarget->TargetFiles.CreateIterator(); Iter; ++Iter)
+				{
+					if (Iter->ChunkType == EContainerChunkType::ShaderCodeLibrary)
+					{
+						// We found a shader code library chunk in an OnDemand container.
+						if (FContainerTargetSpec* NewContainerTarget = FindNonOnDemandContainerTarget(*ContainerTarget, ChunkId))
+						{
+							// Move the target file into the regular container.
+							UE_LOG(LogIoStore, Display, TEXT("Moved ShaderCodeLibrary header '%s' (ChunkId='%s') out of OnDemand container '%s' into regular container '%s'"),
+								*Iter->NormalizedSourcePath, *LexToString(Iter->ChunkId), *ContainerTarget->Name.ToString(), *NewContainerTarget->Name.ToString());
+							NewContainerTarget->TargetFiles.Add(MoveTemp(*Iter));
+							Iter.RemoveCurrent();
+						}
+						else
+						{
+							UE_LOG(LogIoStore, Error, TEXT("Failed to move ShaderCodeLibrary header '%s' (ChunkId='%s') out of OnDemand container '%s'"),
+								*Iter->NormalizedSourcePath, *LexToString(Iter->ChunkId), *ContainerTarget->Name.ToString());
+						}
+					}
+				}
+			}
+		}
+	}
+
+#if WITH_EDITORONLY_DATA
+	// Dump shader chunk re-assignment to text file if debug information is enabled
+	if (IsChunkOverrideEnabled() && IsChunkOverrideDebugEnabled())
+	{
+		for (const FString& FormatAndPlatformName : FormatAndPlatformNames)
+		{
+			const FString DebugLibFolder = FSerializedShaderArchive::MakeDebugDirectory() / FormatAndPlatformName;
+			IFileManager::Get().MakeDirectory(*DebugLibFolder, true);
+			const FString DebugFilename = DebugLibFolder / FString::Printf(TEXT("ShaderChunkOverride-%s.txt"), *FormatAndPlatformName);
+			if (TUniquePtr<FArchive> DebugFileWriter = TUniquePtr<FArchive>(IFileManager::Get().CreateFileWriter(*DebugFilename)))
+			{
+				if (!DumpChunkOverrideDebugInfo(FName(FormatAndPlatformName), *DebugFileWriter))
+				{
+					UE_LOG(LogIoStore, Warning, TEXT("Failed to save debug output: '%s'"), *DebugFilename);
+				}
+			}
+		}
+	}
+#endif
+
 	const double LibraryEnd = FPlatformTime::Seconds();
 	UE_LOG(LogIoStore, Display, TEXT("Shaders processed in %.2f seconds"), LibraryEnd - LibraryStart);
 }
 
 bool FIoStoreShaderLibraryProcessor::ProcessShaderLibraryFromTargetFile(const FIoStoreArguments& Arguments, FContainerTargetFile& TargetFile, FIoStoreShaderLibraryOutput& Output)
 {
+	IOSTORE_CPU_SCOPE(ConvertShaderLibrary);
+
 	FContainerTargetSpec* ContainerTarget = TargetFile.ContainerTarget;
 	checkf(ContainerTarget, TEXT("Target file for shader library '%s' is missing its parent node (ContainerTarget is null)"), *TargetFile.NormalizedSourcePath);
 
-	TSet<FShaderInfo*>& ContainerShaderLibraryShaders = ContainerTargetToShaderLibrary.FindOrAdd(ContainerTarget);
-
-	TArray<FShaderGroupIoChunkMap> ShaderMaps;
-	TTuple<FIoChunkId, FIoBuffer> LibraryChunk;
-	TArray<FShaderGroupIoChunk> CodeChunks;
-	TArray<TTuple<FSHAHash, FShaderMapAssetPaths>> ShaderMapAssetAssociations;
-	if (!ConvertToIoStoreShaderLibrary(*TargetFile.NormalizedSourcePath, TargetFile.ChunkId, Arguments.PackageStore.Get(), Arguments.ShaderOodleCompressor, Arguments.ShaderOodleLevel, LibraryChunk, CodeChunks, ShaderMaps, ShaderMapAssetAssociations))
+	if (IsChunkOverrideEnabled() && !IsTargetFileExcludedFromChunkReassignment(TargetFile))
 	{
+#if WITH_EDITORONLY_DATA
+		// Find shader archive that was assigned to the current target file and convert it into an IoStore shader library
+		if (FSerializedShaderArchiveBuffer* const* SerializedShaders = AllAssignedSerializedShaders.Find(&TargetFile))
+		{
+			if (!ConvertToIoStoreShaderLibrary(TargetFile, **SerializedShaders))
+			{
+				return false;
+			}
+		}
+		else
+		{
+			return true; // Ok, no shader library assigned to this targe file
+		}
+#else
+		ensureMsgf(false, TEXT("Cannot re-assign shader archives, because WITH_EDITORONLY_DATA=0"));
 		return false;
+#endif // WITH_EDITORONLY_DATA
+	}
+	else
+	{
+		// Without chunk override, convert shader library directly from archive chunk to pakchunk
+		FShaderLibraryNameInfo LibraryNameInfo;
+		FSerializedShaderArchiveBuffer SerializedShaders;
+		if (!LoadSerializedShaders(TargetFile, SerializedShaders, false, &LibraryNameInfo))
+		{
+			return false;
+		}
+		ShaderMapAssetAssociations.Append(SerializedShaders.GetShaderMapAssetAssociations());
+
+		// Keep track of old and new pakchunks if debug information is enabled (also for non-chunked targets)
+		if (IsChunkOverrideDebugEnabled())
+		{
+			AddShaderMapsToPakChunkDebugInfo(TargetFile, SerializedShaders, FName(LibraryNameInfo.GetFormatAndPlatformName()), ChunkAssignment_Old | ChunkAssignment_New);
+		}
+
+		if (!ConvertToIoStoreShaderLibrary(TargetFile, SerializedShaders))
+		{
+			return false;
+		}
 	}
-	TargetFile.ChunkId = LibraryChunk.Key;
-	TargetFile.SourceBuffer.Emplace(LibraryChunk.Value);
-	TargetFile.SourceSize = LibraryChunk.Value.GetSize();
 
-	const bool bIsGlobalShaderLibrary = FPaths::GetCleanFilename(TargetFile.NormalizedSourcePath).StartsWith(TEXT("ShaderArchive-Global-"));
+	const bool bIsGlobalShaderLibrary = IsTargetFileGlobalShaderLibrary(TargetFile);
 	const FShaderInfo::EShaderType ShaderType = bIsGlobalShaderLibrary ? FShaderInfo::Global : FShaderInfo::Normal;
-	for (const FShaderGroupIoChunk& CodeChunk : CodeChunks)
+	FlushCodeIoChunks(ContainerTarget, Output, ShaderType);
+
+	return true;
+}
+
+#if WITH_EDITORONLY_DATA
+
+bool FIoStoreShaderLibraryProcessor::PreloadShaderArchive(const FContainerTargetFile& TargetFile)
+{
+	// Preload serialized shader archive for target file
+	FShaderLibraryNameInfo LibraryNameInfo;
+	FSerializedShaderArchiveBuffer SerializedShaders;
+	const bool bIncludeTypeInfo = true;
+	if (!LoadSerializedShaders(TargetFile, SerializedShaders, bIncludeTypeInfo, &LibraryNameInfo))
+	{
+		return false;
+	}
+	ShaderMapAssetAssociations.Append(SerializedShaders.GetShaderMapAssetAssociations());
+
+	// Append serialized shaders to monolithic library so we can later create new chunks from it
+	FEditorShaderCodeArchive*& MonolithicShaderArchive = AllMonolithicShaderArchives.FindOrAdd(LibraryNameInfo.GetFormatAndPlatformName());
+	if (!MonolithicShaderArchive)
+	{
+		MonolithicShaderArchive = new FEditorShaderCodeArchive(FName(LibraryNameInfo.GetFormatAndPlatformName()));
+	}
+	MonolithicShaderArchive->AppendFromArchive(SerializedShaders, SerializedShaders.GetShaderCodeView(), ShaderCodeArchive::ECookShaderLibrarySource::CurrentCook);
+
+	// Keep track of old and new pakchunks if debug information is enabled
+	if (IsChunkOverrideDebugEnabled())
+	{
+		AddShaderMapsToPakChunkDebugInfo(TargetFile, SerializedShaders, FName(LibraryNameInfo.GetFormatAndPlatformName()), ChunkAssignment_Old);
+	}
+
+	return true;
+}
+
+void FIoStoreShaderLibraryProcessor::AddShaderMapsToPakChunkDebugInfo(
+	const FContainerTargetFile& TargetFile, const FSerializedShaderArchiveBuffer& SerializedShaders,
+	FName FormatAndPlatformName, int32 ChunkAssignmentType)
+{
+	check(TargetFile.ContainerTarget);
+	FMonolithicArchiveDebugInfo& DebugInfo = ChunkAssignmentDebugInfo.FindOrAdd(FormatAndPlatformName);
+	for (const FSHAHash& ShaderMapHash : SerializedShaders.ShaderMapHashes)
+	{
+		FChunkAssignmentInfo& ShaderMapInfo = DebugInfo.ShaderMaps.FindOrAdd(ShaderMapHash);
+		if ((ChunkAssignmentType & ChunkAssignment_New) != 0)
+		{
+			ShaderMapInfo.NewPakChunks.Add(TargetFile.ContainerTarget->Name);
+		}
+		if ((ChunkAssignmentType & ChunkAssignment_Old) != 0)
+		{
+			ShaderMapInfo.OldPakChunks.Add(TargetFile.ContainerTarget->Name);
+		}
+	}
+}
+
+bool FIoStoreShaderLibraryProcessor::DumpChunkOverrideDebugInfo(FName FormatAndPlatformName, FArchive& Ar)
+{
+	if (!ChunkAssignmentDebugInfo.Contains(FormatAndPlatformName))
+	{
+		return false;
+	}
+
+	// Generate debug output text
+	FString DebugOutput;
+
+	auto AppendPakChunkNames = [&DebugOutput](const TCHAR* InListName, const TArray<FName>& InPakChunks) -> void
+		{
+			DebugOutput.Appendf(TEXT("\n  %s:"), InListName);
+			if (InPakChunks.IsEmpty())
+			{
+				DebugOutput.Append(TEXT(" <None>"));
+			}
+			else
+			{
+				for (const FName& PakChunk : InPakChunks)
+				{
+					DebugOutput.Appendf(TEXT(" '%s'"), *PakChunk.ToString());
+				}
+			}
+		};
+
+	auto IsAssignmentUnchanged = [](const FChunkAssignmentInfo& InAssignment) -> bool
+		{
+			if (InAssignment.OldPakChunks.Num() != InAssignment.NewPakChunks.Num())
+			{
+				return false;
+			}
+			for (const FName& OldName : InAssignment.OldPakChunks)
+			{
+				if (!InAssignment.NewPakChunks.Contains(OldName))
+				{
+					return false;
+				}
+			}
+			return true;
+		};
+
+	for (const TPair<FSHAHash, FChunkAssignmentInfo>& Pair : ChunkAssignmentDebugInfo[FormatAndPlatformName].ShaderMaps)
+	{
+		DebugOutput.Appendf(TEXT("%s {"), *Pair.Key.ToString());
+		if (IsAssignmentUnchanged(Pair.Value))
+		{
+			AppendPakChunkNames(TEXT("Unchanged"), Pair.Value.OldPakChunks);
+		}
+		else
+		{
+			AppendPakChunkNames(TEXT("Old"), Pair.Value.OldPakChunks);
+			AppendPakChunkNames(TEXT("New"), Pair.Value.NewPakChunks);
+		}
+		DebugOutput.Append(TEXT("\n}\n"));
+	}
+
+	// Write debug output to file
+	FTCHARToUTF8 DebugOutputUtf8Converter(*DebugOutput);
+	Ar.Serialize(const_cast<UTF8CHAR*>(reinterpret_cast<const UTF8CHAR*>(DebugOutputUtf8Converter.Get())), DebugOutputUtf8Converter.Length());
+
+	return true;
+}
+
+bool FIoStoreShaderLibraryProcessor::AssignShaderArchives(const FContainerTargetFile& TargetFile)
+{
+	// Convert shader format and platform name from *.ushaderbytecode filename
+	FShaderLibraryNameInfo LibraryNameInfo;
+	if (!ParseIoStoreShaderLibraryName(LibraryNameInfo, *TargetFile.NormalizedSourcePath))
+	{
+		return false;
+	}
+
+	// Get container for current target file and extract ChunkId from it
+	const FContainerTargetSpec* ContainerTarget = TargetFile.ContainerTarget;
+	check(ContainerTarget);
+
+	// Gather all package names in container target
+	TSet<FName> PackagesInChunk;
+	for (const FCookedPackage* CookedPackage : ContainerTarget->Packages)
+	{
+		PackagesInChunk.Add(CookedPackage->PackageName);
+	}
+
+	// Build intermediate shader archive and assign relevant shaders from preloaded serialized shaders
+	const FString FormatAndPlatformName = LibraryNameInfo.GetFormatAndPlatformName();
+	if (FEditorShaderCodeArchive** FoundMonolithicShaderArchive = AllMonolithicShaderArchives.Find(FormatAndPlatformName))
+	{
+		FEditorShaderCodeArchive* MonolithicShaderArchive = *FoundMonolithicShaderArchive;
+		check(MonolithicShaderArchive);
+
+		// Create new chunk for packages.
+		TUniquePtr<FEditorShaderCodeArchive> NewChunkShaderArchive = MonolithicShaderArchive->CreateNamedChunk(LibraryNameInfo.ChunkName, PackagesInChunk);
+		check(NewChunkShaderArchive);
+
+		// Convert constructed intermediate shader archive to final serialized archive with flat buffer
+		FSerializedShaderArchiveBuffer* AssignedShaderArchive = new FSerializedShaderArchiveBuffer();
+		{
+			if (AssignedShaderArchive->LoadFromEditorArchive(*NewChunkShaderArchive))
+			{
+				// Log stats and debug info for newly assigned shader archive. Without pakchunk override, this is only logged during cooking; see SaveShaderCodeChunk().
+				NewChunkShaderArchive->DumpStatsAndDebugInfo();
+			}
+			else
+			{
+				UE_LOG(LogIoStore, Error, TEXT("Could not serialize shaders from monolithic archive '%s'"), *LibraryNameInfo.GetLibraryName());
+				return false;
+			}
+		}
+		AllAssignedSerializedShaders.Add(&TargetFile, AssignedShaderArchive);
+
+		// Keep track of old and new pakchunks if debug information is enabled
+		if (IsChunkOverrideDebugEnabled())
+		{
+			AddShaderMapsToPakChunkDebugInfo(TargetFile, *AssignedShaderArchive, FName(FormatAndPlatformName), ChunkAssignment_New);
+		}
+	}
+	else
+	{
+		UE_LOG(LogIoStore, Error, TEXT("Could not find monolithic shader archive for format '%s'"), *FormatAndPlatformName);
+		return false;
+	}
+	return true;
+}
+
+#endif // WITH_EDITORONLY_DATA
+
+void FIoStoreShaderLibraryProcessor::FlushCodeIoChunks(FContainerTargetSpec* ContainerTarget, FIoStoreShaderLibraryOutput& Output, FShaderInfo::EShaderType ShaderType)
+{
+	TSet<FShaderInfo*>& ContainerTargetShaderInfos = ContainerTargetToShaderLibrary.FindOrAdd(ContainerTarget);
+
+	for (const FShaderGroupIoChunk& CodeChunk : CodeIoChunks)
 	{
 		const FIoChunkId& ShaderChunkId = CodeChunk.ChunkId;
 		FShaderInfo* ShaderInfo = ChunkIdToShaderInfoMap.FindRef(ShaderChunkId);
@@ -507,29 +1159,22 @@ bool FIoStoreShaderLibraryProcessor::ProcessShaderLibraryFromTargetFile(const FI
 			ShaderInfo->TypeInContainer.Add(ContainerTarget, ShaderType);
 		}
 
-		ContainerShaderLibraryShaders.Add(ShaderInfo);
+		ContainerTargetShaderInfos.Add(ShaderInfo);
 	}
 
-	for (const TTuple<FSHAHash, FShaderMapAssetPaths>& ShaderMapAssetAssociation : ShaderMapAssetAssociations)
-	{
-		for (const FName& PackageName : ShaderMapAssetAssociation.Value)
-		{
-			TSet<FSHAHash>& PackageShaderMaps = PackageNameToShaderMaps.FindOrAdd(PackageName);
-			PackageShaderMaps.Add(ShaderMapAssetAssociation.Key);
-		}
-	}
-
-	for (FShaderGroupIoChunkMap& ShaderMap : ShaderMaps)
+	for (FShaderGroupIoChunkMap& ShaderMap : CodeIoChunkMaps)
 	{
 		TArray<FIoChunkId>& ShaderMapChunkIds = ShaderMapToGroupChunkIds.FindOrAdd(ShaderMap.ShaderMapHash);
 		ShaderMapChunkIds.Append(MoveTemp(ShaderMap.ShaderGroupChunkIds));
 	}
 
-	return true;
+	// Reset intermediate data for processing target files
+	CodeIoChunks.Empty();
+	CodeIoChunkMaps.Empty();
 }
 
 // Adds a new entry of type EContainerChunkType::ShaderCode to the specified container target
-static void AddShaderCodeToContainerTarget(FContainerTargetSpec* ContainerTarget, FShaderInfo& ShaderInfo)
+static void AddShaderCodeToContainerTarget(FContainerTargetSpec* ContainerTarget, const FShaderInfo& ShaderInfo)
 {
 	check(ContainerTarget);
 	FContainerTargetFile& ShaderTargetFile = ContainerTarget->TargetFiles.AddDefaulted_GetRef();
@@ -546,9 +1191,9 @@ void FIoStoreShaderLibraryProcessor::UpdatePackageStoreShaders(FCookedPackage* P
 	// 1. Update ShaderInfos with which packages we reference.
 	// 2. Add to packages which shaders we use.
 	// 3. Add to PackageStore what shaders we use.
-	if (const TSet<FSHAHash>* FindShaderMapHashes = PackageNameToShaderMaps.Find(Package->PackageName))
+	if (const FShaderMapAssetAssociations::FAssociatedAssetData* FindAssetData = ShaderMapAssetAssociations.FindAsset(Package->PackageName))
 	{
-		for (const FSHAHash& ShaderMapHash : *FindShaderMapHashes)
+		for (const FSHAHash& ShaderMapHash : FindAssetData->ShaderMaps)
 		{
 			const TArray<FIoChunkId>* FindChunkIds = ShaderMapToGroupChunkIds.Find(ShaderMapHash);
 			if (!FindChunkIds)
@@ -578,8 +1223,8 @@ void FIoStoreShaderLibraryProcessor::UpdatePackageStoreShaders(FCookedPackage* P
 
 void FIoStoreShaderLibraryProcessor::AddAllReferencedShadersToContainerTarget(FContainerTargetSpec* ContainerTarget, FIoStoreShaderLibraryOutput& Output)
 {
-	const TSet<FShaderInfo*>* FindContainerShaderLibraryShaders = ContainerTargetToShaderLibrary.Find(ContainerTarget);
-	if (!FindContainerShaderLibraryShaders)
+	const TSet<FShaderInfo*>* FoundContainerShaderLibraryShaders = ContainerTargetToShaderLibrary.Find(ContainerTarget);
+	if (!FoundContainerShaderLibraryShaders)
 	{
 		// Early exit if package doesn't contain any shaders
 		return;
@@ -600,7 +1245,7 @@ void FIoStoreShaderLibraryProcessor::AddAllReferencedShadersToContainerTarget(FC
 		}
 	}
 
-	for (FShaderInfo* ShaderInfo : *FindContainerShaderLibraryShaders)
+	for (FShaderInfo* ShaderInfo : *FoundContainerShaderLibraryShaders)
 	{
 		FShaderAssociationInfo::FShaderChunkInfoKey ShaderChunkInfoKey = { ContainerTarget->Name, ShaderInfo->ShaderGroupChunk.ChunkId };
 		FShaderAssociationInfo::FShaderChunkInfo& ShaderChunkInfo = Output.AssocInfo.ShaderChunkInfos.Add(ShaderChunkInfoKey);
@@ -646,7 +1291,7 @@ void FIoStoreShaderLibraryProcessor::AddAllReferencedShadersToContainerTarget(FC
 			{
 				ShaderChunkInfo.Type = FShaderAssociationInfo::FShaderChunkInfo::Package;
 
-				for (FCookedPackage* Package : ShaderInfo->ReferencedByPackages)
+				for (const FCookedPackage* Package : ShaderInfo->ReferencedByPackages)
 				{
 					ShaderChunkInfo.ReferencedByPackages.Add(Package->PackageName);
 
diff --git a/Engine/Source/Developer/IoStoreUtilities/Private/IoStoreShaderLibraryProcessor.h b/Engine/Source/Developer/IoStoreUtilities/Private/IoStoreShaderLibraryProcessor.h
--- a/Engine/Source/Developer/IoStoreUtilities/Private/IoStoreShaderLibraryProcessor.h
+++ b/Engine/Source/Developer/IoStoreUtilities/Private/IoStoreShaderLibraryProcessor.h
@@ -5,11 +5,19 @@
 #include "CookMetadata.h"
 #include "Containers/Array.h"
 #include "Containers/Map.h"
+#include "IoStoreUtilitiesCore.h"
 #include "IO/IoChunkId.h"
 #include "Misc/SecureHash.h"
+#include "Misc/CoreMiscDefines.h"
+#include "ShaderMapAssetAssociation.h"
 #include "Templates/UniquePtr.h"
 #include "UObject/NameTypes.h"
 
+struct FShaderLibraryNameInfo;
+#if WITH_EDITORONLY_DATA
+class FEditorShaderCodeArchive;
+#endif
+
 namespace UE::IoStore::Private
 {
 
@@ -80,19 +88,88 @@ struct FIoStoreShaderLibraryOutput
 	FShaderAssociationInfo AssocInfo;
 };
 
+class FSerializedShaderArchiveBuffer;
+
 // Processor to load cooked shader libraries and package them into IoStore entries
 class FIoStoreShaderLibraryProcessor
 {
+	// Shader map hash with its IoChunk Ids of all its shader groups
+	struct FShaderGroupIoChunkMap
+	{
+		FSHAHash ShaderMapHash;
+		TArray<FIoChunkId> ShaderGroupChunkIds;
+	};
+
+	// Input arguments
+	FCookedPackageStore* PackageStore = nullptr;
+	FOodleDataCompression::ECompressor ShaderOodleCompressor = FOodleDataCompression::ECompressor::NotSet;
+	FOodleDataCompression::ECompressionLevel ShaderOodleLevel = FOodleDataCompression::ECompressionLevel::None;
+
+	// Persitent data for entire shader library processing
 	TMap<FIoChunkId, FShaderInfo*> ChunkIdToShaderInfoMap;
 	TMap<FSHAHash, TArray<FIoChunkId>> ShaderMapToGroupChunkIds; // Maps from shader map hash to its set of shader group IoChunk Ids
-	TMap<FName, TSet<FSHAHash>> PackageNameToShaderMaps;
 	TMap<FContainerTargetSpec*, TSet<FShaderInfo*>> ContainerTargetToShaderLibrary;
+	FShaderMapAssetAssociations ShaderMapAssetAssociations;
+
+	// Intermediate data while processing a subset of target files
+	TArray<FShaderGroupIoChunk> CodeIoChunks;
+	TArray<FShaderGroupIoChunkMap> CodeIoChunkMaps; // Maps from shadermap index to shader group IoChunks
+
+#if WITH_EDITORONLY_DATA
+	// Intermediate data for chunk re-assignment
+	TMap<FString, FEditorShaderCodeArchive*> AllMonolithicShaderArchives; // Shader format to archive
+	TMap<const FContainerTargetFile*, FSerializedShaderArchiveBuffer*> AllAssignedSerializedShaders;
+
+	// Debug information (used with CVar `r.ShaderCodeLibrary.PakChunkOverride.DumpDebugInfo`)
+	enum EChunkAssignmentType
+	{
+		ChunkAssignment_Old = (1 << 0),
+		ChunkAssignment_New = (1 << 1),
+	};
+	struct FChunkAssignmentInfo
+	{
+		TArray<FName> OldPakChunks;
+		TArray<FName> NewPakChunks;
+	};
+	struct FMonolithicArchiveDebugInfo
+	{
+		TMap<FSHAHash, FChunkAssignmentInfo> ShaderMaps;
+	};
+	TMap<FName, FMonolithicArchiveDebugInfo> ChunkAssignmentDebugInfo;
+#endif
 
 public:
-	void ProcessShaderLibraries(const FIoStoreArguments& Arguments, TArray<FContainerTargetSpec*>& ContainerTargets, FIoStoreShaderLibraryOutput& Output);
+	~FIoStoreShaderLibraryProcessor();
+
+	void ProcessShaderLibraries(const FIoStoreArguments& Arguments, TArrayView<FContainerTargetSpec*> ContainerTargets, FIoStoreShaderLibraryOutput& Output);
 
 private:
 	bool ProcessShaderLibraryFromTargetFile(const FIoStoreArguments& Arguments, FContainerTargetFile& TargetFile, FIoStoreShaderLibraryOutput& Output);
+
+#if WITH_EDITORONLY_DATA
+	// Preloads the *.ushaderbytecode for the specified target file for later chunk assignment.
+	bool PreloadShaderArchive(const FContainerTargetFile& TargetFile);
+
+	// Assigns preloaded shader archives to their respective target file (i.e. to their respective pakchunks).
+	bool AssignShaderArchives(const FContainerTargetFile& TargetFile);
+
+	void AddShaderMapsToPakChunkDebugInfo(const FContainerTargetFile& TargetFile, const FSerializedShaderArchiveBuffer& SerializedShaders,
+		FName FormatAndPlatformName, int32 ChunkAssignmentType);
+
+	// Writes out a debug file with information about all the shaders that have been re-assigned
+	// to a difference chunk (if r.ShaderCodeLibrary.PakChunkOverride is enabled).
+	bool DumpChunkOverrideDebugInfo(FName FormatAndPlatformName, FArchive& Ar);
+#endif // WITH_EDITORONLY_DATA
+
+	// Loads the serialized shaders of the *.ushaderbytecode file of the specified target.
+	bool LoadSerializedShaders(const FContainerTargetFile& TargetFile, FSerializedShaderArchiveBuffer& OutSerializedShaders, bool bIncludeTypeInfo = false, FShaderLibraryNameInfo* OutNameInfo = nullptr);
+
+	// Converts the *.ushaderbytecode file of the specified target to an IoStore shader code library.
+	bool ConvertToIoStoreShaderLibrary(FContainerTargetFile& TargetFile, const FSerializedShaderArchiveBuffer& SerializedShaders);
+
+	// Flushes all generated shader code IoChunks to the output.
+	void FlushCodeIoChunks(FContainerTargetSpec* ContainerTarget, FIoStoreShaderLibraryOutput& Output, FShaderInfo::EShaderType ShaderType);
+
 	void UpdatePackageStoreShaders(FCookedPackage* Package, const FName& ContainerTargetName);
 	void AddAllReferencedShadersToContainerTarget(FContainerTargetSpec* ContainerTarget, FIoStoreShaderLibraryOutput& Output);
 
diff --git a/Engine/Source/Developer/TargetPlatform/Private/TargetPlatformManagerModule.cpp b/Engine/Source/Developer/TargetPlatform/Private/TargetPlatformManagerModule.cpp
--- a/Engine/Source/Developer/TargetPlatform/Private/TargetPlatformManagerModule.cpp
+++ b/Engine/Source/Developer/TargetPlatform/Private/TargetPlatformManagerModule.cpp
@@ -1677,7 +1677,14 @@ private:
 
 		if (!Modules.Num())
 		{
+#if IS_PROGRAM
+			// Only log this information to file, but don't report this as an error for standalone programs.
+			//   UnrealPak for instance loads TPM to request shader formats only, but doesn't care about audio formats.
+			//   The only format UnrealPak is built against is MetalShaderFormat to check whether native shader libs are supported.
+			UE_LOG(LogTargetPlatformManager, Log, TEXT("No target %s formats found!"), HelperType::GetFormatDesc());
+#else
 			UE_LOG(LogTargetPlatformManager, Error, TEXT("No target %s formats found!"), HelperType::GetFormatDesc());
+#endif
 		}
 
 		for (int32 Index = 0; Index < Modules.Num(); Index++)
diff --git a/Engine/Source/Runtime/RHI/Public/RHIShaderLibrary.h b/Engine/Source/Runtime/RHI/Public/RHIShaderLibrary.h
--- a/Engine/Source/Runtime/RHI/Public/RHIShaderLibrary.h
+++ b/Engine/Source/Runtime/RHI/Public/RHIShaderLibrary.h
@@ -24,7 +24,7 @@ public:
 	virtual int32 GetNumShaders() const = 0;
 	virtual int32 GetNumShadersForShaderMap(int32 ShaderMapIndex) const = 0;
 	virtual int32 GetShaderIndex(int32 ShaderMapIndex, int32 i) const = 0;
-	virtual void GetAllShaderIndices(int32 ShaderMapIndex, TArray<int32>& ShaderIndices) {}
+	virtual void GetAllShaderIndices(int32 ShaderMapIndex, TArray<int32>& ShaderIndices) const {}
 	virtual uint32 GetSizeBytes() const = 0;
 	virtual FSHAHash GetShaderHash(int32 ShaderMapIndex, int32 ShaderIndex) = 0;
 	virtual FSHAHash GetShaderMapHash(int32 ShaderMapIndex) const = 0;
@@ -32,6 +32,7 @@ public:
 	virtual int32 FindShaderIndex(const FSHAHash& Hash) = 0;
 	virtual uint32 GetShaderSizeBytes(int32 ShaderIndex) const { return 0; }
 	virtual bool IsPreloading(int32 ShaderIndex, FGraphEventArray& OutCompletionEvents) { return false; }
+	virtual bool ResolveShaderMap(int32 ShaderMapIndex, FCoreDelegates::FIoChunkIdResolvedFunc IoChunkIdResolvedFunc) const { return false; }
 	virtual bool PreloadShader(int32 ShaderIndex, FGraphEventArray& OutCompletionEvents) { return false; }
 	virtual bool PreloadShaderMap(int32 ShaderMapIndex, FGraphEventArray& OutCompletionEvents) { return false; }
 	virtual bool PreloadShaderMap(int32 ShaderMapIndex, FCoreDelegates::FAttachShaderReadRequestFunc AttachShaderReadRequestFunc) { return false; }
diff --git a/Engine/Source/Runtime/RenderCore/Internal/ShaderLibrary/EditorShaderCodeArchive.h b/Engine/Source/Runtime/RenderCore/Internal/ShaderLibrary/EditorShaderCodeArchive.h
new file mode 100644
--- /dev/null
+++ b/Engine/Source/Runtime/RenderCore/Internal/ShaderLibrary/EditorShaderCodeArchive.h
@@ -0,0 +1,239 @@
+// Copyright Epic Games, Inc. All Rights Reserved.
+
+/*=============================================================================
+	EditorShaderCodeArchive.h: Declaration of FEditorShaderCodeArchive
+=============================================================================*/
+
+#pragma once
+
+#include "Containers/Array.h"
+#include "Containers/HashTable.h"
+#include "Containers/Set.h"
+#include "Containers/StringView.h"
+#include "Memory/SharedBuffer.h"
+#include "Misc/CoreMiscDefines.h"
+#include "Misc/SecureHash.h"
+#include "Serialization/Archive.h"
+#include "ShaderCodeArchive.h"
+#include "ShaderLibrary/ShaderCodeArchiveTypes.h"
+#include "Templates/Function.h"
+#include "Templates/UniquePtr.h"
+#include "UObject/NameTypes.h"
+
+#if WITH_EDITORONLY_DATA
+
+class IShaderFormat;
+#if WITH_EDITOR
+namespace UE::Cook::Artifact { struct FUpdateOplogPackagesContext; }
+#endif
+
+/**
+ * Shader code statistics structure used by FEditorShaderCodeArchive.
+ */
+struct FShaderCodeStats
+{
+	struct FPerTypeStats
+	{
+		int64 ShadersSize = 0;
+		int32 NumShaders = 0;
+	};
+
+	int32 NumShaderMaps = 0;
+	/** Maps from a shader type hash (see FShaderMapResourceCode::FShaderEditorOnlyDataEntry::ShaderTypeHashes) to its statistics. */
+	TMap<uint64, FPerTypeStats> ShaderTypeStats;
+
+	FORCEINLINE void AddEntry(uint64 InShaderTypeHash, uint32 InShaderCodeSize)
+	{
+		FShaderCodeStats::FPerTypeStats& OutTypeCodeStats = ShaderTypeStats.FindOrAdd(InShaderTypeHash);
+		OutTypeCodeStats.NumShaders++;
+		OutTypeCodeStats.ShadersSize += InShaderCodeSize;
+	}
+};
+
+/**
+ * Wrapper class to create chunks of serialized shaders and write them to a *.ushaderbytecode output file during cooking.
+ * See also FEditorShaderStableInfo.
+ */
+class FEditorShaderCodeArchive
+{
+	using ECookShaderLibrarySource = ShaderCodeArchive::ECookShaderLibrarySource;
+	using ESaveToDiskTarget = ShaderCodeArchive::ESaveToDiskTarget;
+
+	/** Current library name, e.g. "Global" or "Lyra_Chunk0". */
+	FString LibraryName;
+
+	/** Combined shader format and platform name, e.g. "SF_VULKAN_SM6-VULKAN_SM6". */
+	const FName FormatAndPlatformName;
+
+	FSerializedShaderArchive SerializedShaders;
+
+	/**
+	 * The element at index N holds the ShaderCode for the element of SerializedShaders.ShaderEntries at index N.
+	 * In MultiprocessCooking elements can be empty if they have been transferred to the Director (bHasCopiedAndCleared will be true in this case).
+	 */
+	TArray<FSharedBuffer> ShaderCode;
+
+	/** A list of ShaderMaps that have not yet been copied to the CookDirector. Used only by MultiprocessCookWorkers. */
+	TSet<int32> ShaderMapsToCopy;
+
+	/** True if CopyToArchiveAndClear has been called, otherwise false. If false we avoid doing some tracking since it might never be used. */
+	bool bHasCopiedAndCleared = false;
+
+	const IShaderFormat* const Format = nullptr;
+
+public:
+	/** @InFormatAndPlatformName specifies the combined shader and platform format name, e.g. "SF_VULKAN_ES31_ANDROID-VULKAN_ES3_1_ANDROID". */
+	RENDERCORE_API FEditorShaderCodeArchive(FName InFormatAndPlatformName);
+	RENDERCORE_API ~FEditorShaderCodeArchive();
+
+	/** Returns true if the IShaderFormat interface this code archive was created with supports native shader archives. */
+	RENDERCORE_API bool SupportsShaderArchives() const;
+
+	/** Resets the library name and flushes all serialized shaders and their shader code. */
+	RENDERCORE_API void OpenLibrary(FString const& Name);
+
+	/** Clears the library name and validates the specified library name was opened previously. */
+	RENDERCORE_API void CloseLibrary(FString const& Name);
+
+	/** Returns true if the serialized shaders contain the specified shader map hash. */
+	FORCEINLINE bool HasShaderMap(const FSHAHash& Hash) const
+	{
+		return SerializedShaders.FindShaderMap(Hash) != INDEX_NONE;
+	}
+
+	/** Returns true if this archive has no serialized shaders. */
+	FORCEINLINE bool IsEmpty() const
+	{
+		return SerializedShaders.GetNumShaders() == 0;
+	}
+
+	/** Adds and deduplicates shader code from the specified shader map resoure to the list of serialized shaders. */
+	RENDERCORE_API int32 AddShaderCode(const FShaderMapResourceCode* Code, const FShaderMapAssetPaths& AssociatedAssets,
+		ECookShaderLibrarySource CookSource, FShaderCodeStats* OutCodeStats = nullptr);
+
+	/** Returns true if this archive has any data that can be copied in CopyToArchiveAndClear(). */
+	RENDERCORE_API bool HasDataToCopy() const;
+
+	/** Writes the specified transfer archive and transfer code with the compact binary writer.
+	 * This is only a member function for consistency with FEditorShaderStableInfo. */
+	RENDERCORE_API void CopyToCompactBinary(FCbWriter& Writer, FSerializedShaderArchive& TransferArchive, TArray<uint8>& TransferCode);
+
+	/** Loads (LoadFromCompactBinary) and appends (AppendFromArchive) all serialized shaders from the specified compact binary field to this archive. */
+	RENDERCORE_API bool AppendFromCompactBinary(FCbFieldView Field, ECookShaderLibrarySource CookSource, FShaderCodeStats* OutCodeStats = nullptr);
+
+	/** Copies all serialized shaders to the target archive and its flat shader code array.
+	 * This also flushes the shared buffer of all serialized shaders in this archive to release potentially large amounts of memory.
+	 * If output parameter bOutRanOutOfRoom is true after this call, additional serialized shaders should be copied into the next target archive. */
+	RENDERCORE_API void CopyToArchiveAndClear(FSerializedShaderArchive& TargetArchive, TArray<uint8>& TargetFlatShaderCode,
+		bool& bOutRanOutOfRoom, int64 MaxShaderSize, int64 MaxShaderCount, ECookShaderLibrarySource CookSource);
+
+	/** Appends all entries of the source archive to the list of serialized shaders in this archive, including copies of the shader code. */
+	RENDERCORE_API bool AppendFromArchive(const FSerializedShaderArchive& SourceArchive, TConstArrayView64<uint8> SourceFlatShaderCode,
+		ECookShaderLibrarySource CookSource, FShaderCodeStats* OutCodeStats = nullptr);
+
+	/** Produces another archive that contains the code only for these assets and generates a chunked name */
+	RENDERCORE_API TUniquePtr<FEditorShaderCodeArchive> CreateChunk(int32 ChunkId, const TSet<FName>& PackagesInChunk);
+
+	/** Produces another archive that contains the code only for these assets and uses a custom chunk name. */
+	RENDERCORE_API TUniquePtr<FEditorShaderCodeArchive> CreateNamedChunk(const FString& ChunkName, const TSet<FName>& PackagesInChunk);
+
+	/** Adds entries for the specified shader map from another archive with all its individual shaders and their code to this archive. */
+	RENDERCORE_API int32 AddShaderCode(int32 OtherShaderMapIndex, const FEditorShaderCodeArchive& OtherArchive, ECookShaderLibrarySource CookSource);
+
+	/** Shortcut to add shader code library from directory . */
+	FORCEINLINE void AddShaderCodeLibraryFromDirectory(const FString& BaseDir, ECookShaderLibrarySource CookSource, FShaderCodeStats* OutCodeStats = nullptr)
+	{
+		AddShaderCodeLibraryByName(BaseDir, LibraryName, CookSource, OutCodeStats);
+	}
+
+	/** Reads all existing shader code entries from file and adds them to this library. This marks the end of populating this archive. */
+	FORCEINLINE void FinishPopulate(const FString& OutputDir, ECookShaderLibrarySource CookSource, FShaderCodeStats* OutCodeStats = nullptr)
+	{
+		AddExistingShaderCodeLibrary(OutputDir, CookSource, OutCodeStats);
+	}
+
+	/** Saves this shader archive to disk:
+	 * Shader entries and all their code to *.ushaderbytecode file,
+	 * Asset-shader associations to *.assetinfo.json file,
+	 * and shader type information to *.stinfo file. */
+	RENDERCORE_API bool SaveToDisk(const FString& OutputDir, const FString& MetaOutputDir,
+		ESaveToDiskTarget Target = ESaveToDiskTarget::Staging, TArray<FString>* OutputFilenames = nullptr);
+
+	/** Creates a native shader library for this archive, returns the output filenames,
+	 * and deletes the shader code and pipeline files as they are no longer needed. */
+	RENDERCORE_API bool PackageNativeShaderLibrary(const FString& ShaderCodeDir, TArray<FString>* OutputFilenames = nullptr);
+
+	/** Prints the shader library statistics to the log. */
+	RENDERCORE_API void DumpStatsAndDebugInfo();
+
+	/** Invokes the specified callback for each asset that references the specified shader type hash.
+	 * This can be slow and should only be used for debugging code or fallback conditions for assertions.
+	 * The callback must return true to continue iterating. */
+	RENDERCORE_API void ForEachAssetReferencingShaderType(uint64 InShaderTypeHash, FShaderMapAssetAssociations::FAssetFilterFunctionRef AssetReferenceCallback) const;
+
+#if WITH_EDITOR
+public:
+	/** Resets bReferencedByOplog and bReferencedByStaging to false for all shadermap-asset associations. */
+	void InitializeReferenceTracking();
+
+	/** Called by the cooker to report which assets are referenced in the current cook or in incrementalcook oplog. */
+	void UpdateOplogPackages(UE::Cook::Artifact::FUpdateOplogPackagesContext& Context);
+
+	/** Called by the cooker to ask the archive to prune no-longer referenced shaders. */
+	void Prune();
+#endif // WITH_EDITOR
+
+public:
+	/** Loads all existing shader code libraries from the new meta data directory and patches their entries into this archive. */
+	RENDERCORE_API static bool CreatePatchLibrary(FName FormatName, FString const& LibraryName, TArray<FString> const& OldMetaDataDirs, FString const& NewMetaDataDir, FString const& OutDir, bool bNativeFormat);
+
+#if WITH_EDITOR
+	TUniquePtr<FEditorShaderCodeArchive> CreateArchiveFromAssetsReferencedByStaging();
+	TUniquePtr<FEditorShaderCodeArchive> CreateArchiveFromAssetsReferencedByOplog();
+	TUniquePtr<FEditorShaderCodeArchive> CreateArchiveFromAssetsOnlyReferencedByOplog();
+#endif
+
+private:
+	// Primary function to create a new shader code archive with a filtered set of assets under the condition specified by a function callback.
+	TUniquePtr<FEditorShaderCodeArchive> CreateArchiveFromFilteredAssets(const FString& NewLibraryName,
+		FShaderMapAssetAssociations::FAssetFilterFunctionRef ShouldKeepAsset,
+		FShaderMapAssetAssociations::FShaderMapFilterFunction ShouldExcludeShaderMap = nullptr);
+
+	// Calls AddShaderCode() for all shader maps from the new library unless they were already added in one of the old libraries.
+	void MakePatchLibrary(TArray<FEditorShaderCodeArchive*> const& OldLibraries, FEditorShaderCodeArchive const& NewLibrary);
+
+	struct FShaderCodeReadInfo
+	{
+		int64 ShaderCodeIndex;
+		int64 Offset;
+		int32 Size;
+	};
+
+	RENDERCORE_API void AddExistingShaderCodeLibrary(FString const& OutputDir, ECookShaderLibrarySource CookSource, FShaderCodeStats* OutCodeStats = nullptr);
+	bool LoadExistingShaderCodeLibrary(FString const& MetaDataDir, ECookShaderLibrarySource CookSource);
+	int32 PrepareAddShaderCode(int32 OtherShaderMapIndex, const FSerializedShaderArchive& OtherShaders, TArray<FShaderCodeReadInfo>& ReadInfo,
+		ECookShaderLibrarySource CookSource, FShaderCodeStats* OutCodeStats = nullptr);
+
+	enum class EShaderCodeReadStatus : uint8
+	{
+		OK,
+		FileNotFound,
+		MalformedArchive,
+		VersionMismatch,
+	};
+
+	RENDERCORE_API EShaderCodeReadStatus AddShaderCodeLibraryByName(const FString& BaseDir, const FString& InLibraryName,
+		ECookShaderLibrarySource CookSource, FShaderCodeStats* OutCodeStats = nullptr);
+	EShaderCodeReadStatus ReadShaderCode(FArchive& Ar, TArray<FShaderCodeReadInfo>& ReadInfo);
+	EShaderCodeReadStatus ReadShaderCodeBlock(FArchive& Ar, int64 StartOffset, const TArrayView<FShaderCodeReadInfo>& BlockReadInfo);
+
+	void MarkShaderMapDirty(int32 ShaderMapIndex);
+	void MarkAllShaderMapsDirty();
+
+	// Wrapper to move all data from InArchive into this archive.
+	// Constant members such as Format and such must be equal since they must never change.
+	void AssignFrom(FEditorShaderCodeArchive& InArchive);
+
+};
+
+#endif // WITH_EDITORONLY_DATA
diff --git a/Engine/Source/Runtime/RenderCore/Internal/ShaderLibrary/ShaderCodeArchiveTypes.h b/Engine/Source/Runtime/RenderCore/Internal/ShaderLibrary/ShaderCodeArchiveTypes.h
new file mode 100644
--- /dev/null
+++ b/Engine/Source/Runtime/RenderCore/Internal/ShaderLibrary/ShaderCodeArchiveTypes.h
@@ -0,0 +1,18 @@
+// Copyright Epic Games, Inc. All Rights Reserved.
+
+/*=============================================================================
+	ShaderCodeArchiveTypes.h: Shared types for the shader code library
+=============================================================================*/
+
+#pragma once
+
+#include "CoreTypes.h"
+
+namespace ShaderCodeArchive
+{
+	enum class ESaveToDiskTarget : uint8
+	{
+		Staging = 0,
+		CookCache,
+	};
+} // namespace ShaderCodeArchive
diff --git a/Engine/Source/Runtime/RenderCore/Private/ShaderLibrary/EditorShaderCodeArchive.cpp b/Engine/Source/Runtime/RenderCore/Private/ShaderLibrary/EditorShaderCodeArchive.cpp
--- a/Engine/Source/Runtime/RenderCore/Private/ShaderLibrary/EditorShaderCodeArchive.cpp
+++ b/Engine/Source/Runtime/RenderCore/Private/ShaderLibrary/EditorShaderCodeArchive.cpp
@@ -4,9 +4,9 @@
 	EditorShaderCodeArchive.cpp: Implementation of FEditorShaderCodeArchive
 =============================================================================*/
 
-#if WITH_EDITOR
+#include "ShaderLibrary/EditorShaderCodeArchive.h"
 
-#include "EditorShaderCodeArchive.h"
+#if WITH_EDITORONLY_DATA
 
 #include "Async/ParallelFor.h"
 #include "Cooker/CookArtifact.h"
@@ -29,29 +29,26 @@
 
 using namespace UE::ShaderLibrary::Private;
 
-static const IShaderFormat* GetShaderFormatFromName(FName InFormatName)
+// Returns the shader format interface for the specified format and platform name.
+// This is used to determine whether native shader libraries are supported in FEditorShaderCodeArchive::SupportsShaderArchives().
+static const IShaderFormat* GetShaderFormatFromName(FName InFormatAndPlatformName)
 {
 	// Parse library format name into shader format name, e.g. "SF_VULKAN_ES31_ANDROID-VULKAN_ES3_1_ANDROID" to "SF_VULKAN_ES31_ANDROID".
 	TArray<FString> Components;
-	const FString Name = InFormatName.ToString();
+	const FString Name = InFormatAndPlatformName.ToString();
 	Name.ParseIntoArray(Components, TEXT("-"));
-
-	if (Components.IsEmpty())
+	if (!ensureMsgf(Components.Num() == 2, TEXT("Failed to parse shader format and platform from name '%s'"), *InFormatAndPlatformName.ToString()))
 	{
 		return nullptr;
 	}
-
 	const FName ShaderFormatName(Components[0]);
 	return GetTargetPlatformManagerRef().FindShaderFormat(ShaderFormatName);
 }
 
-FEditorShaderCodeArchive::FEditorShaderCodeArchive(FName InFormat, bool bInNeedsDeterministicOrder)
-	: FormatName(InFormat)
-	, Format(GetShaderFormatFromName(InFormat))
-	, bNeedsDeterministicOrder(bInNeedsDeterministicOrder)
+FEditorShaderCodeArchive::FEditorShaderCodeArchive(FName InFormatAndPlatformName)
+	: FormatAndPlatformName(InFormatAndPlatformName)
+	, Format(GetShaderFormatFromName(InFormatAndPlatformName))
 {
-	checkf(Format, TEXT("Failed to parse shader format from name '%s'"), *InFormat.ToString());
-
 	SerializedShaders.ShaderHashTable.Clear(0x10000);
 	SerializedShaders.ShaderMapHashTable.Clear(0x10000);
 }
@@ -61,6 +58,11 @@ FEditorShaderCodeArchive::~FEditorShaderCodeArchive()
 	// dummy
 }
 
+bool FEditorShaderCodeArchive::SupportsShaderArchives() const
+{
+	return Format && Format->SupportsShaderArchives();
+}
+
 void FEditorShaderCodeArchive::OpenLibrary(FString const& Name)
 {
 	check(LibraryName.Len() == 0);
@@ -76,8 +78,8 @@ void FEditorShaderCodeArchive::CloseLibrary(FString const& Name)
 	LibraryName = TEXT("");
 }
 
-int32 FEditorShaderCodeArchive::AddShaderCode(FShaderCodeStats& CodeStats, const FShaderMapResourceCode* Code,
-	const FShaderMapAssetPaths& AssociatedAssets, ECookShaderLibrarySource CookSource)
+int32 FEditorShaderCodeArchive::AddShaderCode(const FShaderMapResourceCode* Code,
+	const FShaderMapAssetPaths& AssociatedAssets, ECookShaderLibrarySource CookSource, FShaderCodeStats* OutCodeStats)
 {
 	int32 ShaderMapIndex = INDEX_NONE;
 
@@ -89,7 +91,10 @@ int32 FEditorShaderCodeArchive::AddShaderCode(FShaderCodeStats& CodeStats, const
 	const int32 NumShaders = Code->ShaderCodeResources.Num();
 	if (SerializedShaders.FindOrAddShaderMapEditor(Code->ResourceHash, ShaderMapIndex, &AssociatedAssets, CookSource))
 	{
-		CodeStats.NumShaderMaps++;
+		if (OutCodeStats)
+		{
+			OutCodeStats->NumShaderMaps++;
+		}
 		FShaderMapEntry& ShaderMapEntry = SerializedShaders.ShaderMapEntries[ShaderMapIndex];
 		ShaderMapEntry.NumShaders = NumShaders;
 		ShaderMapEntry.ShaderIndicesOffset = SerializedShaders.ShaderIndices.AddZeroed(NumShaders);
@@ -124,10 +129,10 @@ int32 FEditorShaderCodeArchive::AddShaderCode(FShaderCodeStats& CodeStats, const
 			if (HashIndex >= SerializedShaderTypes.Num() || SerializedShaderTypes[HashIndex] != ShaderTypeHash)
 			{
 				SerializedShaderTypes.Insert(ShaderTypeHash, HashIndex);
-				FShaderCodeStats::FPerTypeStats& TypeCodeStats = CodeStats.ShaderTypeStats.FindOrAdd(ShaderTypeHash);
-				int32 CodeSize = SerializedShaders.ShaderEntries[ShaderIndex].Size;
-				TypeCodeStats.NumShaders++;
-				TypeCodeStats.ShadersSize += CodeSize;
+				if (OutCodeStats)
+				{
+					OutCodeStats->AddEntry(ShaderTypeHash, SerializedShaders.ShaderEntries[ShaderIndex].Size);
+				}
 			}
 		}
 	}
@@ -171,8 +176,7 @@ void FEditorShaderCodeArchive::CopyToCompactBinary(FCbWriter& Writer, FSerialize
 	Writer.EndObject();
 }
 
-bool FEditorShaderCodeArchive::AppendFromCompactBinary(FCbFieldView Field, FShaderCodeStats& CodeStats,
-	ECookShaderLibrarySource CookSource)
+bool FEditorShaderCodeArchive::AppendFromCompactBinary(FCbFieldView Field, ECookShaderLibrarySource CookSource, FShaderCodeStats* OutCodeStats)
 {
 	FSerializedShaderArchive TransferArchive;
 	if (!LoadFromCompactBinary(Field["SerializedShaders"], TransferArchive))
@@ -182,7 +186,7 @@ bool FEditorShaderCodeArchive::AppendFromCompactBinary(FCbFieldView Field, FShad
 
 	FMemoryView FlatShaderCodeView = Field["ShaderCode"].AsBinaryView();
 	TConstArrayView64<uint8> TransferCode(reinterpret_cast<const uint8*>(FlatShaderCodeView.GetData()), FlatShaderCodeView.GetSize());
-	return AppendFromArchive(TransferArchive, TransferCode, CodeStats, CookSource);
+	return AppendFromArchive(TransferArchive, TransferCode, CookSource, OutCodeStats);
 }
 
 void FEditorShaderCodeArchive::CopyToArchiveAndClear(FSerializedShaderArchive& TargetArchive, TArray<uint8>& TargetFlatShaderCode,
@@ -249,7 +253,7 @@ void FEditorShaderCodeArchive::CopyToArchiveAndClear(FSerializedShaderArchive& T
 					TargetShaderEntry.Offset = TargetFlatShaderCode.Num();
 					check(SourceShaderEntry.Size == SourceShaderCode.GetSize());
 					if ((MaxShaderSize > 0 && TargetFlatShaderCode.Num() + (int64)SourceShaderCode.GetSize() > MaxShaderSize) ||
-						(MaxShaderCount > 0 && NumShadersSentWithCode > MaxShaderCount))
+						(MaxShaderCount > 0 && NumShadersSentWithCode + 1 > MaxShaderCount))
 					{
 						// We have to stop here to avoid overflowing the shader limit. Send the shaders we have accumulated
 						// but do not send&clear any other data.
@@ -263,9 +267,10 @@ void FEditorShaderCodeArchive::CopyToArchiveAndClear(FSerializedShaderArchive& T
 						// Keep any shaders we added before the one we just added
 						return;
 					}
-					++NumShadersSentWithCode;
 
+					// Append shader code to flat target array and keep track of how many shaders have been sent to target archive
 					TargetFlatShaderCode.Append(MakeArrayView(reinterpret_cast<const uint8*>(SourceShaderCode.GetData()), SourceShaderCode.GetSize()));
+					++NumShadersSentWithCode;
 
 					// Reset the ShaderCode reference to (potentially) save memory in the local process. The consumer of the TargetArchive and
 					// TargetFlatShaderCode will be the only one that needs to read it.
@@ -296,8 +301,7 @@ void FEditorShaderCodeArchive::CopyToArchiveAndClear(FSerializedShaderArchive& T
 }
 
 bool FEditorShaderCodeArchive::AppendFromArchive(
-	const FSerializedShaderArchive& SourceArchive, TConstArrayView64<uint8> SourceFlatShaderCode,
-	FShaderCodeStats& CodeStats, ECookShaderLibrarySource CookSource)
+	const FSerializedShaderArchive& SourceArchive, TConstArrayView64<uint8> SourceFlatShaderCode, ECookShaderLibrarySource CookSource, FShaderCodeStats* OutCodeStats)
 {
 	bool bOk = true;
 	FSerializedShaderArchive& TargetArchive = this->SerializedShaders;
@@ -319,9 +323,10 @@ bool FEditorShaderCodeArchive::AppendFromArchive(
 			if (HashIndex >= ExistingShaderTypeHashes.Num() || ExistingShaderTypeHashes[HashIndex] != ShaderTypeHash)
 			{
 				ExistingShaderTypeHashes.Insert(ShaderTypeHash, HashIndex);
-				FShaderCodeStats::FPerTypeStats& TypeCodeStats = CodeStats.ShaderTypeStats.FindOrAdd(ShaderTypeHash);
-				TypeCodeStats.NumShaders++;
-				TypeCodeStats.ShadersSize += SourceShaderEntry.Size;
+				if (OutCodeStats)
+				{
+					OutCodeStats->AddEntry(ShaderTypeHash, SourceShaderEntry.Size);
+				}
 			}
 		}
 
@@ -395,20 +400,25 @@ bool FEditorShaderCodeArchive::AppendFromArchive(
 			// Every shader in the SourceArchive should have already been added by the loop above over SourceArchive.ShaderHashes
 			check(!bShaderIsNew);
 		}
-		CodeStats.NumShaderMaps++;
+		if (OutCodeStats)
+		{
+			OutCodeStats->NumShaderMaps++;
+		}
 	}
 	return bOk;
 }
 
-TUniquePtr<FEditorShaderCodeArchive> FEditorShaderCodeArchive::CreateArchiveFromFilteredAssets(const FString& NewLibraryName, FSerializedShaderArchive::FAssetFilterFunctionRef ShouldKeepAsset)
+TUniquePtr<FEditorShaderCodeArchive> FEditorShaderCodeArchive::CreateArchiveFromFilteredAssets(const FString& NewLibraryName,
+	FShaderMapAssetAssociations::FAssetFilterFunctionRef ShouldKeepAsset,
+	FShaderMapAssetAssociations::FShaderMapFilterFunction ShouldExcludeShaderMap)
 {
 	checkf(!bHasCopiedAndCleared, TEXT("It is not valid to call CreateArchiveFromFilteredAssets on an FEditorShaderCodeArchive that has sent its shaders to another process by calling CopyAndClear."));
 
-	TUniquePtr<FEditorShaderCodeArchive> NewArchive(new FEditorShaderCodeArchive(FormatName, bNeedsDeterministicOrder));
+	TUniquePtr<FEditorShaderCodeArchive> NewArchive(new FEditorShaderCodeArchive(FormatAndPlatformName));
 	NewArchive->OpenLibrary(NewLibraryName);
 
 	TArray<int32> ShaderCodeEntriesNeeded;	// this array is filled with the indices from the existing ShaderCode that will need to be taken
-	NewArchive->SerializedShaders.CreateFromFilteredAssets(SerializedShaders, ShaderCodeEntriesNeeded, ShouldKeepAsset);
+	NewArchive->SerializedShaders.CreateFromFilteredAssets(SerializedShaders, ShaderCodeEntriesNeeded, ShouldKeepAsset, ShouldExcludeShaderMap);
 
 	// extra integrity check
 	checkf(ShaderCodeEntriesNeeded.Num() == NewArchive->SerializedShaders.ShaderHashes.Num(),
@@ -428,11 +438,13 @@ TUniquePtr<FEditorShaderCodeArchive> FEditorShaderCodeArchive::CreateArchiveFrom
 	return NewArchive;
 }
 
+#if WITH_EDITOR
+
 TUniquePtr<FEditorShaderCodeArchive> FEditorShaderCodeArchive::CreateArchiveFromAssetsReferencedByStaging()
 {
 	return CreateArchiveFromFilteredAssets(
 		LibraryName,
-		[](FName AssetName, const FSerializedShaderArchive::FShaderMapAssetAssociations::FAssociatedAssetData& AssetData)
+		[](FName AssetName, const FShaderMapAssetAssociations::FAssociatedAssetData& AssetData) -> bool
 		{
 			return AssetData.bReferencedByStaging;
 		});
@@ -442,7 +454,7 @@ TUniquePtr<FEditorShaderCodeArchive> FEditorShaderCodeArchive::CreateArchiveFrom
 {
 	return CreateArchiveFromFilteredAssets(
 		LibraryName,
-		[](FName AssetName, const FSerializedShaderArchive::FShaderMapAssetAssociations::FAssociatedAssetData& AssetData)
+		[](FName AssetName, const FShaderMapAssetAssociations::FAssociatedAssetData& AssetData) -> bool
 		{
 			return AssetData.bReferencedByOplog;
 		});
@@ -452,17 +464,24 @@ TUniquePtr<FEditorShaderCodeArchive> FEditorShaderCodeArchive::CreateArchiveFrom
 {
 	return CreateArchiveFromFilteredAssets(
 		LibraryName,
-		[](FName AssetName, const FSerializedShaderArchive::FShaderMapAssetAssociations::FAssociatedAssetData& AssetData)
+		[](FName AssetName, const FShaderMapAssetAssociations::FAssociatedAssetData& AssetData) -> bool
 		{
 			return AssetData.bReferencedByOplog && !AssetData.bReferencedByStaging;
 		});
 }
 
-TUniquePtr<FEditorShaderCodeArchive> FEditorShaderCodeArchive::CreateChunk(int ChunkId, const TSet<FName>& PackagesInChunk)
+#endif // WITH_EDITOR
+
+TUniquePtr<FEditorShaderCodeArchive> FEditorShaderCodeArchive::CreateChunk(int32 ChunkId, const TSet<FName>& PackagesInChunk)
+{
+	return CreateNamedChunk(UE::ShaderLibrary::Private::GetShaderLibraryNameForChunk(LibraryName, ChunkId), PackagesInChunk);
+}
+
+TUniquePtr<FEditorShaderCodeArchive> FEditorShaderCodeArchive::CreateNamedChunk(const FString& ChunkName, const TSet<FName>& PackagesInChunk)
 {
 	return CreateArchiveFromFilteredAssets(
-		UE::ShaderLibrary::Private::GetShaderLibraryNameForChunk(LibraryName, ChunkId),
-		[&PackagesInChunk](FName AssetName, const FSerializedShaderArchive::FShaderMapAssetAssociations::FAssociatedAssetData& AssetData)
+		ChunkName,
+		[&PackagesInChunk](FName AssetName, const FShaderMapAssetAssociations::FAssociatedAssetData& AssetData) -> bool
 		{
 			return PackagesInChunk.Contains(AssetName);
 		});
@@ -508,11 +527,12 @@ int32 FEditorShaderCodeArchive::AddShaderCode(int32 OtherShaderMapIndex, const F
 	return ShaderMapIndex;
 }
 
-int32 FEditorShaderCodeArchive::PrepareAddShaderCode(FShaderCodeStats& CodeStats,
+int32 FEditorShaderCodeArchive::PrepareAddShaderCode(
 	int32 OtherShaderMapIndex,
 	const FSerializedShaderArchive& OtherShaders,
 	TArray<FShaderCodeReadInfo>& ReadInfo,
-	ECookShaderLibrarySource CookSource)
+	ECookShaderLibrarySource CookSource,
+	FShaderCodeStats* OutCodeStats)
 {
 	int32 ShaderMapIndex = 0;
 	const FSHAHash& OtherShaderMapHash = OtherShaders.ShaderMapHashes[OtherShaderMapIndex];
@@ -520,7 +540,10 @@ int32 FEditorShaderCodeArchive::PrepareAddShaderCode(FShaderCodeStats& CodeStats
 	bool bIsNew = SerializedShaders.FindOrAddShaderMapEditor(OtherShaderMapHash, ShaderMapIndex, OtherAssets, CookSource);
 	if (bIsNew)
 	{
-		CodeStats.NumShaderMaps++;
+		if (OutCodeStats)
+		{
+			OutCodeStats->NumShaderMaps++;
+		}
 
 		const FShaderMapEntry& OtherEntry = OtherShaders.ShaderMapEntries[OtherShaderMapIndex];
 		FShaderMapEntry& ShaderMapEntry = SerializedShaders.ShaderMapEntries[ShaderMapIndex];
@@ -554,10 +577,10 @@ int32 FEditorShaderCodeArchive::PrepareAddShaderCode(FShaderCodeStats& CodeStats
 				if (HashIndex >= SerializedShaderTypes.Num() || SerializedShaderTypes[HashIndex] != OtherShaderTypeHash)
 				{
 					SerializedShaderTypes.Insert(OtherShaderTypeHash, HashIndex);
-					FShaderCodeStats::FPerTypeStats& TypeCodeStats = CodeStats.ShaderTypeStats.FindOrAdd(OtherShaderTypeHash);
-					int32 CodeSize = SerializedShaders.ShaderEntries[ShaderIndex].Size;
-					TypeCodeStats.NumShaders++;
-					TypeCodeStats.ShadersSize += CodeSize;
+					if (OutCodeStats)
+					{
+						OutCodeStats->AddEntry(OtherShaderTypeHash, SerializedShaders.ShaderEntries[ShaderIndex].Size);
+					}
 				}
 			}
 			SerializedShaders.ShaderIndices[ShaderMapEntry.ShaderIndicesOffset + i] = ShaderIndex;
@@ -643,14 +666,14 @@ FEditorShaderCodeArchive::EShaderCodeReadStatus FEditorShaderCodeArchive::ReadSh
 }
 
 FEditorShaderCodeArchive::EShaderCodeReadStatus FEditorShaderCodeArchive::AddShaderCodeLibraryByName(
-	const FString& BaseDir, const FString& InLibraryName, FShaderCodeStats& CodeStats, ECookShaderLibrarySource CookSource)
+	const FString& BaseDir, const FString& InLibraryName, ECookShaderLibrarySource CookSource, FShaderCodeStats* OutCodeStats)
 {
 	using namespace UE::ShaderLibrary;
 
 	bool bFoundAny = false;
 	for (ESaveToDiskTarget Target : { ESaveToDiskTarget::Staging, ESaveToDiskTarget::CookCache })
 	{
-		const FString CodeArchiveFileName = GetCodeArchiveFilename(BaseDir, InLibraryName, FormatName, Target);
+		const FString CodeArchiveFileName = GetCodeArchiveFilename(BaseDir, InLibraryName, FormatAndPlatformName, Target);
 		TUniquePtr<FArchive> PrevCookedAr(CreateShaderFileReader(*CodeArchiveFileName));
 		if (!PrevCookedAr)
 		{
@@ -658,10 +681,7 @@ FEditorShaderCodeArchive::EShaderCodeReadStatus FEditorShaderCodeArchive::AddSha
 		}
 		bFoundAny = true;
 
-		uint32 Version = 0;
-		*PrevCookedAr << Version;
-
-		if (Version != GShaderCodeArchiveVersion)
+		if (!FSerializedShaderArchive::SerializeHeaderVersion(*PrevCookedAr))
 		{
 			return EShaderCodeReadStatus::VersionMismatch;
 		}
@@ -671,7 +691,7 @@ FEditorShaderCodeArchive::EShaderCodeReadStatus FEditorShaderCodeArchive::AddSha
 		*PrevCookedAr << PrevCookedShaders;
 
 		// check if it also contains the asset info file
-		FString AssetInfoFileName = GetShaderAssetInfoFilename(BaseDir, InLibraryName, FormatName, Target);
+		const FString AssetInfoFileName = GetShaderAssetInfoFilename(BaseDir, InLibraryName, FormatAndPlatformName, Target);
 		TUniquePtr<FArchive> ShaderAssetInfoReader(CreateShaderFileReader(*AssetInfoFileName));
 		if (PrevCookedShaders.LoadAssetInfo(ShaderAssetInfoReader.Get(), CookSource))
 		{
@@ -699,7 +719,7 @@ FEditorShaderCodeArchive::EShaderCodeReadStatus FEditorShaderCodeArchive::AddSha
 		}
 
 		// check if it also contains the shader type info file
-		FString TypeInfoFileName = GetShaderTypeInfoFilename(BaseDir, InLibraryName, FormatName, Target);
+		FString TypeInfoFileName = GetShaderTypeInfoFilename(BaseDir, InLibraryName, FormatAndPlatformName, Target);
 		TUniquePtr<FArchive> ShaderTypeInfoReader(CreateShaderFileReader(*TypeInfoFileName));
 		if (ShaderTypeInfoReader)
 		{
@@ -723,7 +743,7 @@ FEditorShaderCodeArchive::EShaderCodeReadStatus FEditorShaderCodeArchive::AddSha
 		TArray<FShaderCodeReadInfo> ReadInfo;
 		for (int32 PrevShaderMapIndex = 0; PrevShaderMapIndex < PrevCookedShaders.ShaderMapEntries.Num(); ++PrevShaderMapIndex)
 		{
-			PrepareAddShaderCode(CodeStats, PrevShaderMapIndex, PrevCookedShaders, ReadInfo, CookSource);
+			PrepareAddShaderCode(PrevShaderMapIndex, PrevCookedShaders, ReadInfo, CookSource, OutCodeStats);
 		}
 		EShaderCodeReadStatus Status = ReadShaderCode(*PrevCookedAr, ReadInfo);
 		if (Status != EShaderCodeReadStatus::OK)
@@ -738,27 +758,26 @@ bool FEditorShaderCodeArchive::LoadExistingShaderCodeLibrary(FString const& Meta
 {
 	using namespace UE::ShaderLibrary;
 
-	FString BaseDir = MetaDataDir / TEXT("ShaderLibrarySource");
-	FShaderCodeStats CodeStats;
-	EShaderCodeReadStatus Status = AddShaderCodeLibraryByName(BaseDir, LibraryName, CodeStats, CookSource);
+	const FString BaseDir = MetaDataDir / TEXT("ShaderLibrarySource");
+	const EShaderCodeReadStatus Status = AddShaderCodeLibraryByName(BaseDir, LibraryName, CookSource);
 	switch (Status)
 	{
 	case EShaderCodeReadStatus::FileNotFound:
 	{
 		UE_LOG(LogShaderLibrary, Error, TEXT("Failed to load shader code library from %s; file not found"),
-			*GetCodeArchiveFilename(BaseDir, LibraryName, FormatName, ESaveToDiskTarget::Staging));
+			*GetCodeArchiveFilename(BaseDir, LibraryName, FormatAndPlatformName, ESaveToDiskTarget::Staging));
 		return false;
 	}
 	case EShaderCodeReadStatus::MalformedArchive:
 	{
 		UE_LOG(LogShaderLibrary, Error, TEXT("Failed to deserialized code from library %s"),
-			*GetCodeArchiveFilename(BaseDir, LibraryName, FormatName, ESaveToDiskTarget::Staging));
+			*GetCodeArchiveFilename(BaseDir, LibraryName, FormatAndPlatformName, ESaveToDiskTarget::Staging));
 		return false;
 	}
 	case EShaderCodeReadStatus::VersionMismatch:
 	{
 		UE_LOG(LogShaderLibrary, Error, TEXT("Failed to deserialize shader code from %s because the archive format is incompatible with the current version %u"),
-			*GetCodeArchiveFilename(BaseDir, LibraryName, FormatName, ESaveToDiskTarget::Staging), GShaderCodeArchiveVersion);
+			*GetCodeArchiveFilename(BaseDir, LibraryName, FormatAndPlatformName, ESaveToDiskTarget::Staging), GShaderCodeArchiveVersion);
 		return false;
 	}
 	case EShaderCodeReadStatus::OK:
@@ -767,14 +786,13 @@ bool FEditorShaderCodeArchive::LoadExistingShaderCodeLibrary(FString const& Meta
 	}
 }
 
-void FEditorShaderCodeArchive::AddExistingShaderCodeLibrary(FString const& OutputDir, FShaderCodeStats& CodeStats,
-	ECookShaderLibrarySource CookSource)
+void FEditorShaderCodeArchive::AddExistingShaderCodeLibrary(FString const& OutputDir, ECookShaderLibrarySource CookSource, FShaderCodeStats* OutCodeStats)
 {
 	check(LibraryName.Len() > 0);
 	static const FRegexPattern ShaderCodePattern(ShaderCodePatternStr);
 	static const FRegexPattern ShaderArchivePattern(ShaderArchivePatternStr);
 
-	const FString FormatNameStr = FormatName.ToString();
+	const FString FormatNameStr = FormatAndPlatformName.ToString();
 	const FString ShaderIntermediateLocation = GetShaderCodeLibraryDebugInfoPath() / FormatNameStr;
 
 	TArray<FString> ShaderFiles;
@@ -792,7 +810,7 @@ void FEditorShaderCodeArchive::AddExistingShaderCodeLibrary(FString const& Outpu
 
 			if (FoundLibraryName.StartsWith(LibraryName) && FoundFormatName.Equals(FormatNameStr, ESearchCase::IgnoreCase))
 			{
-				AddShaderCodeLibraryByName(OutputDir, FoundLibraryName, CodeStats, CookSource);
+				AddShaderCodeLibraryByName(OutputDir, FoundLibraryName, CookSource, OutCodeStats);
 			}
 		}
 
@@ -804,15 +822,16 @@ void FEditorShaderCodeArchive::AddExistingShaderCodeLibrary(FString const& Outpu
 
 			if (FoundArchiveName.StartsWith(LibraryName) && FoundFormatName.Equals(FormatNameStr, ESearchCase::IgnoreCase))
 			{
-				AddShaderCodeLibraryByName(OutputDir, FoundArchiveName, CodeStats, CookSource);
+				AddShaderCodeLibraryByName(OutputDir, FoundArchiveName, CookSource, OutCodeStats);
 			}
 		}
 	}
 }
 
+#if WITH_EDITOR
 void FEditorShaderCodeArchive::InitializeReferenceTracking()
 {
-	for (TPair<FName, FSerializedShaderArchive::FShaderMapAssetAssociations::FAssociatedAssetData>& Pair
+	for (TPair<FName, FShaderMapAssetAssociations::FAssociatedAssetData>& Pair
 		: SerializedShaders.GetShaderMapAssetAssociations().ViewAssets())
 	{
 		Pair.Value.bReferencedByOplog = false;
@@ -822,7 +841,7 @@ void FEditorShaderCodeArchive::InitializeReferenceTracking()
 
 void FEditorShaderCodeArchive::UpdateOplogPackages(UE::Cook::Artifact::FUpdateOplogPackagesContext& Context)
 {
-	using FShaderMapAssetAssociations = FSerializedShaderArchive::FShaderMapAssetAssociations;
+	using FShaderMapAssetAssociations = FShaderMapAssetAssociations;
 	using FOplogPackageData = UE::Cook::Artifact::FOplogPackageData;
 
 	const TMap<FName, FOplogPackageData>& OplogPackages = Context.GetOplogPackages();
@@ -853,8 +872,6 @@ void FEditorShaderCodeArchive::UpdateOplogPackages(UE::Cook::Artifact::FUpdateOp
 
 void FEditorShaderCodeArchive::Prune()
 {
-	using FShaderMapAssetAssociations = FSerializedShaderArchive::FShaderMapAssetAssociations;
-
 	if (!UE::ShaderLibrary::Private::bIncrementalCookPruningEnabled)
 	{
 		return;
@@ -863,15 +880,14 @@ void FEditorShaderCodeArchive::Prune()
 	TUniquePtr<FEditorShaderCodeArchive> PrunedArchive = this->CreateArchiveFromAssetsReferencedByOplog();
 	AssignFrom(*PrunedArchive);
 }
+#endif // WITH_EDITOR
 
 void FEditorShaderCodeArchive::AssignFrom(FEditorShaderCodeArchive& InArchive)
 {
-	checkf(this->FormatName == InArchive.FormatName,
-		TEXT("Cannot assign shader code archive due to `FormatName` mismatch: Expected '%s', but got '%s'"), *this->FormatName.ToString(), *InArchive.FormatName.ToString());
+	checkf(this->FormatAndPlatformName == InArchive.FormatAndPlatformName,
+		TEXT("Cannot assign shader code archive due to `FormatAndPlatformName` mismatch: Expected '%s', but got '%s'"), *this->FormatAndPlatformName.ToString(), *InArchive.FormatAndPlatformName.ToString());
 	checkf(this->Format == InArchive.Format,
 		TEXT("Cannot assign shader code archive due to `Format` mismatch: Expected %p, but got %p"), this->Format, InArchive.Format);
-	checkf(this->bNeedsDeterministicOrder == InArchive.bNeedsDeterministicOrder,
-		TEXT("Cannot assign shader code archive due to `bNeedsDeterministicOrder` mismatch: Expected %s, but got %s"), this->bNeedsDeterministicOrder ? TEXT("true") : TEXT("false"), InArchive.bNeedsDeterministicOrder ? TEXT("true") : TEXT("false"));
 
 	this->LibraryName = MoveTemp(InArchive.LibraryName);
 	this->SerializedShaders = MoveTemp(InArchive.SerializedShaders);
@@ -894,7 +910,7 @@ bool FEditorShaderCodeArchive::SaveToDisk(const FString& OutputDir, const FStrin
 			if (Result != COPY_OK)
 			{
 				UE_LOG(LogShaderLibrary, Error, TEXT("FEditorShaderCodeArchive copying %s to %s failed. Failed to finalize Shared Shader Library %s with format %s"),
-					*SourcePath, *DestinationPath, *LibraryName, *FormatName.ToString());
+					*SourcePath, *DestinationPath, *LibraryName, *FormatAndPlatformName.ToString());
 				return false;
 			}
 
@@ -905,28 +921,40 @@ bool FEditorShaderCodeArchive::SaveToDisk(const FString& OutputDir, const FStrin
 			return true;
 		};
 
+	auto CopyFilesToOutputDir = [this, &bSuccess, &CopyFile, &OutputDir, &MetaOutputDir, Target](
+		TFunction<FString(const FString&, const FString&, FName, ESaveToDiskTarget)> DestinationPathFunction,
+		const FString& SourcePath, TArray<FString>* OutputFilenames) -> void
+		{
+			// Copy to output location - support for iterative native library cooking
+			if (!CopyFile(DestinationPathFunction(OutputDir, LibraryName, FormatAndPlatformName, Target), SourcePath, OutputFilenames))
+			{
+				bSuccess = false;
+			}
+			else if (MetaOutputDir.Len())
+			{
+				if (!CopyFile(DestinationPathFunction(MetaOutputDir / TEXT("../ShaderLibrarySource"), LibraryName, FormatAndPlatformName, Target), SourcePath, nullptr))
+				{
+					bSuccess = false;
+				}
+			}
+		};
+
 	// Shader library
 	if (bSuccess && SerializedShaders.GetNumShaderMaps() > 0)
 	{
 		// Write to a intermediate file
 		const FString DebugInfoPath = GetShaderCodeLibraryDebugInfoPath();
-		FString IntermediateFormatPath = GetShaderCodeFilename(DebugInfoPath / FormatName.ToString(), LibraryName, FormatName, Target);
-		FString AssetInfoIntermediatePath = GetShaderAssetInfoFilename(DebugInfoPath / FormatName.ToString(), LibraryName, FormatName, Target);
-		FString TypeInfoIntermediatePath = GetShaderTypeInfoFilename(DebugInfoPath / FormatName.ToString(), LibraryName, FormatName, Target);
 
-		// save the ShaderMaps, Shaders, and the Shaders' byte code
+		// Save the ShaderMaps, Shaders, and the Shaders' byte code
 		constexpr bool bSaveOnlyAssetInfo = false; // A currently unused setting that we might want to bring back.
 		if (!bSaveOnlyAssetInfo)
 		{
-			FArchive* FileWriter = IFileManager::Get().CreateFileWriter(*IntermediateFormatPath, FILEWRITE_NoFail);
-			if (FileWriter)
+			const FString IntermediateFormatPath = GetShaderCodeFilename(DebugInfoPath / FormatAndPlatformName.ToString(), LibraryName, FormatAndPlatformName, Target);
+			if (TUniquePtr<FArchive> FileWriter = TUniquePtr<FArchive>(IFileManager::Get().CreateFileWriter(*IntermediateFormatPath, FILEWRITE_NoFail)))
 			{
-				check(Format);
-
 				SerializedShaders.Finalize();
 
-				int32 ArchiveVersion = GShaderCodeArchiveVersion;
-				*FileWriter << ArchiveVersion;
+				FSerializedShaderArchive::SerializeHeaderVersion(*FileWriter);
 
 				// Write shader library
 				*FileWriter << SerializedShaders;
@@ -937,65 +965,32 @@ bool FEditorShaderCodeArchive::SaveToDisk(const FString& OutputDir, const FStrin
 				}
 
 				FileWriter->Close();
-				delete FileWriter;
 
-				// Copy to output location - support for iterative native library cooking
-				bool bOutputDirSuccess = CopyFile(GetCodeArchiveFilename(OutputDir, LibraryName, FormatName, Target),
-					IntermediateFormatPath, OutputFilenames);
-				bSuccess = bSuccess && bOutputDirSuccess;
-				if (bOutputDirSuccess && MetaOutputDir.Len() && Target == ESaveToDiskTarget::Staging)
-				{
-					if (!CopyFile(GetCodeArchiveFilename(MetaOutputDir / TEXT("../ShaderLibrarySource"), LibraryName, FormatName, Target),
-						IntermediateFormatPath, nullptr))
-					{
-						bSuccess = false;
-					}
-				}
+				// Copies and changes filename from "ShaderCode-*.ushaderbytecode" (GetShaderCodeFilename) to "ShaderArchive-*.ushaderbytecode" (GetCodeArchiveFilename)
+				CopyFilesToOutputDir(GetCodeArchiveFilename, IntermediateFormatPath, OutputFilenames);
 			}
 
-			FArchive* ShaderTypeFileWriter = IFileManager::Get().CreateFileWriter(*TypeInfoIntermediatePath, FILEWRITE_NoFail);
-			if (ShaderTypeFileWriter)
+			const FString TypeInfoIntermediatePath = GetShaderTypeInfoFilename(DebugInfoPath / FormatAndPlatformName.ToString(), LibraryName, FormatAndPlatformName, Target);
+			if (TUniquePtr<FArchive> ShaderTypeFileWriter = TUniquePtr<FArchive>(IFileManager::Get().CreateFileWriter(*TypeInfoIntermediatePath, FILEWRITE_NoFail)))
 			{
 				*ShaderTypeFileWriter << SerializedShaders.ShaderTypes;
 				ShaderTypeFileWriter->Close();
-				delete ShaderTypeFileWriter;
 
-				// Copy to output location - support for iterative native library cooking
-				bool bOutputDirSuccess = CopyFile(GetShaderTypeInfoFilename(OutputDir, LibraryName, FormatName, Target),
-					TypeInfoIntermediatePath, OutputFilenames);
-				bSuccess = bSuccess && bOutputDirSuccess;
-				if (bOutputDirSuccess && MetaOutputDir.Len() && Target == ESaveToDiskTarget::Staging)
-				{
-					if (!CopyFile(GetShaderTypeInfoFilename(MetaOutputDir / TEXT("../ShaderLibrarySource"), LibraryName, FormatName, Target),
-						TypeInfoIntermediatePath, nullptr))
-					{
-						bSuccess = false;
-					}
-				}
+				// Copies "ShaderTypeInfo-*.stinfo" to output directory
+				CopyFilesToOutputDir(GetShaderTypeInfoFilename, TypeInfoIntermediatePath, OutputFilenames);
 			}
 		}
 
-		// save asset info
+		// Save asset info
 		{
-			FArchive* AssetInfoWriter = IFileManager::Get().CreateFileWriter(*AssetInfoIntermediatePath, FILEWRITE_NoFail);
-			if (AssetInfoWriter)
+			const FString AssetInfoIntermediatePath = GetShaderAssetInfoFilename(DebugInfoPath / FormatAndPlatformName.ToString(), LibraryName, FormatAndPlatformName, Target);
+			if (TUniquePtr<FArchive> AssetInfoWriter = TUniquePtr<FArchive>(IFileManager::Get().CreateFileWriter(*AssetInfoIntermediatePath, FILEWRITE_NoFail)))
 			{
 				SerializedShaders.SaveAssetInfo(*AssetInfoWriter);
 				AssetInfoWriter->Close();
-				delete AssetInfoWriter;
 
-				bool bOutputDirSuccess = CopyFile(GetShaderAssetInfoFilename(OutputDir, LibraryName, FormatName, Target),
-					AssetInfoIntermediatePath, nullptr);
-				bSuccess = bSuccess && bOutputDirSuccess;
-				if (bOutputDirSuccess && MetaOutputDir.Len() && Target == ESaveToDiskTarget::Staging)
-				{
-					// copy asset info as well for debugging
-					if (!CopyFile(GetShaderAssetInfoFilename(MetaOutputDir / TEXT("../ShaderLibrarySource"), LibraryName, FormatName, Target),
-						AssetInfoIntermediatePath, nullptr))
-					{
-						bSuccess = false;
-					}
-				}
+				// Copies "ShaderAssetInfo-*.assetinfo.json" to output directory
+				CopyFilesToOutputDir(GetShaderAssetInfoFilename, AssetInfoIntermediatePath, nullptr);
 			}
 		}
 	}
@@ -1011,26 +1006,28 @@ bool FEditorShaderCodeArchive::PackageNativeShaderLibrary(const FString& ShaderC
 	}
 	checkf(!bHasCopiedAndCleared, TEXT("It is not valid to call PackageNativeShaderLibrary on an FEditorShaderCodeArchive that has sent its shaders to another process by calling CopyAndClear."));
 
-	FString IntermediateFormatPath = GetShaderDebugFolder(GetShaderCodeLibraryDebugInfoPath() / FormatName.ToString(), LibraryName, FormatName);
+	// If the shader format module was not loaded, we don't support packaging native libraries.
+	// UnrealPak for instance only loads MetalShaderFormat, which is the only format that supports native libraries anyway.
+	if (!Format)
+	{
+		return false;
+	}
+
+	FString IntermediateFormatPath = GetShaderDebugFolder(GetShaderCodeLibraryDebugInfoPath() / FormatAndPlatformName.ToString(), LibraryName, FormatAndPlatformName);
 	FString TempPath = IntermediateFormatPath / TEXT("NativeLibrary");
 
 	IFileManager::Get().MakeDirectory(*TempPath, true);
 	IFileManager::Get().MakeDirectory(*ShaderCodeDir, true);
 
-	EShaderPlatform Platform = ShaderFormatToLegacyShaderPlatform(FormatName);
-	const bool bOK = Format->CreateShaderArchive(LibraryName, FormatName, TempPath, ShaderCodeDir, IntermediateFormatPath, SerializedShaders, ShaderCode, OutputFilenames);
+	EShaderPlatform Platform = ShaderFormatToLegacyShaderPlatform(FormatAndPlatformName);
+	const bool bOK = Format->CreateShaderArchive(LibraryName, FormatAndPlatformName, TempPath, ShaderCodeDir, IntermediateFormatPath, SerializedShaders, ShaderCode, OutputFilenames);
 
 	if (bOK)
 	{
-		// Delete Shader code library / pipelines as we now have native versions
+		// Delete Shader code library as we now have native versions
 		// [RCL] 2020-12-02 FIXME: check if this doesn't ruin iterative cooking during Launch On
 		{
-			FString OutputFilePath = GetCodeArchiveFilename(ShaderCodeDir, LibraryName, FormatName,
-				ESaveToDiskTarget::Staging);
-			IFileManager::Get().Delete(*OutputFilePath);
-		}
-		{
-			FString OutputFilePath = GetPipelinesArchiveFilename(ShaderCodeDir, LibraryName, FormatName);
+			FString OutputFilePath = GetCodeArchiveFilename(ShaderCodeDir, LibraryName, FormatAndPlatformName, ESaveToDiskTarget::Staging);
 			IFileManager::Get().Delete(*OutputFilePath);
 		}
 	}
@@ -1065,13 +1062,13 @@ void FEditorShaderCodeArchive::MakePatchLibrary(TArray<FEditorShaderCodeArchive*
 		}
 	}
 }
-	
-bool FEditorShaderCodeArchive::CreatePatchLibrary(FName FormatName, FString const& LibraryName, TArray<FString> const& OldMetaDataDirs, FString const& NewMetaDataDir, FString const& OutDir, bool bNativeFormat, bool bNeedsDeterministicOrder)
+
+bool FEditorShaderCodeArchive::CreatePatchLibrary(FName InFormatAndPlatformName, FString const& LibraryName, TArray<FString> const& OldMetaDataDirs, FString const& NewMetaDataDir, FString const& OutDir, bool bNativeFormat)
 {
 	TArray<FEditorShaderCodeArchive*> OldLibraries;
 	for (FString const& OldMetaDataDir : OldMetaDataDirs)
 	{
-		FEditorShaderCodeArchive* OldLibrary = new FEditorShaderCodeArchive(FormatName, bNeedsDeterministicOrder);
+		FEditorShaderCodeArchive* OldLibrary = new FEditorShaderCodeArchive(InFormatAndPlatformName);
 		OldLibrary->OpenLibrary(LibraryName);
 		if (OldLibrary->LoadExistingShaderCodeLibrary(OldMetaDataDir, ECookShaderLibrarySource::PatchBase))
 		{
@@ -1079,37 +1076,36 @@ bool FEditorShaderCodeArchive::CreatePatchLibrary(FName FormatName, FString cons
 		}
 	}
 
-	FEditorShaderCodeArchive NewLibrary(FormatName, bNeedsDeterministicOrder);
+	FEditorShaderCodeArchive NewLibrary(InFormatAndPlatformName);
 	NewLibrary.OpenLibrary(LibraryName);
 	bool bOK = NewLibrary.LoadExistingShaderCodeLibrary(NewMetaDataDir, ECookShaderLibrarySource::PatchBase);
 	if (bOK)
 	{
-		FEditorShaderCodeArchive OutLibrary(FormatName, bNeedsDeterministicOrder);
-		FShaderCodeStats OutStats;
+		FEditorShaderCodeArchive OutLibrary(InFormatAndPlatformName);
 		OutLibrary.OpenLibrary(LibraryName);
 		OutLibrary.MakePatchLibrary(OldLibraries, NewLibrary);
 		bOK = OutLibrary.SerializedShaders.GetNumShaderMaps() > 0;
 		if (bOK)
 		{
 			FString Empty;
-			OutLibrary.FinishPopulate(OutDir, OutStats, ECookShaderLibrarySource::PatchBase);
+			OutLibrary.FinishPopulate(OutDir, ECookShaderLibrarySource::PatchBase);
 			bOK = OutLibrary.SaveToDisk(OutDir, Empty);
-			UE_CLOG(!bOK, LogShaderLibrary, Error, TEXT("Failed to save %s shader patch library %s, %s, %s"), bNativeFormat ? TEXT("native") : TEXT(""), *FormatName.ToString(), *LibraryName, *OutDir);
+			UE_CLOG(!bOK, LogShaderLibrary, Error, TEXT("Failed to save %s shader patch library %s, %s, %s"), bNativeFormat ? TEXT("native") : TEXT(""), *InFormatAndPlatformName.ToString(), *LibraryName, *OutDir);
 
-			if (bOK && bNativeFormat && OutLibrary.GetFormat()->SupportsShaderArchives())
+			if (bOK && bNativeFormat && OutLibrary.SupportsShaderArchives())
 			{
 				bOK = OutLibrary.PackageNativeShaderLibrary(OutDir);
-				UE_CLOG(!bOK, LogShaderLibrary, Error, TEXT("Failed to package native shader patch library %s, %s, %s"), *FormatName.ToString(), *LibraryName, *OutDir);
+				UE_CLOG(!bOK, LogShaderLibrary, Error, TEXT("Failed to package native shader patch library %s, %s, %s"), *InFormatAndPlatformName.ToString(), *LibraryName, *OutDir);
 			}
 		}
 		else
 		{
-			UE_LOG(LogShaderLibrary, Verbose, TEXT("No shaders to patch for library %s, %s, %s"), *FormatName.ToString(), *LibraryName, *OutDir);
+			UE_LOG(LogShaderLibrary, Verbose, TEXT("No shaders to patch for library %s, %s, %s"), *InFormatAndPlatformName.ToString(), *LibraryName, *OutDir);
 		}
 	}
 	else
 	{
-		UE_LOG(LogShaderLibrary, Error, TEXT("Failed to open the shader library to patch against %s, %s, %s"), *FormatName.ToString(), *LibraryName, *NewMetaDataDir);
+		UE_LOG(LogShaderLibrary, Error, TEXT("Failed to open the shader library to patch against %s, %s, %s"), *InFormatAndPlatformName.ToString(), *LibraryName, *NewMetaDataDir);
 	}
 
 	for (FEditorShaderCodeArchive* Lib : OldLibraries)
@@ -1124,7 +1120,7 @@ void FEditorShaderCodeArchive::DumpStatsAndDebugInfo()
 	bool bUseExtendedDebugInfo = UE::ShaderLibrary::Private::GProduceExtendedStats != 0;
 
 	UE_LOG(LogShaderLibrary, Display, TEXT(""));
-	UE_LOG(LogShaderLibrary, Display, TEXT("Shader Library '%s' (%s) Stats:"), *LibraryName, *FormatName.ToString());
+	UE_LOG(LogShaderLibrary, Display, TEXT("Shader Library '%s' (%s) Stats:"), *LibraryName, *FormatAndPlatformName.ToString());
 	UE_LOG(LogShaderLibrary, Display, TEXT("================="));
 
 	FSerializedShaderArchive::FDebugStats Stats;
@@ -1186,10 +1182,7 @@ void FEditorShaderCodeArchive::DumpStatsAndDebugInfo()
 			UE_LOG(LogShaderLibrary, Display, TEXT("No shader usage info is provided"));
 		}
 
-		const FString DebugInfoPath = GetShaderCodeLibraryDebugInfoPath();
-		const FString DebugLibFolder = GetShaderDebugFolder(DebugInfoPath / FormatName.ToString(), LibraryName, FormatName);
-		IFileManager::Get().MakeDirectory(*DebugLibFolder, true);
-
+		const FString DebugLibFolder = FSerializedShaderArchive::MakeDebugDirectory(LibraryName, FormatAndPlatformName);
 		{
 			const FString DumpFilename = DebugLibFolder / TEXT("Dump.txt");
 			if (TUniquePtr<FArchive> DumpAr = TUniquePtr<FArchive>(IFileManager::Get().CreateFileWriter(*DumpFilename)))
@@ -1200,7 +1193,7 @@ void FEditorShaderCodeArchive::DumpStatsAndDebugInfo()
 			}
 		}
 #if 0 // creating a graphviz graph - maybe one day we'll return to this
-		FString DebugGraphFolder = GetShaderDebugFolder(DebugInfoPath / FormatName.ToString(), LibraryName, FormatName);
+		FString DebugGraphFolder = GetShaderDebugFolder(DebugInfoPath / FormatAndPlatformName.ToString(), LibraryName, FormatAndPlatformName);
 		FString DebugGraphFile = DebugGraphFolder / TEXT("RelationshipGraph.gv");
 
 		IFileManager::Get().MakeDirectory(*DebugGraphFolder, true);
@@ -1228,7 +1221,7 @@ void FEditorShaderCodeArchive::DumpStatsAndDebugInfo()
 	UE_LOG(LogShaderLibrary, Display, TEXT("================="));
 }
 
-void FEditorShaderCodeArchive::ForEachAssetReferencingShaderType(uint64 InShaderTypeHash, FSerializedShaderArchive::FAssetFilterFunctionRef AssetReferenceCallback) const
+void FEditorShaderCodeArchive::ForEachAssetReferencingShaderType(uint64 InShaderTypeHash, FShaderMapAssetAssociations::FAssetFilterFunctionRef AssetReferenceCallback) const
 {
 	check(SerializedShaders.ShaderTypes.Num() == SerializedShaders.ShaderHashes.Num());
 	const int32 NumShaders = SerializedShaders.ShaderTypes.Num();
@@ -1250,7 +1243,7 @@ void FEditorShaderCodeArchive::ForEachAssetReferencingShaderType(uint64 InShader
 	}
 
 	// Find shader type hash in all of serialized shaders
-	const FSerializedShaderArchive::FShaderMapAssetAssociations& ShaderMapAssetAssoc = SerializedShaders.GetShaderMapAssetAssociations();
+	const FShaderMapAssetAssociations& ShaderMapAssetAssoc = SerializedShaders.GetShaderMapAssetAssociations();
 	for (int32 ShaderIndex = 0; ShaderIndex < NumShaders; ++ShaderIndex)
 	{
 		if (!SerializedShaders.ShaderTypes[ShaderIndex].Data.Contains(InShaderTypeHash))
@@ -1269,7 +1262,7 @@ void FEditorShaderCodeArchive::ForEachAssetReferencingShaderType(uint64 InShader
 			// Invoke input callback for each asset for this shadermap
 			for (const FName& AssetName : *ShaderMapAssetNames)
 			{
-				const FSerializedShaderArchive::FShaderMapAssetAssociations::FAssociatedAssetData* AssetData = ShaderMapAssetAssoc.FindAsset(AssetName);
+				const FShaderMapAssetAssociations::FAssociatedAssetData* AssetData = ShaderMapAssetAssoc.FindAsset(AssetName);
 				if (!AssetData)
 				{
 					continue;
@@ -1286,4 +1279,4 @@ void FEditorShaderCodeArchive::ForEachAssetReferencingShaderType(uint64 InShader
 	}
 }
 
-#endif // WITH_EDITOR
+#endif // WITH_EDITORONLY_DATA
diff --git a/Engine/Source/Runtime/RenderCore/Private/ShaderLibrary/EditorShaderCodeArchive.h b/Engine/Source/Runtime/RenderCore/Private/ShaderLibrary/EditorShaderCodeArchive.h
deleted file mode 100644
diff --git a/Engine/Source/Runtime/RenderCore/Private/ShaderLibrary/ShaderCodeArchive.cpp b/Engine/Source/Runtime/RenderCore/Private/ShaderLibrary/ShaderCodeArchive.cpp
--- a/Engine/Source/Runtime/RenderCore/Private/ShaderLibrary/ShaderCodeArchive.cpp
+++ b/Engine/Source/Runtime/RenderCore/Private/ShaderLibrary/ShaderCodeArchive.cpp
@@ -17,17 +17,18 @@
 #include "Serialization/MemoryReader.h"
 #include "Serialization/MemoryWriter.h"
 #include "ShaderCodeLibrary.h"
+#include "ShaderLibrary/ShaderCodeLibraryUtilities.h"
 #include "ShaderCore.h"
 #include "Stats/Stats.h"
 #include "Serialization/StaticMemoryReader.h"
 #include "ProfilingDebugging/IoStoreTrace.h"
+#include "IO/IoDispatcherBackend.h"
 
-#if WITH_EDITOR
+#if WITH_EDITORONLY_DATA
 #include "Misc/Optional.h"
 #include "Misc/StringBuilder.h"
 #include "Serialization/CompactBinarySerialization.h"
 #include "Serialization/CompactBinaryWriter.h"
-#include "ShaderLibrary/ShaderCodeLibraryUtilities.h"
 #include "Templates/Greater.h"
 #endif
 
@@ -157,7 +158,7 @@ bool FSerializedShaderArchive::FindOrAddShaderMapNonEditorInternal(const FSHAHas
 bool FSerializedShaderArchive::FindOrAddShaderMap(const FSHAHash& Hash, int32& OutIndex,
 	const FShaderMapAssetPaths* AssociatedAssets)
 {
-#if WITH_EDITOR
+#if WITH_EDITORONLY_DATA
 	FShaderMapAssetPaths Buffer;
 	if (!AssociatedAssets)
 	{
@@ -169,7 +170,7 @@ bool FSerializedShaderArchive::FindOrAddShaderMap(const FSHAHash& Hash, int32& O
 #endif
 }
 
-#if WITH_EDITOR
+#if WITH_EDITORONLY_DATA
 bool FSerializedShaderArchive::FindOrAddShaderMapEditor(const FSHAHash& Hash, int32& OutIndex,
 	const FShaderMapAssetPaths* AssociatedAssets, ECookShaderLibrarySource CookSource)
 {
@@ -284,7 +285,7 @@ bool FSerializedShaderArchive::FindOrAddShader(const FSHAHash& Hash, int32& OutI
 		ShaderEntries.AddDefaulted();
 		check(ShaderEntries.Num() == ShaderHashes.Num());
 		ShaderHashTable.Add(Key, OutIndex);
-#if WITH_EDITOR
+#if WITH_EDITORONLY_DATA
 		ShaderTypes.AddDefaulted();
 		checkf(ShaderTypes.Num() == ShaderEntries.Num(), TEXT("Mismatch in number of shader (%d) and shader type (%d) entries when adding shader to archive"), ShaderEntries.Num(), ShaderTypes.Num());
 #endif
@@ -302,13 +303,13 @@ void FSerializedShaderArchive::RemoveLastAddedShader()
 	ShaderHashTable.Remove(Key, ShaderIndex);
 	ShaderHashes.RemoveAt(ShaderIndex);
 	ShaderEntries.RemoveAt(ShaderIndex);
-#if WITH_EDITOR
+#if WITH_EDITORONLY_DATA
 	ShaderTypes.RemoveAt(ShaderIndex);
 	checkf(ShaderTypes.Num() == ShaderEntries.Num(), TEXT("Mismatch in number of shader (%d) and shader type (%d) entries when removing last added shader from archive"), ShaderEntries.Num(), ShaderTypes.Num());
 #endif
 }
 
-#if WITH_EDITOR
+#if WITH_EDITORONLY_DATA
 FCbWriter& operator<<(FCbWriter& Writer, const FSerializedShaderArchive& Archive)
 {
 	TArray64<uint8> SerializedBytes;
@@ -327,8 +328,7 @@ FCbWriter& operator<<(FCbWriter& Writer, const FSerializedShaderArchive& Archive
 
 		// Serialize is meant for runtime fields only, so copy the editor-only fields separately
 		Writer.BeginArray("ShaderMapAssetAssociations");
-		for (const TPair<FName, FSerializedShaderArchive::FShaderMapAssetAssociations::FAssociatedAssetData>& Pair
-			: Archive.ShaderMapAssetAssociations.ViewAssets())
+		for (const TPair<FName, FShaderMapAssetAssociations::FAssociatedAssetData>& Pair : Archive.ShaderMapAssetAssociations.ViewAssets())
 		{
 			Writer.BeginArray();
 			Writer << Pair.Key;
@@ -396,8 +396,7 @@ bool LoadFromCompactBinary(FCbFieldView Field, FSerializedShaderArchive& OutArch
 			FCbFieldViewIterator ArrayIter = AssetField.CreateViewIterator();
 			FName AssetName;
 			TArray<FSHAHash> ShaderMaps;
-			ShaderCodeArchive::ECookShaderLibrarySource LatestSource
-				= ShaderCodeArchive::ECookShaderLibrarySource::CurrentCook;
+			ShaderCodeArchive::ECookShaderLibrarySource LatestSource = ShaderCodeArchive::ECookShaderLibrarySource::CurrentCook;
 			bOk = LoadFromCompactBinary(*ArrayIter++, AssetName) && bOk;
 			uint8 LatestSourceInt = 0xFF;
 			if (LoadFromCompactBinary(*ArrayIter++, LatestSourceInt)
@@ -413,7 +412,7 @@ bool LoadFromCompactBinary(FCbFieldView Field, FSerializedShaderArchive& OutArch
 			bOk = LoadFromCompactBinary(*ArrayIter++, ShaderMaps) && bOk;
 			if (bOk && ShaderMaps.Num() > 0)
 			{
-				FSerializedShaderArchive::FShaderMapAssetAssociations::FAssociatedAssetData* AssetData;
+				FShaderMapAssetAssociations::FAssociatedAssetData* AssetData;
 				AssetData = &OutArchive.ShaderMapAssetAssociations.FindOrAddAsset(AssetName, ShaderMaps[0]);
 				AssetData->LatestSource = LatestSource;
 				for (int32 Index = 1; Index < ShaderMaps.Num(); ++Index)
@@ -450,7 +449,7 @@ bool LoadFromCompactBinary(FCbFieldView Field, FSerializedShaderArchive& OutArch
 
 	return bOk;
 }
-#endif //WITH_EDITOR
+#endif //WITH_EDITORONLY_DATA
 #endif //!USE_MMAPPED_SHADERARCHIVE
 
 
@@ -646,9 +645,22 @@ void FSerializedShaderArchive::DecompressShader(int32 Index, const TArray<FShare
 	}
 	else
 	{
-		uint32 UncompressedSize = Entry.GetUncompressedSize();
-		OutDecompressedShader.SetNum(UncompressedSize, EAllowShrinking::No);
-		ShaderCodeArchive::DecompressShaderWithOodle(OutDecompressedShader.GetData(), UncompressedSize, reinterpret_cast<const uint8*>(ShaderCode[Index].GetData()), Entry.Size);
+		OutDecompressedShader.SetNum(0, EAllowShrinking::No);
+		(void)FSerializedShaderArchive::DecompressShaderEntryAndAppend(Entry, static_cast<const uint8*>(ShaderCode[Index].GetData()), OutDecompressedShader);
+	}
+}
+
+void FSerializedShaderArchive::DecompressShaderEntryAndAppend(const FShaderCodeEntry& InEntry, const uint8* InCompressedShaderCode, TArray<uint8>& OutDecompressedShader)
+{
+	const int32 StartOffset = OutDecompressedShader.Num();
+	OutDecompressedShader.SetNum(StartOffset + InEntry.GetUncompressedSize(), EAllowShrinking::No);
+	if (!InEntry.IsCompressed())
+	{
+		FMemory::Memcpy(OutDecompressedShader.GetData() + StartOffset, InCompressedShaderCode, InEntry.GetUncompressedSize());
+	}
+	else
+	{
+		ShaderCodeArchive::DecompressShaderWithOodle(OutDecompressedShader.GetData() + StartOffset, InEntry.GetUncompressedSize(), InCompressedShaderCode, InEntry.Size);
 	}
 }
 
@@ -706,6 +718,34 @@ void FSerializedShaderArchive::Finalize()
 }
 #endif
 
+static void MakeDebugDirectoryTreeGuarded(const TCHAR* InPath)
+{
+	if (!IFileManager::Get().MakeDirectory(InPath, true))
+	{
+		UE_LOG(LogShaderLibrary, Warning, TEXT("Failed to create directory tree for debug output: %s"), InPath);
+	}
+}
+
+FString FSerializedShaderArchive::MakeDebugDirectory(const FString LibraryName, const FName& FormatAndPlatformName)
+{
+	using namespace UE::ShaderLibrary::Private;
+	FString DebugLibFolder = GetShaderCodeLibraryDebugInfoPath();
+	if (!LibraryName.IsEmpty() && FormatAndPlatformName != NAME_None)
+	{
+		DebugLibFolder = GetShaderDebugFolder(DebugLibFolder / FormatAndPlatformName.ToString(), LibraryName, FormatAndPlatformName);
+	}
+	MakeDebugDirectoryTreeGuarded(*DebugLibFolder);
+	return DebugLibFolder;
+}
+
+bool FSerializedShaderArchive::SerializeHeaderVersion(FArchive& Ar)
+{
+	using UE::ShaderLibrary::Private::GShaderCodeArchiveVersion;
+	uint32 Version = Ar.IsLoading() ? 0u : GShaderCodeArchiveVersion;
+	Ar << Version;
+	return Version == GShaderCodeArchiveVersion;
+}
+
 void FSerializedShaderArchive::Serialize(FArchive& Ar)
 {
 #if USE_MMAPPED_SHADERARCHIVE
@@ -772,7 +812,7 @@ void FSerializedShaderArchive::Serialize(FArchive& Ar)
 	}
 }
 
-#if WITH_EDITOR
+#if WITH_EDITORONLY_DATA
 void FSerializedShaderArchive::SaveAssetInfo(FArchive& Ar)
 {
 	if (Ar.IsSaving())
@@ -919,61 +959,6 @@ bool FSerializedShaderArchive::LoadAssetInfo(FArchive* Ar, ECookShaderLibrarySou
 	return true;
 }
 
-void FSerializedShaderArchive::FShaderMapAssetAssociations::ReserveAssets(int32 NumAssets)
-{
-	ShaderMapToAssets.Reserve(NumAssets);
-	AssetToShaderMaps.Reserve(NumAssets);
-}
-
-FSerializedShaderArchive::FShaderMapAssetAssociations::FAssociatedAssetData&
-FSerializedShaderArchive::FShaderMapAssetAssociations::FindOrAddAsset(FName Asset, const FSHAHash& ShaderMap)
-{
-	FAssociatedAssetData& Result = AssetToShaderMaps.FindOrAdd(Asset);
-	FShaderMapAssetPaths& Assets = ShaderMapToAssets.FindOrAdd(ShaderMap);
-	bool bAlreadyExists = false;
-	Assets.Add(Asset, &bAlreadyExists);
-	if (!bAlreadyExists)
-	{
-		Result.ShaderMaps.AddUnique(ShaderMap);
-	}
-	check(AssetToShaderMaps.IsEmpty() == ShaderMapToAssets.IsEmpty());
-	return Result;
-}
-
-void FSerializedShaderArchive::FShaderMapAssetAssociations::RemoveAsset(FName Asset)
-{
-	FAssociatedAssetData Removed;
-	if (AssetToShaderMaps.RemoveAndCopyValue(Asset, Removed))
-	{
-		for (const FSHAHash& ShaderMap : Removed.ShaderMaps)
-		{
-			FShaderMapAssetPaths* AssetsUsingShader = ShaderMapToAssets.Find(ShaderMap);
-			if (AssetsUsingShader)
-			{
-				AssetsUsingShader->Remove(Asset);
-				if (AssetsUsingShader->IsEmpty())
-				{
-					ShaderMapToAssets.Remove(ShaderMap);
-				}
-			}
-		}
-	}
-	check(AssetToShaderMaps.IsEmpty() == ShaderMapToAssets.IsEmpty());
-}
-
-void FSerializedShaderArchive::FShaderMapAssetAssociations::SortForSaving()
-{
-	ShaderMapToAssets.KeySort(
-		[](const FSHAHash& A, const FSHAHash& B)
-		{
-			return GetTypeHash(A) < GetTypeHash(B);
-		});
-	for (TPair<FSHAHash, FShaderMapAssetPaths>& Pair : ShaderMapToAssets)
-	{
-		Pair.Value.Sort(FNameLexicalLess());
-	}
-}
-
 void FSerializedShaderArchive::CreateAsChunkFrom(const FSerializedShaderArchive& Parent, const TSet<FName>& PackagesInChunk, TArray<int32>& OutShaderCodeEntriesNeeded)
 {
 	CreateFromFilteredAssets(Parent, OutShaderCodeEntriesNeeded,
@@ -983,7 +968,9 @@ void FSerializedShaderArchive::CreateAsChunkFrom(const FSerializedShaderArchive&
 		});
 }
 
-void FSerializedShaderArchive::CreateFromFilteredAssets(const FSerializedShaderArchive& Parent, TArray<int32>& OutShaderCodeEntriesNeeded, FAssetFilterFunctionRef ShouldKeepAsset)
+void FSerializedShaderArchive::CreateFromFilteredAssets(const FSerializedShaderArchive& Parent, TArray<int32>& OutShaderCodeEntriesNeeded,
+	FShaderMapAssetAssociations::FAssetFilterFunctionRef ShouldKeepAsset,
+	FShaderMapAssetAssociations::FShaderMapFilterFunction ShouldExcludeShaderMap)
 {
 	// we should begin with a clean slate
 	checkf(ShaderMapHashes.Num() == 0 && ShaderHashes.Num() == 0 && ShaderMapEntries.Num() == 0 && ShaderEntries.Num() == 0 && PreloadEntries.Num() == 0 && ShaderIndices.Num() == 0,
@@ -1013,7 +1000,7 @@ void FSerializedShaderArchive::CreateFromFilteredAssets(const FSerializedShaderA
 				bIncludeSM = true;
 			}
 		}
-		else
+		else if (!ShouldExcludeShaderMap || !ShouldExcludeShaderMap(ShaderMapHash))
 		{
 			// if any asset is in the chunk, include
 			for (const FName& AssetName : *Assets)
@@ -1334,7 +1321,7 @@ void FSerializedShaderArchive::DumpContentsInPlaintext(FString& OutText) const
 	OutText = FStringView(Out);
 }
 
-#endif // WITH_EDITOR
+#endif // WITH_EDITORONLY_DATA
 
 FShaderCodeArchive* FShaderCodeArchive::Create(EShaderPlatform InPlatform, FArchive& Ar, const FString& InDestFilePath, const FString& InLibraryDir, const FString& InLibraryName)
 {
@@ -1670,6 +1657,88 @@ TRefCountPtr<FRHIShader> FShaderCodeArchive::CreateShader(int32 Index, bool bReq
 	return Shader;
 }
 
+bool FShaderLibraryNameInfo::ParseFromFilename(const TCHAR* Filename)
+{
+	// Parse 
+	// ShaderArchive-<MyProject>_Chunk<N>-<ShaderFormat>-<ShaderPlatform>.ushaderbytecode
+	// E.g. "ShaderArchive-Lyra_Chunk0-PCD3D_SM6-PCD3D_SM6.ushaderbytecode
+	const FString BaseFilename = FPaths::GetBaseFilename(Filename);
+
+	TArray<FString> LibraryNameComponents;
+	BaseFilename.ParseIntoArray(LibraryNameComponents, TEXT("-"));
+	if (LibraryNameComponents.Num() != 3 && LibraryNameComponents.Num() != 4)
+	{
+		return false;
+	}
+
+	ChunkName = MoveTemp(LibraryNameComponents[1]);
+	if (LibraryNameComponents.Num() == 4)
+	{
+		ShaderFormatName = FName(LibraryNameComponents[2]);
+		ShaderPlatformName = FName(LibraryNameComponents[3]);
+	}
+	else
+	{
+		ShaderFormatName = NAME_None;
+		ShaderPlatformName = FName(LibraryNameComponents[2]);
+	}
+
+	return true;
+}
+
+FString FShaderLibraryNameInfo::GetAssetInfoFilename() const
+{
+	//@todo-lh - Should be merged with GetShaderAssetInfoFilename() in ShaderCodeLibrary.cpp
+	return FString::Printf(TEXT("ShaderAssetInfo-%s-%s.assetinfo.json"), *ChunkName, *GetFormatAndPlatformName());
+}
+
+FString FShaderLibraryNameInfo::GetTypeInfoFilename() const
+{
+	//@todo-lh - Should be merged with GetShaderTypeInfoFilename() in ShaderCodeLibrary.cpp
+	return FString::Printf(TEXT("ShaderTypeInfo-%s-%s.stinfo"), *ChunkName, *GetFormatAndPlatformName());
+}
+
+// Parse ID number from pakchunk name, e.g. return 100 and "iad" for "pakchunk100iad-Windows".
+bool FShaderLibraryNameInfo::ParsePakChunkId(const FStringView& PakChunkName, int32& OutChunkId, FString& OutChunkIdSuffix)
+{
+	const FStringView PakchunkPrefix = TEXTVIEW("pakchunk");
+	if (!PakChunkName.StartsWith(PakchunkPrefix))
+	{
+		return false;
+	}
+
+	// Parse ID number
+	const int32 ChunkIdStart = PakchunkPrefix.Len();
+	int32 ChunkIdEnd = ChunkIdStart;
+	while (ChunkIdEnd < PakChunkName.Len() && FChar::IsDigit(PakChunkName[ChunkIdEnd]))
+	{
+		++ChunkIdEnd;
+	}
+	if (ChunkIdEnd == ChunkIdStart)
+	{
+		return false;
+	}
+	OutChunkId = FCString::Atoi(*FString(PakChunkName.Mid(ChunkIdStart, ChunkIdEnd - ChunkIdStart)));
+
+	// Parse optional suffix
+	int32 ChunkIdSuffixEnd = ChunkIdEnd;
+	while (ChunkIdSuffixEnd < PakChunkName.Len() && PakChunkName[ChunkIdSuffixEnd] != TEXT('-'))
+	{
+		++ChunkIdSuffixEnd;
+	}
+
+	if (ChunkIdSuffixEnd == ChunkIdEnd)
+	{
+		OutChunkIdSuffix.Empty();
+	}
+	else
+	{
+		OutChunkIdSuffix = PakChunkName.Mid(ChunkIdEnd, ChunkIdSuffixEnd - ChunkIdEnd);
+	}
+
+	return true;
+}
+
 FIoChunkId FIoStoreShaderCodeArchive::GetShaderCodeArchiveChunkId(const FString& LibraryName, FName FormatName)
 {
 	FString Name = FString::Printf(TEXT("%s-%s"), *LibraryName, *FormatName.ToString());
@@ -1910,13 +1979,13 @@ void FShaderCodeArchiveProcessor::ProcessSerializedShaders(const FSerializedShad
 	checkf(OutHeader.ShaderGroupIoHashes.Num() == OutHeader.ShaderGroupEntries.Num(),
 		TEXT("Error creating IoStoreShaderArchive header - mismatch between shader group hashes and descriptors (%d descriptors, %d hashes). Bug in grouping logic?"),
 		OutHeader.ShaderGroupEntries.Num(), OutHeader.ShaderGroupIoHashes.Num());
-	checkf(OutHeader.ShaderGroupEntries.Num() != 0, TEXT("At least one group must have been created"));
 
+	const int32 NumShaderGroupEntries = OutHeader.ShaderGroupEntries.Num();
 	UE_LOG(LogShaderLibrary, Display,
 		TEXT("Created IoStoreShaderArchive header: shaders grouped in %d groups (%d of them didn't need new indices), average uncompressed size %llu bytes, min %u bytes, max %u bytes (r.ShaderCodeLibrary.MaxShaderGroupSize=%u)"),
-		OutHeader.ShaderGroupEntries.Num(),
-		OutHeader.ShaderGroupEntries.Num() - Stats.GroupsThatAppendedToShaderIndices,
-		Stats.TotalUncompressedMemory / static_cast<uint64>(OutHeader.ShaderGroupEntries.Num()),
+		NumShaderGroupEntries,
+		NumShaderGroupEntries - Stats.GroupsThatAppendedToShaderIndices,
+		NumShaderGroupEntries > 0 ? Stats.TotalUncompressedMemory / static_cast<uint64>(NumShaderGroupEntries) : 0llu,
 		Stats.MinGroupSize,
 		Stats.MaxGroupSize,
 		MaxUncompressedShaderGroupSize);
@@ -2213,6 +2282,47 @@ void FIoStoreShaderCodeArchive::SaveIoStoreShaderCodeArchive(const FIoStoreShade
 	OutLibraryAr << const_cast<FIoStoreShaderCodeArchiveHeader &>(Header);
 }
 
+// Dumps all shader, shader map, and shader group hashes and how they are connected to a text file for debugging purposes.
+static void DumpShaderCodeArchiveHashes(const FIoStoreShaderCodeArchiveHeader& InLibraryHeader, const TCHAR* InFilename)
+{
+	if (TUniquePtr<FArchive> DumpFile = TUniquePtr<FArchive>(IFileManager::Get().CreateFileWriter(InFilename)))
+	{
+		FString DumpContent;
+
+		DumpContent.Appendf(TEXT("--------------- ShaderGroupIoHashes: ---------------\n"));
+		for (int32 GroupIndex = 0; GroupIndex < InLibraryHeader.ShaderGroupIoHashes.Num(); ++GroupIndex)
+		{
+			const FIoChunkId& ShaderGroupChunkId = InLibraryHeader.ShaderGroupIoHashes[GroupIndex];
+			const FIoStoreShaderGroupEntry& GroupEntry = InLibraryHeader.ShaderGroupEntries[GroupIndex];
+			DumpContent.Appendf(TEXT("%s (%d shaders)\n"), *LexToString(ShaderGroupChunkId), GroupEntry.NumShaders);
+
+			for (uint32 ShaderIndex = 0; ShaderIndex < GroupEntry.NumShaders; ++ShaderIndex)
+			{
+				const int32 MappedShaderIndex = InLibraryHeader.ShaderIndices[GroupEntry.ShaderIndicesOffset + ShaderIndex];
+				DumpContent.Appendf(TEXT("  %s\n"), *InLibraryHeader.ShaderHashes[MappedShaderIndex].ToString());
+			}
+			DumpContent.Append(TEXT("\n"));
+		}
+		DumpContent.Append(TEXT("\n\n\n"));
+
+		DumpContent.Appendf(TEXT("--------------- ShaderHashes: ---------------\n"));
+		for (const FSHAHash& ShaderHash : InLibraryHeader.ShaderHashes)
+		{
+			DumpContent.Appendf(TEXT("%s\n"), *ShaderHash.ToString());
+		}
+		DumpContent.Append(TEXT("\n\n\n"));
+
+		DumpContent.Appendf(TEXT("--------------- ShaderMapHashes: ---------------\n"));
+		for (const FSHAHash& ShaderMapHash : InLibraryHeader.ShaderMapHashes)
+		{
+			DumpContent.Appendf(TEXT("%s\n"), *ShaderMapHash.ToString());
+		}
+		DumpContent.Append(TEXT("\n\n\n"));
+
+		DumpFile->Serialize(DumpContent.GetCharArray().GetData(), DumpContent.Len() * sizeof(TCHAR));
+	}
+}
+
 FIoStoreShaderCodeArchive* FIoStoreShaderCodeArchive::Create(EShaderPlatform InPlatform, const FString& InLibraryName, FIoDispatcher& InIoDispatcher)
 {
 	const FName PlatformName = FName(FDataDrivenShaderPlatformInfo::GetShaderFormat(InPlatform).ToString() + TEXT("-") + FDataDrivenShaderPlatformInfo::GetName(InPlatform).ToString());
@@ -2254,6 +2364,26 @@ FIoStoreShaderCodeArchive* FIoStoreShaderCodeArchive::Create(EShaderPlatform InP
 
 			Library->DebugVisualizer.Initialize(Library->Header.ShaderEntries.Num());
 
+			// Dump debug info if enabled
+			if (FParse::Param(FCommandLine::Get(), TEXT("DumpShaderLibraryHashes")))
+			{
+				if (!Library->Header.ShaderGroupIoHashes.IsEmpty() ||
+					!Library->Header.ShaderHashes.IsEmpty() ||
+					!Library->Header.ShaderMapHashes.IsEmpty())
+				{
+					const FString DumpPath = FPaths::Combine(FPaths::ProjectSavedDir(), TEXT("Shaders"));
+					if (!FPaths::DirectoryExists(DumpPath))
+					{
+						MakeDebugDirectoryTreeGuarded(*DumpPath);
+					}
+					const FString DumpFilename = FPaths::Combine(DumpPath, FString::Printf(TEXT("ShaderHashes-%s-%s.txt"), *InLibraryName, *PlatformName.ToString()));
+					UE_LOG(LogShaderLibrary, Display, TEXT("Dump IoStoreShaderCodeArchive '%s-%s' header to file '%s'"),
+						*InLibraryName, *PlatformName.ToString(), *DumpFilename);
+
+					DumpShaderCodeArchiveHashes(Library->Header, *DumpFilename);
+				}
+			}
+
 			UE_LOG(LogShaderLibrary, Display, TEXT("Using IoDispatcher for shader code library %s. Total %d unique shaders."), *InLibraryName, Library->Header.ShaderEntries.Num());
 			INC_DWORD_STAT_BY(STAT_Shaders_ShaderResourceMemory, Library->GetSizeBytes());
 			return Library;
@@ -2461,7 +2591,7 @@ bool FIoStoreShaderCodeArchive::PreloadShader(int32 ShaderIndex, FGraphEventArra
 	);
 }
 
-bool FIoStoreShaderCodeArchive::GroupOnlyContainsRaytracingShaders(int32 ShaderGroupIndex)
+bool FIoStoreShaderCodeArchive::GroupOnlyContainsRaytracingShaders(int32 ShaderGroupIndex) const
 {
 	const FIoStoreShaderGroupEntry& GroupEntry = Header.ShaderGroupEntries[ShaderGroupIndex];
 	for (uint32 ShaderIdxIdx = GroupEntry.ShaderIndicesOffset, StopIdxIdx = GroupEntry.ShaderIndicesOffset + GroupEntry.NumShaders; ShaderIdxIdx < StopIdxIdx; ++ShaderIdxIdx)
@@ -2514,6 +2644,28 @@ bool FIoStoreShaderCodeArchive::PreloadShaderMap(int32 ShaderMapIndex, FGraphEve
 	return true;
 }
 
+bool FIoStoreShaderCodeArchive::ResolveShaderMap(int32 ShaderMapIndex, FCoreDelegates::FIoChunkIdResolvedFunc IoChunkIdResolvedFunc) const
+{
+	// Invoke callback for the shader group of each shader within the specified shadermap
+	const FIoStoreShaderMapEntry& ShaderMapEntry = Header.ShaderMapEntries[ShaderMapIndex];
+	for (uint32 Index = 0u; Index < ShaderMapEntry.NumShaders; ++Index)
+	{
+		const int32 ShaderIndex = Header.ShaderIndices[ShaderMapEntry.ShaderIndicesOffset + Index];
+		const int32 ShaderGroupIndex = GetGroupIndexForShader(ShaderIndex);
+
+#if RHI_RAYTRACING
+		if (!IsRayTracingAllowed() && !IsCreateShadersOnLoadEnabled() && GroupOnlyContainsRaytracingShaders(ShaderGroupIndex))
+		{
+			continue;
+		}
+#endif
+
+		const FIoChunkId& ShaderGroupChunkId = Header.ShaderGroupIoHashes[ShaderGroupIndex];
+		IoChunkIdResolvedFunc(ShaderGroupChunkId);
+	}
+	return true;
+}
+
 bool FIoStoreShaderCodeArchive::PreloadShaderMap(int32 ShaderMapIndex, FCoreDelegates::FAttachShaderReadRequestFunc AttachShaderReadRequestFunc)
 {
 	LLM_SCOPE(ELLMTag::Shaders);
@@ -2656,7 +2808,7 @@ TRefCountPtr<FRHIShader> FIoStoreShaderCodeArchive::CreateShader(int32 ShaderInd
 	TRefCountPtr<FRHIShader> Shader;
 
 	const FIoStoreShaderCodeEntry& ShaderEntry = Header.ShaderEntries[ShaderIndex];
-	int32 GroupIndex = GetGroupIndexForShader(ShaderIndex);
+	const int32 GroupIndex = GetGroupIndexForShader(ShaderIndex);
 
 	// Preload shader group if it wasn't yet. This will also addref it so we can be sure it will exist.
 	FGraphEventArray Dummy;
@@ -2718,10 +2870,11 @@ TRefCountPtr<FRHIShader> FIoStoreShaderCodeArchive::CreateShader(int32 ShaderInd
 	{
 		FString ShaderGroupIoHash;
 		Header.ShaderGroupIoHashes[GroupIndex].ToString(ShaderGroupIoHash);
-		UE_LOG(LogShaderLibrary, Warning, TEXT("IoRequest backing buffer is not there when we need it, about to crash. Library '%s', shader hash '%s' (freq=%d), group IoChunkId '%s', started preloading at frame %d, current frame is %d, preloading %s"),
+		UE_LOG(LogShaderLibrary, Warning, TEXT("IoRequest backing buffer is not there when we need it, about to crash. Library '%s', shader hash '%s' (freq=%d), group %d IoChunkId '%s', started preloading at frame %d, current frame is %d, preloading %s"),
 			*GetName(),
 			*Header.ShaderHashes[ShaderIndex].ToString(),
 			ShaderEntry.Frequency,
+			GroupIndex,
 			*ShaderGroupIoHash,
 			PreloadEntryPtr->FramePreloadStarted,
 			GFrameNumber,
diff --git a/Engine/Source/Runtime/RenderCore/Private/ShaderLibrary/ShaderCodeLibrary.cpp b/Engine/Source/Runtime/RenderCore/Private/ShaderLibrary/ShaderCodeLibrary.cpp
--- a/Engine/Source/Runtime/RenderCore/Private/ShaderLibrary/ShaderCodeLibrary.cpp
+++ b/Engine/Source/Runtime/RenderCore/Private/ShaderLibrary/ShaderCodeLibrary.cpp
@@ -77,6 +77,8 @@ using namespace UE::ShaderLibrary::Private;
 
 namespace UE::ShaderLibrary::Private
 {
+	using ESaveToDiskTarget = ShaderCodeArchive::ESaveToDiskTarget;
+
 	static FDelegateHandle OnPakFileMountedDelegateHandle;
 	static FDelegateHandle OnPluginMountedDelegateHandle;
 	static FDelegateHandle OnPluginUnmountedDelegateHandle;
@@ -189,15 +191,37 @@ namespace UE::ShaderLibrary::Private
 			return MoveTemp(PresentChunks);
 		}
 
-		void OnPakFileMounted(const FMountedPakFileInfo& MountInfo, const FString& Directory)
+		// Transforms the chunk Id of the pak file for a named shader library component,
+		// e.g. pakchunk100 and pakchunk100iad must not point to the same ID 100, so the latter one will be transformed to a higher number.
+		TTuple<FString, int32> GetLibraryNameAndChunkId(const FMountedPakFileInfo& MountInfo) const
 		{
-			if (!GetPresentChunks().Contains(MountInfo.ChunkId))
+			constexpr int32 kChunkIdSuffixOffsetRange = 100000;
+			constexpr int32 kChunkIdSuffixBitmask = 0xFF;
+
+			FString ChunkLibraryName = GetShaderLibraryNameForChunk(LogicalName, MountInfo.ChunkId);
+			int32 ChunkId = MountInfo.ChunkId;
+
+			// If the pakchunk ID contains a suffix such as "iad", the shader library must have been created as a proxy target file and will also contain that suffix.
+			FString ChunkIdSuffix;
+			uint64 ChunkIdSuffixHash = 0;
+			if (FShaderLibraryNameInfo::ParsePakChunkId(FPaths::GetBaseFilename(MountInfo.PakFilename), ChunkId, ChunkIdSuffix) && !ChunkIdSuffix.IsEmpty())
 			{
-				FString ChunkLibraryName = GetShaderLibraryNameForChunk(LogicalName, MountInfo.ChunkId);
+				ChunkLibraryName.Append(ChunkIdSuffix);
+				ChunkIdSuffixHash = CityHash64(reinterpret_cast<const char*>(*ChunkIdSuffix), ChunkIdSuffix.Len() * sizeof(TCHAR));
+				ChunkId += (ChunkIdSuffixHash & kChunkIdSuffixBitmask) * kChunkIdSuffixOffsetRange;
+			}
 
+			return { MoveTemp(ChunkLibraryName), ChunkId };
+		}
+
+		void OnPakFileMounted(const FMountedPakFileInfo& MountInfo, const FString& Directory)
+		{
+			auto [ChunkLibraryName, ChunkId] = GetLibraryNameAndChunkId(MountInfo);
+			if (!GetPresentChunks().Contains(ChunkId))
+			{
 				// Parts of shader library might be in the UFS (e.g. .metallibs), hence we need to look for them in the
 				// appropriate directory.
-				OpenShaderCode(Directory, ChunkLibraryName, MountInfo.ChunkId);
+				OpenShaderCode(Directory, ChunkLibraryName, ChunkId);
 			}
 		}
 
@@ -855,12 +879,14 @@ public:
 			);
 #endif
 
+		FCoreDelegates::ResolvePackageShaderMaps.BindRaw(this, &FShaderLibrariesCollection::ResolvePackageShaderMapsDelegate);
 		FCoreDelegates::PreloadPackageShaderMaps.BindRaw(this, &FShaderLibrariesCollection::PreloadPackageShaderMapsDelegate);
 		FCoreDelegates::ReleasePreloadedPackageShaderMaps.BindRaw(this, &FShaderLibrariesCollection::ReleasePreloadedPackageShaderMapsDelegate);
 	}
 
 	~FShaderLibrariesCollection()
 	{
+		FCoreDelegates::ResolvePackageShaderMaps.Unbind();
 		FCoreDelegates::PreloadPackageShaderMaps.Unbind();
 		FCoreDelegates::ReleasePreloadedPackageShaderMaps.Unbind();
 
@@ -1359,8 +1385,6 @@ public:
 
 	void CookShaderFormats(TArray<FShaderLibraryCooker::FShaderFormatDescriptor> const& ShaderFormats)
 	{
-		bool bAtLeastOneFormatNeedsDeterminism = false;
-
 		for (const FShaderLibraryCooker::FShaderFormatDescriptor& Descriptor : ShaderFormats)
 		{
 			FName const& Format = Descriptor.ShaderFormat;
@@ -1371,16 +1395,11 @@ public:
 			if (!CodeArchive)
 			{
 				bIsEditorShaderCodeArchiveEmpty = false;
-				CodeArchive = new FEditorShaderCodeArchive(FName(FormatAndPlatformName), Descriptor.bNeedsDeterministicOrder);
+				CodeArchive = new FEditorShaderCodeArchive(FName(FormatAndPlatformName));
 				EditorShaderCodeArchive[Platform] = CodeArchive;
 				EditorArchivePipelines[Platform] = !bNativeFormat;
 			}
 			check(CodeArchive);
-
-			if (Descriptor.bNeedsDeterministicOrder)
-			{
-				bAtLeastOneFormatNeedsDeterminism = true;
-			}
 		}
 		for (const FShaderLibraryCooker::FShaderFormatDescriptor& Descriptor : ShaderFormats)
 		{
@@ -1420,7 +1439,7 @@ public:
 		FEditorShaderCodeArchive* CodeArchive = EditorShaderCodeArchive[Platform];
 		checkf(CodeArchive, TEXT("EditorShaderCodeArchive for (EShaderPlatform)%d is null!"), (int32)Platform);
 
-		CodeArchive->AddShaderCode(CodeStats, Code, AssociatedAssets, ShaderCodeArchive::ECookShaderLibrarySource::CurrentCook);
+		CodeArchive->AddShaderCode(Code, AssociatedAssets, ShaderCodeArchive::ECookShaderLibrarySource::CurrentCook, &CodeStats);
 	}
 
 	void CopyToCompactBinaryAndClear(FCbWriter& Writer, bool& bOutHasData, bool& bOutRanOutOfRoom, int64 MaxShaderSize)
@@ -1528,8 +1547,8 @@ public:
 				if (CodeArchive)
 				{
 					FShaderCodeStats& CodeStats = EditorShaderCodeStats[Platform];
-					bOk = CodeArchive->AppendFromCompactBinary(CodeArchiveField, CodeStats,
-						ShaderCodeArchive::ECookShaderLibrarySource::CurrentCook) && bOk;
+					bOk = CodeArchive->AppendFromCompactBinary(CodeArchiveField,
+						ShaderCodeArchive::ECookShaderLibrarySource::CurrentCook, &CodeStats) && bOk;
 				}
 				else
 				{
@@ -1578,8 +1597,8 @@ public:
 			FEditorShaderCodeArchive* CodeArchive = EditorShaderCodeArchive[SPlatform];
 			if (CodeArchive)
 			{
-				CodeArchive->FinishPopulate(ShaderCodeDir, CodeStats,
-					ShaderCodeArchive::ECookShaderLibrarySource::PreviousIncremental);
+				CodeArchive->FinishPopulate(ShaderCodeDir,
+					ShaderCodeArchive::ECookShaderLibrarySource::PreviousIncremental, &CodeStats);
 				if (bInitializeReferenceTracking)
 				{
 					CodeArchive->InitializeReferenceTracking();
@@ -1649,7 +1668,7 @@ public:
 						bool bOk = ReferencedAssetsArchive->SaveToDisk(ShaderCodeDir, MetaOutputDir,
 							ESaveToDiskTarget::Staging);
 
-						bool bShouldWriteInNativeFormat = bOk && bNativeFormat && ReferencedAssetsArchive->GetFormat()->SupportsShaderArchives();
+						bool bShouldWriteInNativeFormat = bOk && bNativeFormat && ReferencedAssetsArchive->SupportsShaderArchives();
 						if (bShouldWriteInNativeFormat)
 						{
 							bOk = ReferencedAssetsArchive->PackageNativeShaderLibrary(ShaderCodeDir) && bOk;
@@ -1729,7 +1748,7 @@ public:
 				bOk = PerChunkArchive->SaveToDisk(SandboxDestinationPath, SandboxMetadataPath,
 					ESaveToDiskTarget::Staging, &OutChunkFilenames) && bOk;
 
-				bool bShouldWriteInNativeFormat = bOk && bNativeFormat && PerChunkArchive->GetFormat()->SupportsShaderArchives();
+				bool bShouldWriteInNativeFormat = bOk && bNativeFormat && PerChunkArchive->SupportsShaderArchives();
 				if (bShouldWriteInNativeFormat)
 				{
 					bOk = PerChunkArchive->PackageNativeShaderLibrary(SandboxDestinationPath, &OutChunkFilenames) && bOk;
@@ -1786,22 +1805,6 @@ public:
 		}
 	}
 
-	bool PackageNativeShaderLibrary(const FString& ShaderCodeDir, const TArray<FName>& ShaderFormats)
-	{
-		bool bOK = true;
-		for (const FName ShaderFormatName : ShaderFormats)
-		{
-			EShaderPlatform SPlatform = ShaderFormatToLegacyShaderPlatform(ShaderFormatName);
-			FEditorShaderCodeArchive* CodeArchive = EditorShaderCodeArchive[SPlatform];
-
-			if (CodeArchive && CodeArchive->GetFormat()->SupportsShaderArchives())
-			{
-				bOK &= CodeArchive->PackageNativeShaderLibrary(ShaderCodeDir);
-			}
-		}
-		return bOK;
-	}
-
 	/**
 	 * Returns true if the specified shader type is referenced by any asset for staging,
 	 * and optionally reports referencers.
@@ -1819,7 +1822,7 @@ public:
 			ShaderArchive->ForEachAssetReferencingShaderType(
 				InShaderTypeHash,
 				[&bOutIsReferencedByStaging, OutAssetNames, MaxNumAssetNames]
-				(FName AssetName, const FSerializedShaderArchive::FShaderMapAssetAssociations::FAssociatedAssetData& AssetData) -> bool
+				(FName AssetName, const FShaderMapAssetAssociations::FAssociatedAssetData& AssetData) -> bool
 				{
 					if (AssetData.bReferencedByStaging)
 					{
@@ -1965,6 +1968,19 @@ public:
 	}
 #endif // WITH_EDITOR
 
+	void ResolvePackageShaderMapsDelegate(TArrayView<const FSHAHash> ShaderMapHashes, FCoreDelegates::FIoChunkIdResolvedFunc IoChunkIdResolvedFunc)
+	{
+		for (const FSHAHash& ShaderMapHash : ShaderMapHashes)
+		{
+			int32 ShaderMapIndex;
+			FShaderLibraryInstance* LibraryInstance = FindShaderLibraryForShaderMap(ShaderMapHash, ShaderMapIndex);
+			if (LibraryInstance)
+			{
+				LibraryInstance->ResolvePackageShaderMap(ShaderMapIndex, IoChunkIdResolvedFunc);
+			}
+		}
+	}
+
 	void PreloadPackageShaderMapsDelegate(TArrayView<const FSHAHash> ShaderMapHashes, FCoreDelegates::FAttachShaderReadRequestFunc AttachShaderReadRequestFunc)
 	{
 		LLM_SCOPE(ELLMTag::Shaders);
@@ -2485,7 +2501,7 @@ void FShaderLibraryCooker::DumpShaderTypeStats(const FString& DebugInfoDir, cons
 	}
 }
 
-bool FShaderLibraryCooker::CreatePatchLibrary(TArray<FString> const& OldMetaDataDirs, FString const& NewMetaDataDir, FString const& OutDir, bool bNativeFormat, bool bNeedsDeterministicOrder)
+bool FShaderLibraryCooker::CreatePatchLibrary(TArray<FString> const& OldMetaDataDirs, FString const& NewMetaDataDir, FString const& OutDir, bool bNativeFormat, bool /*bNeedsDeterministicOrder*/)
 {
 	TMap<FName, TSet<FString>> FormatLibraryMap;
 	TArray<FString> LibraryFiles;
@@ -2514,7 +2530,7 @@ bool FShaderLibraryCooker::CreatePatchLibrary(TArray<FString> const& OldMetaData
 	{
 		for (auto const& Library : Entry.Value)
 		{
-			bOK |= FEditorShaderCodeArchive::CreatePatchLibrary(Entry.Key, Library, OldMetaDataDirs, NewMetaDataDir, OutDir, bNativeFormat, bNeedsDeterministicOrder);
+			bOK |= FEditorShaderCodeArchive::CreatePatchLibrary(Entry.Key, Library, OldMetaDataDirs, NewMetaDataDir, OutDir, bNativeFormat);
 		}
 	}
 	return bOK;
@@ -2564,9 +2580,6 @@ bool FShaderLibraryCooker::MergeShaderCodeArchive(const TArray<FString>& CookedM
 
 	OutWrittenFiles.Empty();
 
-	bool bNeedsDeterministicOrder = false;
-	GConfig->GetBool(TEXT("/Script/UnrealEd.ProjectPackagingSettings"), TEXT("bDeterministicShaderCodeOrder"), bNeedsDeterministicOrder, GGameIni);
-
 	static const FRegexPattern ShaderArchivePattern(ShaderArchivePatternStr);
 	static const FRegexPattern StableInfoPattern(StableInfoPatternStr);
 
@@ -2586,7 +2599,7 @@ bool FShaderLibraryCooker::MergeShaderCodeArchive(const TArray<FString>& CookedM
 			if (ShaderCodeArchives.Contains(ByteCodeFile))
 			{
 				ShaderCodeArchives[ByteCodeFile].Key.AddShaderCodeLibraryFromDirectory(ShaderCodeDir,
-					ShaderCodeArchives[ByteCodeFile].Value, ShaderCodeArchive::ECookShaderLibrarySource::CurrentCook);
+					ShaderCodeArchive::ECookShaderLibrarySource::CurrentCook, &ShaderCodeArchives[ByteCodeFile].Value);
 			}
 			else
 			{
@@ -2594,11 +2607,11 @@ bool FShaderLibraryCooker::MergeShaderCodeArchive(const TArray<FString>& CookedM
 				if (ensureMsgf(FindShaderFormat.FindNext(), TEXT("Unable to parse out shader format from %s"), *ByteCodeFile))
 				{
 					const FName ShaderFormat(*FindShaderFormat.GetCaptureGroup(2));
-					FEditorShaderCodeArchive NewCodeArchive(ShaderFormat, bNeedsDeterministicOrder);
+					FEditorShaderCodeArchive NewCodeArchive(ShaderFormat);
 					FShaderCodeStats NewCodeStats;
 					NewCodeArchive.OpenLibrary(FindShaderFormat.GetCaptureGroup(1));
-					NewCodeArchive.AddShaderCodeLibraryFromDirectory(ShaderCodeDir, NewCodeStats,
-						ShaderCodeArchive::ECookShaderLibrarySource::CurrentCook);
+					NewCodeArchive.AddShaderCodeLibraryFromDirectory(ShaderCodeDir,
+						ShaderCodeArchive::ECookShaderLibrarySource::CurrentCook, &NewCodeStats);
 
 					ShaderCodeArchives.FindOrAdd(ByteCodeFile, { MoveTemp(NewCodeArchive), MoveTemp(NewCodeStats) });
 				}
diff --git a/Engine/Source/Runtime/RenderCore/Private/ShaderLibrary/ShaderCodeLibraryUtilities.cpp b/Engine/Source/Runtime/RenderCore/Private/ShaderLibrary/ShaderCodeLibraryUtilities.cpp
--- a/Engine/Source/Runtime/RenderCore/Private/ShaderLibrary/ShaderCodeLibraryUtilities.cpp
+++ b/Engine/Source/Runtime/RenderCore/Private/ShaderLibrary/ShaderCodeLibraryUtilities.cpp
@@ -19,9 +19,9 @@
 
 namespace UE::ShaderLibrary::Private
 {
+	using ESaveToDiskTarget = ShaderCodeArchive::ESaveToDiskTarget;
 
 	const uint32 GShaderCodeArchiveVersion = 2;
-	const uint32 GShaderPipelineArchiveVersion = 1;
 
 	const FString ShaderExtension = ".ushaderbytecode";
 	const FString CookCacheShaderExtension = ".ucacheshaderbytecode";
@@ -109,11 +109,6 @@ namespace UE::ShaderLibrary::Private
 		return BaseDir / FString::Printf(TEXT("ShaderStableInfo-%s-"), *LibraryName) + Platform.ToString() + StableExtension;
 	}
 
-	FString GetPipelinesArchiveFilename(const FString& BaseDir, const FString& LibraryName, FName Platform)
-	{
-		return BaseDir / FString::Printf(TEXT("ShaderArchive-%s-"), *LibraryName) + Platform.ToString() + PipelineExtension;
-	}
-
 	FString GetShaderCodeFilename(const FString& BaseDir, const FString& LibraryName, FName Platform, ESaveToDiskTarget Target)
 	{
 		switch (Target)
diff --git a/Engine/Source/Runtime/RenderCore/Private/ShaderLibrary/ShaderCodeLibraryUtilities.h b/Engine/Source/Runtime/RenderCore/Private/ShaderLibrary/ShaderCodeLibraryUtilities.h
--- a/Engine/Source/Runtime/RenderCore/Private/ShaderLibrary/ShaderCodeLibraryUtilities.h
+++ b/Engine/Source/Runtime/RenderCore/Private/ShaderLibrary/ShaderCodeLibraryUtilities.h
@@ -10,24 +10,18 @@
 #include "Containers/UnrealString.h"
 #include "Serialization/Archive.h"
 #include "UObject/NameTypes.h"
+#include "ShaderLibrary/ShaderCodeArchiveTypes.h"
 
 class ICookArtifactReader;
 
 namespace UE::ShaderLibrary::Private
 {
-	enum class ESaveToDiskTarget : uint8
-	{
-		Staging = 0,
-		CookCache,
-	};
-
 	// Compile time switch to enable the new feature of pruning stale shader data from incremental cooks.
 	// Remove this switch once the feature is stable enough not to need an off switch.
 	constexpr bool bIncrementalCookPruningEnabled = true;
 
 	// Shader library constants:
 	extern const uint32 GShaderCodeArchiveVersion;
-	extern const uint32 GShaderPipelineArchiveVersion;
 
 	extern const FString ShaderExtension;
 	extern const FString CookCacheShaderExtension;
@@ -55,12 +49,11 @@ namespace UE::ShaderLibrary::Private
 
 	// Shader library utility functions:
 	FString GetShaderCodeLibraryDebugInfoPath();
-	FString GetCodeArchiveFilename(const FString& BaseDir, const FString& LibraryName, FName Platform, ESaveToDiskTarget Target);
+	FString GetCodeArchiveFilename(const FString& BaseDir, const FString& LibraryName, FName Platform, ShaderCodeArchive::ESaveToDiskTarget Target);
 	FString GetStableInfoArchiveFilename(const FString& BaseDir, const FString& LibraryName, FName Platform);
-	FString GetPipelinesArchiveFilename(const FString& BaseDir, const FString& LibraryName, FName Platform);
-	FString GetShaderCodeFilename(const FString& BaseDir, const FString& LibraryName, FName Platform, ESaveToDiskTarget Target);
-	FString GetShaderAssetInfoFilename(const FString& BaseDir, const FString& LibraryName, FName Platform, ESaveToDiskTarget Target);
-	FString GetShaderTypeInfoFilename(const FString& BaseDir, const FString& LibraryName, FName Platform, ESaveToDiskTarget Target);
+	FString GetShaderCodeFilename(const FString& BaseDir, const FString& LibraryName, FName Platform, ShaderCodeArchive::ESaveToDiskTarget Target);
+	FString GetShaderAssetInfoFilename(const FString& BaseDir, const FString& LibraryName, FName Platform, ShaderCodeArchive::ESaveToDiskTarget Target);
+	FString GetShaderTypeInfoFilename(const FString& BaseDir, const FString& LibraryName, FName Platform, ShaderCodeArchive::ESaveToDiskTarget Target);
 	FString GetShaderDebugFolder(const FString& BaseDir, const FString& LibraryName, FName Platform);
 
 	/** Helper function shared between the cooker and runtime */
diff --git a/Engine/Source/Runtime/RenderCore/Private/ShaderLibrary/ShaderLibraryInstance.cpp b/Engine/Source/Runtime/RenderCore/Private/ShaderLibrary/ShaderLibraryInstance.cpp
--- a/Engine/Source/Runtime/RenderCore/Private/ShaderLibrary/ShaderLibraryInstance.cpp
+++ b/Engine/Source/Runtime/RenderCore/Private/ShaderLibrary/ShaderLibraryInstance.cpp
@@ -45,13 +45,11 @@ FShaderLibraryInstance* FShaderLibraryInstance::Create(EShaderPlatform InShaderP
 		const FName ShaderFormatName = LegacyShaderPlatformToShaderFormat(InShaderPlatform);
 		FString ShaderFormatAndPlatform = ShaderFormatName.ToString() + TEXT("-") + PlatformName.ToString();
 
-		const FString DestFilePath = GetCodeArchiveFilename(ShaderCodeDir, InLibraryName, FName(ShaderFormatAndPlatform), ESaveToDiskTarget::Staging);
+		const FString DestFilePath = GetCodeArchiveFilename(ShaderCodeDir, InLibraryName, FName(ShaderFormatAndPlatform), ShaderCodeArchive::ESaveToDiskTarget::Staging);
 		TUniquePtr<FArchive> Ar(CreateShaderFileReader(*DestFilePath));
 		if (Ar)
 		{
-			uint32 Version = 0;
-			*Ar << Version;
-			if (Version == GShaderCodeArchiveVersion)
+			if (FSerializedShaderArchive::SerializeHeaderVersion(*Ar))
 			{
 				Library = FShaderCodeArchive::Create(InShaderPlatform, *Ar, DestFilePath, ShaderCodeDir, InLibraryName);
 				if (Library)
@@ -324,6 +322,11 @@ void FShaderLibraryInstance::ReleasePreloadedShaderIfNecessary(int32 ShaderIndex
 	}
 }
 
+void FShaderLibraryInstance::ResolvePackageShaderMap(int32 ShaderMapIndex, FCoreDelegates::FIoChunkIdResolvedFunc IoChunkIdResolvedFunc) const
+{
+	Library->ResolveShaderMap(ShaderMapIndex, IoChunkIdResolvedFunc);
+}
+
 void FShaderLibraryInstance::PreloadPackageShaderMap(int32 ShaderMapIndex, FCoreDelegates::FAttachShaderReadRequestFunc AttachShaderReadRequestFunc)
 {
 	FRWScopeLock Locker(ResourceLock, SLT_Write);
diff --git a/Engine/Source/Runtime/RenderCore/Private/ShaderLibrary/ShaderLibraryInstance.h b/Engine/Source/Runtime/RenderCore/Private/ShaderLibrary/ShaderLibraryInstance.h
--- a/Engine/Source/Runtime/RenderCore/Private/ShaderLibrary/ShaderLibraryInstance.h
+++ b/Engine/Source/Runtime/RenderCore/Private/ShaderLibrary/ShaderLibraryInstance.h
@@ -51,6 +51,8 @@ public:
 
 	void ReleasePreloadedShaderIfNecessary(int32 ShaderIndex);
 
+	void ResolvePackageShaderMap(int32 ShaderMapIndex, FCoreDelegates::FIoChunkIdResolvedFunc IoChunkIdResolvedFunc) const;
+
 	void PreloadPackageShaderMap(int32 ShaderMapIndex, FCoreDelegates::FAttachShaderReadRequestFunc AttachShaderReadRequestFunc);
 
 	void ReleasePreloadedPackageShaderMap(int32 ShaderMapIndex);
diff --git a/Engine/Source/Runtime/RenderCore/Private/ShaderLibrary/ShaderMapAssetAssociation.cpp b/Engine/Source/Runtime/RenderCore/Private/ShaderLibrary/ShaderMapAssetAssociation.cpp
new file mode 100644
--- /dev/null
+++ b/Engine/Source/Runtime/RenderCore/Private/ShaderLibrary/ShaderMapAssetAssociation.cpp
@@ -0,0 +1,73 @@
+// Copyright Epic Games, Inc. All Rights Reserved.
+
+#include "ShaderMapAssetAssociation.h"
+
+#if WITH_EDITORONLY_DATA
+
+void FShaderMapAssetAssociations::ReserveAssets(int32 NumAssets)
+{
+	ShaderMapToAssets.Reserve(NumAssets);
+	AssetToShaderMaps.Reserve(NumAssets);
+}
+
+FShaderMapAssetAssociations::FAssociatedAssetData& FShaderMapAssetAssociations::FindOrAddAsset(FName Asset, const FSHAHash& ShaderMap)
+{
+	FAssociatedAssetData& Result = AssetToShaderMaps.FindOrAdd(Asset);
+	FShaderMapAssetPaths& Assets = ShaderMapToAssets.FindOrAdd(ShaderMap);
+	bool bAlreadyExists = false;
+	Assets.Add(Asset, &bAlreadyExists);
+	if (!bAlreadyExists)
+	{
+		Result.ShaderMaps.AddUnique(ShaderMap);
+	}
+	check(AssetToShaderMaps.IsEmpty() == ShaderMapToAssets.IsEmpty());
+	return Result;
+}
+
+void FShaderMapAssetAssociations::RemoveAsset(FName Asset)
+{
+	FAssociatedAssetData Removed;
+	if (AssetToShaderMaps.RemoveAndCopyValue(Asset, Removed))
+	{
+		for (const FSHAHash& ShaderMap : Removed.ShaderMaps)
+		{
+			FShaderMapAssetPaths* AssetsUsingShader = ShaderMapToAssets.Find(ShaderMap);
+			if (AssetsUsingShader)
+			{
+				AssetsUsingShader->Remove(Asset);
+				if (AssetsUsingShader->IsEmpty())
+				{
+					ShaderMapToAssets.Remove(ShaderMap);
+				}
+			}
+		}
+	}
+	check(AssetToShaderMaps.IsEmpty() == ShaderMapToAssets.IsEmpty());
+}
+
+void FShaderMapAssetAssociations::SortForSaving()
+{
+	ShaderMapToAssets.KeySort(
+		[](const FSHAHash& A, const FSHAHash& B)
+		{
+			return GetTypeHash(A) < GetTypeHash(B);
+		});
+	for (TPair<FSHAHash, FShaderMapAssetPaths>& Pair : ShaderMapToAssets)
+	{
+		Pair.Value.Sort(FNameLexicalLess());
+	}
+}
+
+void FShaderMapAssetAssociations::Append(const FShaderMapAssetAssociations& InOtherAssociation)
+{
+	//@todo-lh - This can be optimized slightly, but let's keep it straightforward for now
+	for (const TPair<FName, FAssociatedAssetData>& AssetData : InOtherAssociation.ViewAssets())
+	{
+		for (const FSHAHash& ShaderMapHash : AssetData.Value.ShaderMaps)
+		{
+			this->FindOrAddAsset(AssetData.Key, ShaderMapHash);
+		}
+	}
+}
+
+#endif // WITH_EDITORONLY_DATA
diff --git a/Engine/Source/Runtime/RenderCore/Private/ShaderPipelineCache.cpp b/Engine/Source/Runtime/RenderCore/Private/ShaderPipelineCache.cpp
--- a/Engine/Source/Runtime/RenderCore/Private/ShaderPipelineCache.cpp
+++ b/Engine/Source/Runtime/RenderCore/Private/ShaderPipelineCache.cpp
@@ -24,6 +24,7 @@
 #include "Misc/ScopeLock.h"
 #include <Algo/RemoveIf.h>
 #include "Modules/BuildVersion.h"
+#include "ShaderLibrary/ShaderCodeLibraryUtilities.h"
 
 DECLARE_DWORD_ACCUMULATOR_STAT(TEXT("Outstanding Tasks"), STAT_ShaderPipelineTaskCount, STATGROUP_PipelineStateCache );
 DECLARE_DWORD_ACCUMULATOR_STAT(TEXT("Waiting Tasks"), STAT_ShaderPipelineWaitingTaskCount, STATGROUP_PipelineStateCache );
@@ -339,15 +340,6 @@ namespace UE
 	namespace ShaderPipeline
 	{
 
-		FString LibNameToPSOFCName(FString const& Name, int32 ComponentID)
-		{
-			if (ComponentID == -1)
-			{
-				return Name;
-			}
-			return FString::Printf(TEXT("%s_Chunk%d"), *Name, ComponentID);
-		}
-
 		bool ShouldLoadUserCache()
 		{
 			return GetShaderPipelineCacheSaveBoundPSOLog() || GetPSOFileCacheSaveUserCache();
@@ -959,7 +951,7 @@ void FShaderPipelineCache::ShaderLibraryStateChanged(ELibraryState State, EShade
 {
 	if (ShaderPipelineCache)
 	{
-		const FString PSOCacheName = UE::ShaderPipeline::LibNameToPSOFCName(PSOCacheBaseName, ComponentID);
+		const FString PSOCacheName = UE::ShaderLibrary::Private::GetShaderLibraryNameForChunk(PSOCacheBaseName, ComponentID);
 		// Key == Name for bundles PSOFCs
 		const FString& PSOCacheKey = PSOCacheName;
 
diff --git a/Engine/Source/Runtime/RenderCore/Public/Shader.h b/Engine/Source/Runtime/RenderCore/Public/Shader.h
--- a/Engine/Source/Runtime/RenderCore/Public/Shader.h
+++ b/Engine/Source/Runtime/RenderCore/Public/Shader.h
@@ -43,6 +43,7 @@
 #include "ShaderParameters.h"
 #include "ShaderPermutation.h"
 #include "ShaderPermutationUtils.h"
+#include "ShaderMapAssetAssociation.h" // For FShaderMapAssetPaths
 #include "Templates/RefCounting.h"
 #include "Templates/TypeHash.h"
 #include "Templates/UniquePtr.h"
@@ -85,7 +86,6 @@ struct FShaderCompiledShaderInitializerType;
 struct FShaderCompileJobKey;
 struct FShaderCompilerOutput;
 struct FShaderSerializeContext;
-using FShaderMapAssetPaths = TSet<FName>; // Copied from ShaderCodeLibrary.h
 
 /** Define a shader permutation uniquely according to its type, and permutation id.*/
 template<typename MetaShaderType>
diff --git a/Engine/Source/Runtime/RenderCore/Public/ShaderCodeArchive.h b/Engine/Source/Runtime/RenderCore/Public/ShaderCodeArchive.h
--- a/Engine/Source/Runtime/RenderCore/Public/ShaderCodeArchive.h
+++ b/Engine/Source/Runtime/RenderCore/Public/ShaderCodeArchive.h
@@ -22,15 +22,12 @@
 #include "Serialization/Archive.h"
 #include "Shader.h"
 #include "ShaderCodeLibrary.h"
+#include "ShaderMapAssetAssociation.h"
 #include "Templates/RefCounting.h"
 #include "UObject/NameTypes.h"
 
 #include <type_traits>
 
-#if WITH_EDITOR
-#include "Templates/Function.h"
-#endif
-
 #if WITH_EDITOR
 class FCbFieldView;
 class FCbWriter;
@@ -147,36 +144,68 @@ struct FShaderCodeEntry
 static_assert(std::is_standard_layout_v<FShaderCodeEntry>);
 static_assert(std::is_trivially_copyable_v<FShaderCodeEntry>);
 
-#if WITH_EDITOR
-namespace ShaderCodeArchive
-{
-
 /**
- * Descriptor when loading a ShaderLibrary during cooking for which operation is loading the ShaderLibrary.
- * When the operation is loading data from previous incremental cooks, we have more strict validation requirements
- * and we ignore data for stale Assets that have been recooked.
- * Smaller values are newer sources and override larger values which come from older (more stale) sources.
+ * Stores information about a shader library name.
+ * This splits the filename into names for chunk, shader format, and platform format.
+ * For example, "Global-SF_VULKAN_SM6-VULKAN_SM6" is split into chunk "Global", shader format "SF_VULKAN_SM6", and platform format "VULKAN_SM6".
+ * Note that some platforms only have a single string for the format-platform combination, in which case ShaderFormatName can be empty.
  */
-enum class ECookShaderLibrarySource : uint8
+struct FShaderLibraryNameInfo
 {
-	/** The data is provided by compilation that executed during the current cook; this is the most up-to-date. */
-	CurrentCook = 0,
+	// Shader library chunk name, e.g. "Global", "Lyra_Chunk0", "MyPlugin" etc.
+	FString ChunkName;
 
-	/** The current cook is a DLC cook and the data being provided came from the basegame cook. */
-	PatchBase,
+	// Shader format nane, e.g, "SF_VULKAN_SM6".
+	FName ShaderFormatName;
 
-	/** The data was written to the ShaderLibrary artifacts in a previous incremental cook. */
-	PreviousIncremental,
-};
+	// Shader platform format name, e.g. "VULKAN_SM6".
+	FName ShaderPlatformName;
+
+	/** Initializes all name sections from the specified filename (usually a *.ushaderbytecode file). */
+	RENDERCORE_API bool ParseFromFilename(const TCHAR* Filename);
 
-} // namespace ShaderCodeArchive
-#endif // WITH_EDITOR
+	/** Returns the shader asset info filename (*.assetinfo.json) that belongs to the respective *.ushaderbytecode file. */
+	RENDERCORE_API FString GetAssetInfoFilename() const;
 
-// Portion of shader code archive that's serialize to disk
+	/** Returns the shader type info filename (*.stinfo) that belongs to the respective *.ushaderbytecode file. */
+	RENDERCORE_API FString GetTypeInfoFilename() const;
+
+	/** Returns the name for this shader library, which is a combination of chunk and format names, e.g. "Global-SF_VULKAN_SM6". Used for debugging. */
+	FORCEINLINE FString GetLibraryName() const
+	{
+		return FString::Printf(TEXT("%s-%s"), *ChunkName, *ShaderFormatName.ToString());
+	}
+
+	/** Returns the combined format and platform names. If shader format name is empty, only the platform name is returned. Used for FEditorShaderCodeArchive::FormatAndPlatformName. */
+	FORCEINLINE FString GetFormatAndPlatformName() const
+	{
+		return ShaderFormatName.IsNone()
+			? ShaderPlatformName.ToString()
+			: FString::Printf(TEXT("%s-%s"), *ShaderFormatName.ToString(), *ShaderPlatformName.ToString());
+	}
+
+	/** Returns the full name with chunk, shader format, and platform format names. */
+	FORCEINLINE FString GetFullName() const
+	{
+		return FString::Printf(TEXT("%s-%s"), *ChunkName, *GetFormatAndPlatformName());
+	}
+
+	/** Helper function to parse the ID and optional suffix from a pakchunk name, e.g. "pakchunk100iad-Windows" -> 100 and "iad" */
+	RENDERCORE_API static bool ParsePakChunkId(const FStringView& PakChunkName, int32& OutChunkId, FString& OutChunkIdSuffix);
+};
+
+/**
+ * Header portion of shader code archive that's serialized to disk.
+ * 
+ * The shader codes themselves are stored separately, either with individual buffers (one for each shader group),
+ * composite buffers, or a single consecutive buffer (those can become quite large).
+ * 
+ * @see FEditorShaderCodeArchive, FSerializedShaderArchiveBuffer
+ */
 class FSerializedShaderArchive
 {
 public:
-#if WITH_EDITOR
+#if WITH_EDITORONLY_DATA
 	using ECookShaderLibrarySource = ShaderCodeArchive::ECookShaderLibrarySource;
 #endif
 #if USE_MMAPPED_SHADERARCHIVE
@@ -218,68 +247,27 @@ public:
 	FHashTable ShaderMapHashTable;
 	FHashTable ShaderHashTable;
 
-#if WITH_EDITOR
-	/**
-	 * Bi-directional mapping from shadermap hashes to an array of asset names.
-	 * Assets are identified by LongPackageName.
-	 * ShaderMaps are identified by their SHAHash.
-	 * This structure is used to identify the ShaderMaps required in the projection of this library for a pakfile chunk
-	 * based on assets in the chunk, and to remove stale ShaderMaps and their Shaders from this library when the assets
-	 * are recooked or are pruned from the incremental cook oplog.
-	 */
-	class FShaderMapAssetAssociations
-	{
-	public:
-		struct FAssociatedAssetData
-		{
-			TArray<FSHAHash> ShaderMaps;
-			ECookShaderLibrarySource LatestSource = ECookShaderLibrarySource::PreviousIncremental;
-			// These reference flags are transient. By default they are initialized to true, for shaderlibraries that do not use
-			// reference tracking such as the global shader library. When tracking references, we initialize them to false.
-			// These flags are read/written only during end of cook.
-			/** Set to true if the Asset is needed either for incremental cook or for staging, and should therefore not be pruned. */
-			bool bReferencedByOplog = true;
-			/** Set to true if the Asset is needed for runtime and staging, and should therefore be added to the staging artifacts. */
-			bool bReferencedByStaging = true;
-
-			/** this->LatestSource = NewestOf(LatestSource, InCookSource). */
-			void MergeSource(ECookShaderLibrarySource InCookSource);
-		};
-		bool IsEmpty() const;
-		SIZE_T GetAllocatedSize() const;
-		TMap<FName, FAssociatedAssetData>& ViewAssets();
-		const TMap<FName, FAssociatedAssetData>& ViewAssets() const;
-		const TMap<FSHAHash, FShaderMapAssetPaths>& ViewShaderMaps() const;
-
-		FAssociatedAssetData* FindAsset(FName AssetName);
-		const FAssociatedAssetData* FindAsset(FName AssetName) const;
-		FShaderMapAssetPaths* FindShaderMap(const FSHAHash& ShaderMap);
-		const FShaderMapAssetPaths* FindShaderMap(const FSHAHash& ShaderMap) const;
-
-		void Empty();
-		RENDERCORE_API void ReserveAssets(int32 NumAssets);
-		RENDERCORE_API FAssociatedAssetData& FindOrAddAsset(FName Asset, const FSHAHash& ShaderMap);
-		RENDERCORE_API void RemoveAsset(FName Asset);
-		RENDERCORE_API void SortForSaving();
-
-	private:
-		TMap<FSHAHash, FShaderMapAssetPaths> ShaderMapToAssets;
-		TMap<FName, FAssociatedAssetData> AssetToShaderMaps;
-	};
+#if WITH_EDITORONLY_DATA
 
 private:
 	FShaderMapAssetAssociations ShaderMapAssetAssociations;
 
 public:
-	FShaderMapAssetAssociations& GetShaderMapAssetAssociations();
-	const FShaderMapAssetAssociations& GetShaderMapAssetAssociations() const;
+	FORCEINLINE FShaderMapAssetAssociations& GetShaderMapAssetAssociations()
+	{
+		return ShaderMapAssetAssociations;
+	}
+	FORCEINLINE const FShaderMapAssetAssociations& GetShaderMapAssetAssociations() const
+	{
+		return ShaderMapAssetAssociations;
+	}
 
 	UE_DEPRECATED(5.8, "No longer referenced. Use GetShaderMapAssetAssociations() instead.")
 	TMap<FSHAHash, FShaderMapAssetPaths> ShaderCodeToAssets;
 
 	struct FShaderTypeHashes
 	{
-		/** Most of the time, this will only contain one element. But some shaders are deduplicated to two or more shader types. */
+		/** Most of the time, this will only contain one element. But some shaders are shared across two or more shader types. */
 		TArray<uint64> Data;
 
 		friend FArchive& operator<<(FArchive& Ar, FShaderTypeHashes& Ref)
@@ -333,7 +321,7 @@ public:
 		/** Compressed (individually) size of all shaders of a given frequency. */
 		uint64 CompressedSizePerFrequency[SF_NumFrequencies];
 	};
-#endif // WITH_EDITOR
+#endif // WITH_EDITORONLY_DATA
 
 	FSerializedShaderArchive()
 	{
@@ -360,10 +348,10 @@ public:
 			PreloadEntries.GetAllocatedSize() +
 			ShaderIndices.GetAllocatedSize()
 #endif
-#if WITH_EDITOR
+#if WITH_EDITORONLY_DATA
 			+ ShaderMapAssetAssociations.GetAllocatedSize()
 			+ ShaderTypes.GetAllocatedSize()
-#endif // WITH_EDITOR
+#endif // WITH_EDITORONLY_DATA
 			;
 	}
 
@@ -375,7 +363,7 @@ public:
 		ShaderEntries.Empty();
 #endif
 		ShaderHashTable.Clear();
-#if WITH_EDITOR
+#if WITH_EDITORONLY_DATA
 		ShaderTypes.Empty();
 #endif
 	}
@@ -389,7 +377,7 @@ public:
 		ShaderIndices.Empty();
 #endif
 		ShaderMapHashTable.Clear();
-#if WITH_EDITOR
+#if WITH_EDITORONLY_DATA
 		ShaderMapAssetAssociations.Empty();
 #endif
 	}
@@ -407,7 +395,7 @@ public:
 	bool IsEmpty() const
 	{
 		return ShaderMapEntries.IsEmpty() && ShaderEntries.IsEmpty() && PreloadEntries.IsEmpty()
-#if WITH_EDITOR
+#if WITH_EDITORONLY_DATA
 			&& ShaderMapAssetAssociations.IsEmpty()
 			&& ShaderTypes.IsEmpty()
 #endif
@@ -423,7 +411,7 @@ public:
 	RENDERCORE_API bool FindOrAddShader(const FSHAHash& Hash, int32& OutIndex);
 	UE_DEPRECATED(5.8, "Call FindOrAddShaderMapNonEditor or FindOrAddShaderMapEditor.")
 	RENDERCORE_API bool FindOrAddShaderMap(const FSHAHash& Hash, int32& OutIndex, const FShaderMapAssetPaths* AssociatedAssets);
-#if WITH_EDITOR
+#if WITH_EDITORONLY_DATA
 	RENDERCORE_API bool FindOrAddShaderMapEditor(const FSHAHash& Hash, int32& OutIndex,
 		const FShaderMapAssetPaths* AssociatedAssets, ECookShaderLibrarySource CookSource);
 #else
@@ -434,17 +422,22 @@ public:
 #endif // !USE_MMAPPED_SHADERARCHIVE
 
 	RENDERCORE_API void DecompressShader(int32 Index, const TArray<FSharedBuffer>& ShaderCode, TArray<uint8>& OutDecompressedShader) const;
+	RENDERCORE_API static void DecompressShaderEntryAndAppend(const FShaderCodeEntry& Entry, const uint8* CompressedShaderCode, TArray<uint8>& OutDecompressedShader);
 
+	RENDERCORE_API static FString MakeDebugDirectory(const FString LibraryName = TEXT(""), const FName& FormatAndPlatformName = NAME_None);
+	
+	RENDERCORE_API static bool SerializeHeaderVersion(FArchive& Ar);
 	RENDERCORE_API void Serialize(FArchive& Ar);
-#if WITH_EDITOR
+#if WITH_EDITORONLY_DATA
 	RENDERCORE_API void SaveAssetInfo(FArchive& Ar);
 	RENDERCORE_API bool LoadAssetInfo(const FString& Filename, ECookShaderLibrarySource CookSource = ECookShaderLibrarySource::CurrentCook);
 	RENDERCORE_API bool LoadAssetInfo(FArchive* Ar, ECookShaderLibrarySource CookSource = ECookShaderLibrarySource::CurrentCook);
 	UE_DEPRECATED(5.8, "Use CreateFromFilteredAssets instead")
 	RENDERCORE_API void CreateAsChunkFrom(const FSerializedShaderArchive& Parent, const TSet<FName>& PackagesInChunk, TArray<int32>& OutShaderCodeEntriesNeeded);
 
-	using FAssetFilterFunctionRef = TFunctionRef<bool(FName AssetName, const FShaderMapAssetAssociations::FAssociatedAssetData& AssetData)>;
-	RENDERCORE_API void CreateFromFilteredAssets(const FSerializedShaderArchive& Parent, TArray<int32>& OutShaderCodeEntriesNeeded, FAssetFilterFunctionRef ShouldKeepAsset);
+	RENDERCORE_API void CreateFromFilteredAssets(const FSerializedShaderArchive& Parent, TArray<int32>& OutShaderCodeEntriesNeeded,
+		FShaderMapAssetAssociations::FAssetFilterFunctionRef ShouldKeepAsset,
+		FShaderMapAssetAssociations::FShaderMapFilterFunction ShouldExcludeShaderMap = nullptr);
 	RENDERCORE_API void CollectStatsAndDebugInfo(FDebugStats& OutDebugStats, FExtendedDebugStats* OutExtendedDebugStats);
 	RENDERCORE_API void DumpContentsInPlaintext(FString& OutText) const;
 	RENDERCORE_API friend FCbWriter& operator<<(FCbWriter& Writer, const FSerializedShaderArchive& Archive);
@@ -538,6 +531,10 @@ struct FShaderUsageVisualizer
 #endif // UE_SCA_VISUALIZE_SHADER_USAGE
 };
 
+/**
+ * Implements FRHIShaderLibrary and reads shader code from a single file cache handle.
+ * @todo-lh - Consider renaming to FFileCacheRHIShaderLibrary
+ */
 class FShaderCodeArchive : public FRHIShaderLibrary
 {
 public:
@@ -564,12 +561,12 @@ public:
 		return SerializedShaders.GetShaderIndices()[ShaderMapEntry.ShaderIndicesOffset + i];
 	}
 
-	virtual void GetAllShaderIndices(int32 ShaderMapIndex, TArray<int32>& ShaderIndices)
+	virtual void GetAllShaderIndices(int32 ShaderMapIndex, TArray<int32>& OutShaderIndices) const override
 	{
 		const FShaderMapEntry& ShaderMapEntry = SerializedShaders.GetShaderMapEntries()[ShaderMapIndex];
 		for (uint32 i = 0u; i < ShaderMapEntry.NumShaders; ++i)
 		{
-			ShaderIndices.AddUnique(ShaderIndices[ShaderMapEntry.ShaderIndicesOffset + i]);
+			OutShaderIndices.AddUnique(SerializedShaders.GetShaderIndices()[ShaderMapEntry.ShaderIndicesOffset + i]);
 		}
 	}
 
@@ -594,9 +591,7 @@ public:
 	};
 
 	virtual bool PreloadShader(int32 ShaderIndex, FGraphEventArray& OutCompletionEvents) override;
-
 	virtual bool PreloadShaderMap(int32 ShaderMapIndex, FGraphEventArray& OutCompletionEvents) override;
-
 	virtual void ReleasePreloadedShader(int32 ShaderIndex) override;
 
 	virtual TRefCountPtr<FRHIShader> CreateShader(int32 Index, bool bRequired = true) override;
@@ -715,7 +710,7 @@ static_assert(sizeof(FIoStoreShaderCodeEntry) == sizeof(uint64), "To reduce memo
 /** Descriptor of a group of shaders compressed together. This groups already deduplicated, and possibly unrelated, shaders, so this is a distinct concept from a shader map. */
 struct FIoStoreShaderGroupEntry
 {
-	/** Offset to an the first shader index referenced by this group in the array of shader indices. This extra level of indirection allows arbitrary grouping. */
+	/** Offset to the first shader index referenced by this group in the array of shader indices. This extra level of indirection allows arbitrary grouping. */
 	uint32 ShaderIndicesOffset = 0u;
 	/** Number of shaders in this group. */
 	uint32 NumShaders = 0u;
@@ -807,6 +802,11 @@ public:
 	}
 };
 
+/**
+ * Implements FRHIShaderLibrary and reads shader code through IoStore interface (either a ZenServer or CDN).
+ * This class stores its shader code via IoBuffers, one per group whereas each shadermap references one or more groups and these groups can be shared between shadermaps.
+ * @todo-lh - Consider renaming to FIoStoreRHIShaderLibrary
+ */
 class FIoStoreShaderCodeArchive : public FRHIShaderLibrary
 {
 public:
@@ -840,12 +840,12 @@ public:
 		return Header.ShaderIndices[ShaderMapEntry.ShaderIndicesOffset + i];
 	}
 
-	virtual void GetAllShaderIndices(int32 ShaderMapIndex, TArray<int32>& ShaderIndices)
+	virtual void GetAllShaderIndices(int32 ShaderMapIndex, TArray<int32>& OutShaderIndices) const override
 	{
 		const FIoStoreShaderMapEntry& ShaderMapEntry = Header.ShaderMapEntries[ShaderMapIndex];
 		for (uint32 i = 0u; i < ShaderMapEntry.NumShaders; ++i)
 		{
-			ShaderIndices.AddUnique(Header.ShaderIndices[ShaderMapEntry.ShaderIndicesOffset + i]);
+			OutShaderIndices.AddUnique(Header.ShaderIndices[ShaderMapEntry.ShaderIndicesOffset + i]);
 		}
 	}
 
@@ -872,6 +872,7 @@ public:
 		return Header.ShaderEntries[ShaderIndex].ShaderGroupIndex;
 	}
 	
+	virtual bool ResolveShaderMap(int32 ShaderMapIndex, FCoreDelegates::FIoChunkIdResolvedFunc IoChunkIdResolvedFunc) const override;
 	virtual bool PreloadShaderMap(int32 ShaderMapIndex, FGraphEventArray& OutCompletionEvents) override;
 	virtual bool PreloadShaderMap(int32 ShaderMapIndex, FCoreDelegates::FAttachShaderReadRequestFunc AttachShaderReadRequestFunc) override;
 	virtual void ReleasePreloadedShader(int32 ShaderIndex) override;
@@ -950,7 +951,7 @@ private:
 	}
 
 	/** Returns true if the group contains only RTX shaders. We can avoid preloading it when running with RTX off. */
-	bool GroupOnlyContainsRaytracingShaders(int32 ShaderGroupIndex);
+	bool GroupOnlyContainsRaytracingShaders(int32 ShaderGroupIndex) const;
 
 	/** Archive header with all the metadata */
 	FIoStoreShaderCodeArchiveHeader Header;
@@ -981,84 +982,3 @@ inline FSerializedShaderArchive& FSerializedShaderArchive::operator=(FSerialized
 inline FSerializedShaderArchive& FSerializedShaderArchive::operator=(const FSerializedShaderArchive&) = default;
 PRAGMA_ENABLE_DEPRECATION_WARNINGS
 
-#if WITH_EDITOR
-
-inline FSerializedShaderArchive::FShaderMapAssetAssociations&
-FSerializedShaderArchive::GetShaderMapAssetAssociations()
-{
-	return ShaderMapAssetAssociations;
-}
-
-inline const FSerializedShaderArchive::FShaderMapAssetAssociations&
-FSerializedShaderArchive::GetShaderMapAssetAssociations() const
-{
-	return ShaderMapAssetAssociations;
-}
-
-inline bool FSerializedShaderArchive::FShaderMapAssetAssociations::IsEmpty() const
-{
-	return ShaderMapToAssets.IsEmpty();
-}
-
-inline SIZE_T FSerializedShaderArchive::FShaderMapAssetAssociations::GetAllocatedSize() const
-{
-	return ShaderMapToAssets.GetAllocatedSize() + AssetToShaderMaps.GetAllocatedSize();
-}
-
-inline TMap<FName, FSerializedShaderArchive::FShaderMapAssetAssociations::FAssociatedAssetData>&
-FSerializedShaderArchive::FShaderMapAssetAssociations::ViewAssets()
-{
-	return AssetToShaderMaps;
-}
-
-inline const TMap<FName, FSerializedShaderArchive::FShaderMapAssetAssociations::FAssociatedAssetData>&
-FSerializedShaderArchive::FShaderMapAssetAssociations::ViewAssets() const
-{
-	return AssetToShaderMaps;
-}
-
-inline const TMap<FSHAHash, FShaderMapAssetPaths>&
-FSerializedShaderArchive::FShaderMapAssetAssociations::ViewShaderMaps() const
-{
-	return ShaderMapToAssets;
-}
-
-inline FSerializedShaderArchive::FShaderMapAssetAssociations::FAssociatedAssetData*
-FSerializedShaderArchive::FShaderMapAssetAssociations::FindAsset(FName AssetName)
-{
-	return AssetToShaderMaps.Find(AssetName);
-}
-
-inline const FSerializedShaderArchive::FShaderMapAssetAssociations::FAssociatedAssetData*
-FSerializedShaderArchive::FShaderMapAssetAssociations::FindAsset(FName AssetName) const
-{
-	return AssetToShaderMaps.Find(AssetName);
-
-}
-
-inline FShaderMapAssetPaths*
-FSerializedShaderArchive::FShaderMapAssetAssociations::FindShaderMap(const FSHAHash& ShaderMap)
-{
-	return ShaderMapToAssets.Find(ShaderMap);
-}
-
-inline const FShaderMapAssetPaths*
-FSerializedShaderArchive::FShaderMapAssetAssociations::FindShaderMap(const FSHAHash& ShaderMap) const
-{
-	return ShaderMapToAssets.Find(ShaderMap);
-}
-
-inline void FSerializedShaderArchive::FShaderMapAssetAssociations::Empty()
-{
-	ShaderMapToAssets.Empty();
-	AssetToShaderMaps.Empty();
-}
-
-inline void FSerializedShaderArchive::FShaderMapAssetAssociations::FAssociatedAssetData::MergeSource(
-	ShaderCodeArchive::ECookShaderLibrarySource InCookSource)
-{
-	using namespace ShaderCodeArchive;
-	LatestSource = (ECookShaderLibrarySource)FMath::Min((uint8)LatestSource, (uint8)InCookSource);
-}
-
-#endif // WITH_EDITOR
diff --git a/Engine/Source/Runtime/RenderCore/Public/ShaderCodeLibrary.h b/Engine/Source/Runtime/RenderCore/Public/ShaderCodeLibrary.h
--- a/Engine/Source/Runtime/RenderCore/Public/ShaderCodeLibrary.h
+++ b/Engine/Source/Runtime/RenderCore/Public/ShaderCodeLibrary.h
@@ -23,6 +23,7 @@
 #include "RHI.h"
 #include "RHIDefinitions.h"
 #include "Serialization/Archive.h"
+#include "ShaderMapAssetAssociation.h"
 #include "Templates/Function.h"
 #include "Templates/RefCounting.h"
 #include "UObject/NameTypes.h"
@@ -33,7 +34,6 @@ class FShaderMapResource;
 class FShaderMapResourceCode;
 class FShaderPipeline;
 
-using FShaderMapAssetPaths = TSet<FName>;
 class FIoBuffer;
 class FIoChunkId;
 class IPlugin;
diff --git a/Engine/Source/Runtime/RenderCore/Public/ShaderMapAssetAssociation.h b/Engine/Source/Runtime/RenderCore/Public/ShaderMapAssetAssociation.h
new file mode 100644
--- /dev/null
+++ b/Engine/Source/Runtime/RenderCore/Public/ShaderMapAssetAssociation.h
@@ -0,0 +1,141 @@
+// Copyright Epic Games, Inc. All Rights Reserved.
+
+#pragma once
+
+#include "Containers/Array.h"
+#include "Containers/Map.h"
+#include "Containers/Set.h"
+#include "Misc/SecureHash.h"
+#include "Misc/CoreMiscDefines.h"
+#include "UObject/NameTypes.h"
+#include "Templates/Function.h"
+
+using FShaderMapAssetPaths = TSet<FName>;
+
+#if WITH_EDITORONLY_DATA
+
+namespace ShaderCodeArchive
+{
+
+	/**
+	 * Descriptor when loading a ShaderLibrary during cooking for which operation is loading the ShaderLibrary.
+	 * When the operation is loading data from previous incremental cooks, we have more strict validation requirements
+	 * and we ignore data for stale Assets that have been recooked.
+	 * Smaller values are newer sources and override larger values which come from older (more stale) sources.
+	 */
+	enum class ECookShaderLibrarySource : uint8
+	{
+		/** The data is provided by compilation that executed during the current cook; this is the most up-to-date. */
+		CurrentCook = 0,
+
+		/** The current cook is a DLC cook and the data being provided came from the basegame cook. */
+		PatchBase,
+
+		/** The data was written to the ShaderLibrary artifacts in a previous incremental cook. */
+		PreviousIncremental,
+	};
+
+} // namespace ShaderCodeArchive
+
+/**
+ * Bi-directional mapping from shadermap hashes to an array of asset names.
+ * Assets are identified by LongPackageName.
+ * ShaderMaps are identified by their SHAHash.
+ * This structure is used to identify the ShaderMaps required in the projection of this library for a pakfile chunk
+ * based on assets in the chunk, and to remove stale ShaderMaps and their Shaders from this library when the assets
+ * are recooked or are pruned from the incremental cook oplog.
+ */
+class FShaderMapAssetAssociations
+{
+public:
+	using ECookShaderLibrarySource = ShaderCodeArchive::ECookShaderLibrarySource;
+
+	struct FAssociatedAssetData
+	{
+		TArray<FSHAHash> ShaderMaps;
+		ECookShaderLibrarySource LatestSource = ECookShaderLibrarySource::PreviousIncremental;
+		// These reference flags are transient. By default they are initialized to true, for shaderlibraries that do not use
+		// reference tracking such as the global shader library. When tracking references, we initialize them to false.
+		// These flags are read/written only during end of cook.
+		/** Set to true if the Asset is needed either for incremental cook or for staging, and should therefore not be pruned. */
+		bool bReferencedByOplog = true;
+		/** Set to true if the Asset is needed for runtime and staging, and should therefore be added to the staging artifacts. */
+		bool bReferencedByStaging = true;
+
+		/** this->LatestSource = NewestOf(LatestSource, InCookSource). */
+		FORCEINLINE void MergeSource(ECookShaderLibrarySource InCookSource)
+		{
+			LatestSource = (ECookShaderLibrarySource)FMath::Min((uint8)LatestSource, (uint8)InCookSource);
+		}
+	};
+
+	/** Function signature to filter asset associated with its metadata. */
+	using FAssetFilterFunctionRef = TFunctionRef<bool(FName AssetName, const FAssociatedAssetData& AssetData)>;
+
+	/** Function signature to filter shadermaps. */
+	using FShaderMapFilterFunction = TFunction<bool(const FSHAHash& ShaderMapHash)>;
+
+	inline bool IsEmpty() const
+	{
+		return ShaderMapToAssets.IsEmpty();
+	}
+
+	inline SIZE_T GetAllocatedSize() const
+	{
+		return ShaderMapToAssets.GetAllocatedSize() + AssetToShaderMaps.GetAllocatedSize();
+	}
+
+	FORCEINLINE TMap<FName, FAssociatedAssetData>& ViewAssets()
+	{
+		return AssetToShaderMaps;
+	}
+
+	FORCEINLINE const TMap<FName, FAssociatedAssetData>& ViewAssets() const
+	{
+		return AssetToShaderMaps;
+	}
+
+	FORCEINLINE const TMap<FSHAHash, FShaderMapAssetPaths>& ViewShaderMaps() const
+	{
+		return ShaderMapToAssets;
+	}
+
+	FORCEINLINE FAssociatedAssetData* FindAsset(FName AssetName)
+	{
+		return AssetToShaderMaps.Find(AssetName);
+	}
+
+	FORCEINLINE const FAssociatedAssetData* FindAsset(FName AssetName) const
+	{
+		return AssetToShaderMaps.Find(AssetName);
+
+	}
+
+	FORCEINLINE FShaderMapAssetPaths* FindShaderMap(const FSHAHash& ShaderMap)
+	{
+		return ShaderMapToAssets.Find(ShaderMap);
+	}
+
+	FORCEINLINE const FShaderMapAssetPaths* FindShaderMap(const FSHAHash& ShaderMap) const
+	{
+		return ShaderMapToAssets.Find(ShaderMap);
+	}
+
+	FORCEINLINE void Empty()
+	{
+		ShaderMapToAssets.Empty();
+		AssetToShaderMaps.Empty();
+	}
+
+	RENDERCORE_API void ReserveAssets(int32 NumAssets);
+	RENDERCORE_API FAssociatedAssetData& FindOrAddAsset(FName Asset, const FSHAHash& ShaderMap);
+	RENDERCORE_API void RemoveAsset(FName Asset);
+	RENDERCORE_API void SortForSaving();
+	RENDERCORE_API void Append(const FShaderMapAssetAssociations& InOtherAssociation);
+
+private:
+	TMap<FSHAHash, FShaderMapAssetPaths> ShaderMapToAssets;
+	TMap<FName, FAssociatedAssetData> AssetToShaderMaps;
+};
+
+#endif // WITH_EDITORONLY_DATA
