[Shader] Use TUniquePtr for re-assigned shader chunks to fix potential memory leaks.

Also don't reset cooked package store pointer for all remaining serialized shaders, if there was one invalid ChunkId.

#rnx

#Synced-CL 50342528

diff --git a/Engine/Source/Developer/IoStoreUtilities/Private/IoStoreShaderLibraryProcessor.cpp b/Engine/Source/Developer/IoStoreUtilities/Private/IoStoreShaderLibraryProcessor.cpp
--- a/Engine/Source/Developer/IoStoreUtilities/Private/IoStoreShaderLibraryProcessor.cpp
+++ b/Engine/Source/Developer/IoStoreUtilities/Private/IoStoreShaderLibraryProcessor.cpp
@@ -326,6 +326,16 @@ static bool ParseIoStoreShaderLibraryName(FShaderLibraryNameInfo& OutLibraryName
 	return true;
 }
 
+FIoStoreShaderLibraryProcessor::FIoStoreShaderLibraryProcessor()
+{
+	// Explicit default constructor to ensure AllMonolithicShaderArchives and AllAssignedSerializedShaders need their full type declaration only in this file
+}
+
+FIoStoreShaderLibraryProcessor::~FIoStoreShaderLibraryProcessor()
+{
+	// Explicit default destructor to ensure AllMonolithicShaderArchives and AllAssignedSerializedShaders need their full type declaration only in this file
+}
+
 bool FIoStoreShaderLibraryProcessor::LoadSerializedShaders(
 	const FContainerTargetFile& TargetFile, FSerializedShaderArchiveBuffer& OutSerializedShaders, bool bIncludeTypeInfo, FShaderLibraryNameInfo* OutNameInfo)
 {
@@ -351,21 +361,21 @@ bool FIoStoreShaderLibraryProcessor::LoadSerializedShaders(
 	}
 
 	// Load serialized shaders from *.ushaderbytecode file
-	FCookedPackageStore* ActivePackageStore = PackageStore;
-	if (!OutSerializedShaders.LoadFromTargetFile(InFilename, &TargetFile, PackageStore))
+	FCookedPackageStore* CurrentPackageStore = PackageStore;
+	if (!OutSerializedShaders.LoadFromTargetFile(InFilename, &TargetFile, CurrentPackageStore))
 	{
 		return false;
 	}
 
-	if (PackageStore && !TargetFile.ChunkId.IsValid())
+	if (CurrentPackageStore && !TargetFile.ChunkId.IsValid())
 	{
-		PackageStore = nullptr;
+		CurrentPackageStore = nullptr;
 	}
 
 	// Load *.assetinfo.json file for the specified input shader library and store its content in the ShaderMapAssetAssociations container
 	{
 		const FString AssetInfoFilename = FPaths::GetPath(InFilename) / LibraryNameInfo.GetAssetInfoFilename();
-		if (!LoadShaderAssetInfo(AssetInfoFilename, PackageStore, OutSerializedShaders.GetShaderMapAssetAssociations()))
+		if (!LoadShaderAssetInfo(AssetInfoFilename, CurrentPackageStore, OutSerializedShaders.GetShaderMapAssetAssociations()))
 		{
 			UE_LOG(LogIoStore, Error, TEXT("Failed loading shader asset info file '%s'"), *AssetInfoFilename);
 			return false;
@@ -377,7 +387,7 @@ bool FIoStoreShaderLibraryProcessor::LoadSerializedShaders(
 	if (bIncludeTypeInfo)
 	{
 		const FString TypeInfoFilename = FPaths::GetPath(InFilename) / LibraryNameInfo.GetTypeInfoFilename();
-		if (!LoadShaderTypeInfo(TypeInfoFilename, PackageStore, OutSerializedShaders.ShaderTypes))
+		if (!LoadShaderTypeInfo(TypeInfoFilename, CurrentPackageStore, OutSerializedShaders.ShaderTypes))
 		{
 			UE_LOG(LogIoStore, Error, TEXT("Failed loading shader type info file '%s'"), *TypeInfoFilename);
 			return false;
@@ -567,20 +577,6 @@ bool FIoStoreShaderLibraryProcessor::ConvertToIoStoreShaderLibrary(FContainerTar
 	return true;
 }
 
-FIoStoreShaderLibraryProcessor::~FIoStoreShaderLibraryProcessor()
-{
-#if WITH_EDITORONLY_DATA
-	for (TPair<FString, FEditorShaderCodeArchive*> ShaderArchive : AllMonolithicShaderArchives)
-	{
-		delete ShaderArchive.Value;
-	}
-	for (TPair<const FContainerTargetFile*, FSerializedShaderArchiveBuffer*> ShaderArchive : AllAssignedSerializedShaders)
-	{
-		delete ShaderArchive.Value;
-	}
-#endif
-}
-
 // Invokes TargetFileCallback on each ShaderCodeLibrary entry in all container targets.
 static void ForEachShaderCodeLibraryTargetFile(const TArrayView<FContainerTargetSpec*>& ContainerTargets, const TFunction<void(FContainerTargetFile&)>& TargetFileCallback)
 {
@@ -897,7 +893,7 @@ bool FIoStoreShaderLibraryProcessor::ProcessShaderLibraryFromTargetFile(const FI
 	{
 #if WITH_EDITORONLY_DATA
 		// Find shader archive that was assigned to the current target file and convert it into an IoStore shader library
-		if (FSerializedShaderArchiveBuffer* const* SerializedShaders = AllAssignedSerializedShaders.Find(&TargetFile))
+		if (TUniquePtr<FSerializedShaderArchiveBuffer> const* SerializedShaders = AllAssignedSerializedShaders.Find(&TargetFile))
 		{
 			if (!ConvertToIoStoreShaderLibrary(TargetFile, **SerializedShaders))
 			{
@@ -958,10 +954,10 @@ bool FIoStoreShaderLibraryProcessor::PreloadShaderArchive(const FContainerTarget
 	ShaderMapAssetAssociations.Append(SerializedShaders.GetShaderMapAssetAssociations());
 
 	// Append serialized shaders to monolithic library so we can later create new chunks from it
-	FEditorShaderCodeArchive*& MonolithicShaderArchive = AllMonolithicShaderArchives.FindOrAdd(LibraryNameInfo.GetFormatAndPlatformName());
+	TUniquePtr<FEditorShaderCodeArchive>& MonolithicShaderArchive = AllMonolithicShaderArchives.FindOrAdd(LibraryNameInfo.GetFormatAndPlatformName());
 	if (!MonolithicShaderArchive)
 	{
-		MonolithicShaderArchive = new FEditorShaderCodeArchive(FName(LibraryNameInfo.GetFormatAndPlatformName()));
+		MonolithicShaderArchive = TUniquePtr<FEditorShaderCodeArchive>(new FEditorShaderCodeArchive(FName(LibraryNameInfo.GetFormatAndPlatformName())));
 	}
 	MonolithicShaderArchive->AppendFromArchive(SerializedShaders, SerializedShaders.GetShaderCodeView(), ShaderCodeArchive::ECookShaderLibrarySource::CurrentCook);
 
@@ -1080,9 +1076,9 @@ bool FIoStoreShaderLibraryProcessor::AssignShaderArchives(const FContainerTarget
 
 	// Build intermediate shader archive and assign relevant shaders from preloaded serialized shaders
 	const FString FormatAndPlatformName = LibraryNameInfo.GetFormatAndPlatformName();
-	if (FEditorShaderCodeArchive** FoundMonolithicShaderArchive = AllMonolithicShaderArchives.Find(FormatAndPlatformName))
+	if (TUniquePtr<FEditorShaderCodeArchive>* FoundMonolithicShaderArchive = AllMonolithicShaderArchives.Find(FormatAndPlatformName))
 	{
-		FEditorShaderCodeArchive* MonolithicShaderArchive = *FoundMonolithicShaderArchive;
+		FEditorShaderCodeArchive* MonolithicShaderArchive = FoundMonolithicShaderArchive->Get();
 		check(MonolithicShaderArchive);
 
 		// Create new chunk for packages.
@@ -1090,7 +1086,7 @@ bool FIoStoreShaderLibraryProcessor::AssignShaderArchives(const FContainerTarget
 		check(NewChunkShaderArchive);
 
 		// Convert constructed intermediate shader archive to final serialized archive with flat buffer
-		FSerializedShaderArchiveBuffer* AssignedShaderArchive = new FSerializedShaderArchiveBuffer();
+		TUniquePtr<FSerializedShaderArchiveBuffer> AssignedShaderArchive = TUniquePtr<FSerializedShaderArchiveBuffer>(new FSerializedShaderArchiveBuffer());
 		{
 			if (AssignedShaderArchive->LoadFromEditorArchive(*NewChunkShaderArchive))
 			{
@@ -1102,14 +1098,14 @@ bool FIoStoreShaderLibraryProcessor::AssignShaderArchives(const FContainerTarget
 				UE_LOG(LogIoStore, Error, TEXT("Could not serialize shaders from monolithic archive '%s'"), *LibraryNameInfo.GetLibraryName());
 				return false;
 			}
-		}
-		AllAssignedSerializedShaders.Add(&TargetFile, AssignedShaderArchive);
 
-		// Keep track of old and new pakchunks if debug information is enabled
-		if (IsChunkOverrideDebugEnabled())
-		{
-			AddShaderMapsToPakChunkDebugInfo(TargetFile, *AssignedShaderArchive, FName(FormatAndPlatformName), ChunkAssignment_New);
+			// Keep track of old and new pakchunks if debug information is enabled
+			if (IsChunkOverrideDebugEnabled())
+			{
+				AddShaderMapsToPakChunkDebugInfo(TargetFile, *AssignedShaderArchive, FName(FormatAndPlatformName), ChunkAssignment_New);
+			}
 		}
+		AllAssignedSerializedShaders.Add(&TargetFile, MoveTemp(AssignedShaderArchive));
 	}
 	else
 	{
diff --git a/Engine/Source/Developer/IoStoreUtilities/Private/IoStoreShaderLibraryProcessor.h b/Engine/Source/Developer/IoStoreUtilities/Private/IoStoreShaderLibraryProcessor.h
--- a/Engine/Source/Developer/IoStoreUtilities/Private/IoStoreShaderLibraryProcessor.h
+++ b/Engine/Source/Developer/IoStoreUtilities/Private/IoStoreShaderLibraryProcessor.h
@@ -117,8 +117,8 @@ class FIoStoreShaderLibraryProcessor
 
 #if WITH_EDITORONLY_DATA
 	// Intermediate data for chunk re-assignment
-	TMap<FString, FEditorShaderCodeArchive*> AllMonolithicShaderArchives; // Shader format to archive
-	TMap<const FContainerTargetFile*, FSerializedShaderArchiveBuffer*> AllAssignedSerializedShaders;
+	TMap<FString, TUniquePtr<FEditorShaderCodeArchive>> AllMonolithicShaderArchives; // Shader format to archive
+	TMap<const FContainerTargetFile*, TUniquePtr<FSerializedShaderArchiveBuffer>> AllAssignedSerializedShaders;
 
 	// Debug information (used with CVar `r.ShaderCodeLibrary.PakChunkOverride.DumpDebugInfo`)
 	enum EChunkAssignmentType
@@ -139,6 +139,7 @@ class FIoStoreShaderLibraryProcessor
 #endif
 
 public:
+	FIoStoreShaderLibraryProcessor();
 	~FIoStoreShaderLibraryProcessor();
 
 	void ProcessShaderLibraries(const FIoStoreArguments& Arguments, TArrayView<FContainerTargetSpec*> ContainerTargets, FIoStoreShaderLibraryOutput& Output);
