[Shaders] Support shadermap deduplication during shader chunk override using the chunk dependency tree.

This is implemented by passing a lambda to CreateNamedChunk() when we assign shaders from the monolithic shader archive.
This lambda returns true if a given shadermap is meant to be excluded because it is already stored in one of the current container's chunk dependencies.

#rnx

#Synced-CL 50477682

diff --git a/Engine/Source/Developer/IoStoreUtilities/Private/IoStoreShaderLibraryProcessor.cpp b/Engine/Source/Developer/IoStoreUtilities/Private/IoStoreShaderLibraryProcessor.cpp
--- a/Engine/Source/Developer/IoStoreUtilities/Private/IoStoreShaderLibraryProcessor.cpp
+++ b/Engine/Source/Developer/IoStoreUtilities/Private/IoStoreShaderLibraryProcessor.cpp
@@ -6,6 +6,7 @@
 #include "IO/IoStore.h"
 #include "IoStoreUtilitiesCore.h"
 #include "Misc/FileHelper.h"
+#include "Misc/ChunkDependencyInfo.h"
 #include "Serialization/JsonReader.h"
 #include "Serialization/JsonSerializer.h"
 #include "Serialization/JsonWriter.h"
@@ -1054,6 +1055,28 @@ bool FIoStoreShaderLibraryProcessor::DumpChunkOverrideDebugInfo(FName FormatAndP
 	return true;
 }
 
+// Print the set of chunk dependency IDs to string for logging, e.g. "None", "0", "0, 1" etc.
+static FString PrintChunkDependenciesToString(const TSet<int32>& InChunkDependencies)
+{
+	TStringBuilder<1024> OutString;
+	if (InChunkDependencies.IsEmpty())
+	{
+		OutString.Append(TEXT("None"));
+	}
+	else
+	{
+		for (int32 ChunkId : InChunkDependencies)
+		{
+			if (OutString.Len() > 0)
+			{
+				OutString.Append(TEXT(", "));
+			}
+			OutString.Appendf(TEXT("%d"), ChunkId);
+		}
+	}
+	return FString(OutString);
+}
+
 bool FIoStoreShaderLibraryProcessor::AssignShaderArchives(const FContainerTargetFile& TargetFile)
 {
 	// Convert shader format and platform name from *.ushaderbytecode filename
@@ -1074,6 +1097,50 @@ bool FIoStoreShaderLibraryProcessor::AssignShaderArchives(const FContainerTarget
 		PackagesInChunk.Add(CookedPackage->PackageName);
 	}
 
+	// Gather all chunk dependencies to deduplicate redundant shadermap entries
+	FContainerChunkDependency& ContainerChunkDependency = ContainerTargetToChunkDependencies.FindOrAdd(ContainerTarget);
+	if (ContainerChunkDependency.ChunkId == INDEX_NONE)
+	{
+		FString ContainerTargetChunkIdSuffix;
+		if (FShaderLibraryNameInfo::ParsePakChunkId(ContainerTarget->Name.ToString(), ContainerChunkDependency.ChunkId, ContainerTargetChunkIdSuffix))
+		{
+			if (UChunkDependencyInfo* DependencyInfo = GetMutableDefault<UChunkDependencyInfo>())
+			{
+				// Get chunk dependency graph to determine what shadermaps are already included in a parent chunk
+				DependencyInfo->GetOrBuildChunkDependencyGraph(ContainerChunkDependency.ChunkId);
+				DependencyInfo->GetChunkDependencies(ContainerChunkDependency.ChunkId, ContainerChunkDependency.ChunkDependencies);
+				ContainerChunkDependency.ChunkDependencies.Remove(ContainerChunkDependency.ChunkId);
+				ContainerChunkDependency.bIsExternalContainer = !ContainerTargetChunkIdSuffix.IsEmpty();
+
+				// Log information about chunk dependencies for shader library
+				UE_LOG(LogIoStore, Display, TEXT("Processing chunk dependency for container '%s': ChunkId=%d, bIsExternalContainer=%s, ChunkDependencies=(%s)"),
+					*ContainerTarget->Name.ToString(), ContainerChunkDependency.ChunkId, ContainerChunkDependency.bIsExternalContainer ? TEXT("Yes") : TEXT("No"),
+					*PrintChunkDependenciesToString(ContainerChunkDependency.ChunkDependencies));
+			}
+		}
+	}
+
+	// Lambda to determine what shadermaps to exclude from current target file to deduplicate them across their chunk dependency tree
+	auto ExcludeShaderMap = [this, &ContainerChunkDependency, ContainerTarget](const FSHAHash& ShaderMapHash) -> bool
+		{
+			// Check if shadermap hash is already stored in any of the current container's chunk dependencies
+			for (int32 ChunkDependencyId : ContainerChunkDependency.ChunkDependencies)
+			{
+				if (const FContainerShaderArchives* ContainerInfo = this->ChunkIdToContainerInfo.Find(ChunkDependencyId))
+				{
+					for (const FSerializedShaderArchiveBuffer* ShaderArchive : ContainerInfo->AssignedSerializedShaders)
+					{
+						if (ShaderArchive->ShaderMapHashes.Contains(ShaderMapHash))
+						{
+							// Found shadermap in a container that is a parent of the current container, so exclude this shadermap from assigning it to this one
+							return true;
+						}
+					}
+				}
+			}
+			return false;
+		};
+
 	// Build intermediate shader archive and assign relevant shaders from preloaded serialized shaders
 	const FString FormatAndPlatformName = LibraryNameInfo.GetFormatAndPlatformName();
 	if (TUniquePtr<FEditorShaderCodeArchive>* FoundMonolithicShaderArchive = AllMonolithicShaderArchives.Find(FormatAndPlatformName))
@@ -1082,7 +1149,7 @@ bool FIoStoreShaderLibraryProcessor::AssignShaderArchives(const FContainerTarget
 		check(MonolithicShaderArchive);
 
 		// Create new chunk for packages.
-		TUniquePtr<FEditorShaderCodeArchive> NewChunkShaderArchive = MonolithicShaderArchive->CreateNamedChunk(LibraryNameInfo.ChunkName, PackagesInChunk);
+		TUniquePtr<FEditorShaderCodeArchive> NewChunkShaderArchive = MonolithicShaderArchive->CreateNamedChunk(LibraryNameInfo.ChunkName, PackagesInChunk, ExcludeShaderMap);
 		check(NewChunkShaderArchive);
 
 		// Convert constructed intermediate shader archive to final serialized archive with flat buffer
@@ -1099,6 +1166,18 @@ bool FIoStoreShaderLibraryProcessor::AssignShaderArchives(const FContainerTarget
 				return false;
 			}
 
+			// Link shader archive to current container chunk dependency if this is not an external container, i.e. no OnDemand or Optional ID suffix.
+			// This information can be retrieved by subsequent shader archives in `ExcludeShaderMap` using their chunk dependency IDs
+			// to check what shadermaps are already contained in their parent chunks to deduplicate them.
+			// @lh-todo:
+			//  This currently relies on the containers arriving in the order of their dependency, which means this is potentially missing out on a tighter deduplication.
+			//  For the exact order, we need to iterate through all containers according to the dependency graph.
+			if (ContainerChunkDependency.ChunkId != INDEX_NONE && !ContainerChunkDependency.bIsExternalContainer)
+			{
+				FContainerShaderArchives& ContainerInfo = ChunkIdToContainerInfo.FindOrAdd(ContainerChunkDependency.ChunkId);
+				ContainerInfo.AssignedSerializedShaders.Add(AssignedShaderArchive.Get());
+			}
+
 			// Keep track of old and new pakchunks if debug information is enabled
 			if (IsChunkOverrideDebugEnabled())
 			{
diff --git a/Engine/Source/Developer/IoStoreUtilities/Private/IoStoreShaderLibraryProcessor.h b/Engine/Source/Developer/IoStoreUtilities/Private/IoStoreShaderLibraryProcessor.h
--- a/Engine/Source/Developer/IoStoreUtilities/Private/IoStoreShaderLibraryProcessor.h
+++ b/Engine/Source/Developer/IoStoreUtilities/Private/IoStoreShaderLibraryProcessor.h
@@ -100,6 +100,11 @@ class FIoStoreShaderLibraryProcessor
 		TArray<FIoChunkId> ShaderGroupChunkIds;
 	};
 
+	struct FContainerShaderArchives
+	{
+		TArray<FSerializedShaderArchiveBuffer*> AssignedSerializedShaders;
+	};
+
 	// Input arguments
 	FCookedPackageStore* PackageStore = nullptr;
 	FOodleDataCompression::ECompressor ShaderOodleCompressor = FOodleDataCompression::ECompressor::NotSet;
@@ -111,6 +116,21 @@ class FIoStoreShaderLibraryProcessor
 	TMap<FContainerTargetSpec*, TSet<FShaderInfo*>> ContainerTargetToShaderLibrary;
 	FShaderMapAssetAssociations ShaderMapAssetAssociations;
 
+	//@lh-todo
+	// Move this information into FContainerSourceSpec and FContainerTargetSpec.
+	// This information should be determined in `ParseContainerGenerationArguments()` when the source containers are built,
+	// and then transferred to the target containers in `InitializeContainerTargetsAndPackages()`.
+	// More testing beyond shaders for this is required, so keep this information private to the shader library processor for now.
+	// Same goes for `IsTargetFileGlobalShaderLibrary()`, which should be determined as soon as the source container is generated
+	// instead of parsing the filename later on to extract those flags.
+	struct FContainerChunkDependency
+	{
+		int32 ChunkId = INDEX_NONE;
+		bool bIsExternalContainer = false; // Is the container loaded externally? E.g. pakchunk100iad, pakchunk100ondemand, etc.
+		TSet<int32> ChunkDependencies;
+	};
+	TMap<const FContainerTargetSpec*, FContainerChunkDependency> ContainerTargetToChunkDependencies;
+
 	// Intermediate data while processing a subset of target files
 	TArray<FShaderGroupIoChunk> CodeIoChunks;
 	TArray<FShaderGroupIoChunkMap> CodeIoChunkMaps; // Maps from shadermap index to shader group IoChunks
@@ -119,6 +139,7 @@ class FIoStoreShaderLibraryProcessor
 	// Intermediate data for chunk re-assignment
 	TMap<FString, TUniquePtr<FEditorShaderCodeArchive>> AllMonolithicShaderArchives; // Shader format to archive
 	TMap<const FContainerTargetFile*, TUniquePtr<FSerializedShaderArchiveBuffer>> AllAssignedSerializedShaders;
+	TMap<int32, FContainerShaderArchives> ChunkIdToContainerInfo; // Maps ChunkId to the container information with all their assigned shader archives
 
 	// Debug information (used with CVar `r.ShaderCodeLibrary.PakChunkOverride.DumpDebugInfo`)
 	enum EChunkAssignmentType
diff --git a/Engine/Source/Runtime/RenderCore/Internal/ShaderLibrary/EditorShaderCodeArchive.h b/Engine/Source/Runtime/RenderCore/Internal/ShaderLibrary/EditorShaderCodeArchive.h
--- a/Engine/Source/Runtime/RenderCore/Internal/ShaderLibrary/EditorShaderCodeArchive.h
+++ b/Engine/Source/Runtime/RenderCore/Internal/ShaderLibrary/EditorShaderCodeArchive.h
@@ -135,7 +135,7 @@ public:
 	RENDERCORE_API TUniquePtr<FEditorShaderCodeArchive> CreateChunk(int32 ChunkId, const TSet<FName>& PackagesInChunk);
 
 	/** Produces another archive that contains the code only for these assets and uses a custom chunk name. */
-	RENDERCORE_API TUniquePtr<FEditorShaderCodeArchive> CreateNamedChunk(const FString& ChunkName, const TSet<FName>& PackagesInChunk);
+	RENDERCORE_API TUniquePtr<FEditorShaderCodeArchive> CreateNamedChunk(const FString& ChunkName, const TSet<FName>& PackagesInChunk, FShaderMapAssetAssociations::FShaderMapFilterFunction ShouldExcludeShaderMap = nullptr);
 
 	/** Adds entries for the specified shader map from another archive with all its individual shaders and their code to this archive. */
 	RENDERCORE_API int32 AddShaderCode(int32 OtherShaderMapIndex, const FEditorShaderCodeArchive& OtherArchive, ECookShaderLibrarySource CookSource);
diff --git a/Engine/Source/Runtime/RenderCore/Private/ShaderLibrary/EditorShaderCodeArchive.cpp b/Engine/Source/Runtime/RenderCore/Private/ShaderLibrary/EditorShaderCodeArchive.cpp
--- a/Engine/Source/Runtime/RenderCore/Private/ShaderLibrary/EditorShaderCodeArchive.cpp
+++ b/Engine/Source/Runtime/RenderCore/Private/ShaderLibrary/EditorShaderCodeArchive.cpp
@@ -477,14 +477,15 @@ TUniquePtr<FEditorShaderCodeArchive> FEditorShaderCodeArchive::CreateChunk(int32
 	return CreateNamedChunk(UE::ShaderLibrary::Private::GetShaderLibraryNameForChunk(LibraryName, ChunkId), PackagesInChunk);
 }
 
-TUniquePtr<FEditorShaderCodeArchive> FEditorShaderCodeArchive::CreateNamedChunk(const FString& ChunkName, const TSet<FName>& PackagesInChunk)
+TUniquePtr<FEditorShaderCodeArchive> FEditorShaderCodeArchive::CreateNamedChunk(const FString& ChunkName, const TSet<FName>& PackagesInChunk, FShaderMapAssetAssociations::FShaderMapFilterFunction ShouldExcludeShaderMap)
 {
 	return CreateArchiveFromFilteredAssets(
 		ChunkName,
 		[&PackagesInChunk](FName AssetName, const FShaderMapAssetAssociations::FAssociatedAssetData& AssetData) -> bool
 		{
 			return PackagesInChunk.Contains(AssetName);
-		});
+		},
+		ShouldExcludeShaderMap);
 }
 
 int32 FEditorShaderCodeArchive::AddShaderCode(int32 OtherShaderMapIndex, const FEditorShaderCodeArchive& OtherArchive,
