[ShaderLibrary] Don't allow shader chunk override if chunks aren't meant to be generated for the project being packaged.

If the project is not configured to generate chunks for a specific platform, i.e. `[/Script/UnrealEd.ProjectPackagingSettings]:bGenerateChunks=False`, there is no reason to re-assign shaders.
This might even fail if attempting to re-assign shaders to a single large consecutive buffer that exceeds the limit of `TArray<uint8>::Num()`.

Also set a shader code size limit in CopyToArchiveAndClear() to prevent a crash when attempting to assign shaders to a container larger than 4 GB .
This is limited by the maximum capacity of `TArray<uint8>::Num()`.

#rnx

#Synced-CL 50573404

diff --git a/Engine/Source/Developer/IoStoreUtilities/Private/IoStoreShaderLibraryProcessor.cpp b/Engine/Source/Developer/IoStoreUtilities/Private/IoStoreShaderLibraryProcessor.cpp
--- a/Engine/Source/Developer/IoStoreUtilities/Private/IoStoreShaderLibraryProcessor.cpp
+++ b/Engine/Source/Developer/IoStoreUtilities/Private/IoStoreShaderLibraryProcessor.cpp
@@ -7,6 +7,7 @@
 #include "IoStoreUtilitiesCore.h"
 #include "Misc/FileHelper.h"
 #include "Misc/ChunkDependencyInfo.h"
+#include "Misc/ConfigCacheIni.h"
 #include "Serialization/JsonReader.h"
 #include "Serialization/JsonSerializer.h"
 #include "Serialization/JsonWriter.h"
@@ -36,11 +37,39 @@ static TAutoConsoleVariable<bool> CVarShaderLibraryPakChunkOverrideDumpDebugInfo
 	TEXT("Also requires `r.ShaderCodeLibrary.PakChunkOverride=true` or command line argument `-OverrideShaderChunks`.")
 );
 
-// Returns true if shader code chunk override is enabled (CVar 'r.ShaderCodeLibrary.PakChunkOverride').
+// Returns true if the project being packaged is meant to generate chunks for the specified platform.
+// Without chunks being generated, there is nothing to re-assign shader chunks to.
+// This might even fail if attempting to re-assign shaders to a single large (> 4 GB) consecutive buffer.
+static bool IsGenerateChunksEnabled()
+{
+	FString PlatformIniName;
+	if (FParse::Value(FCommandLine::Get(), TEXT("platform="), PlatformIniName, false))
+	{
+		FConfigFile PlatformIniFile;
+		FConfigCacheIni::LoadLocalIniFile(PlatformIniFile, TEXT("Game"), true, *PlatformIniName);
+		bool bGenerateChunks = false;
+		if (PlatformIniFile.GetBool(TEXT("/Script/UnrealEd.ProjectPackagingSettings"), TEXT("bGenerateChunks"), bGenerateChunks))
+		{
+			return bGenerateChunks;
+		}
+	}
+	return false;
+}
+
+// Returns true if shader code chunk override is enabled, via CVar `r.ShaderCodeLibrary.PakChunkOverride` or commandline `-OverrideShaderChunks`.
+// This also requires that the project being packaged is configured to generate chunks, i.e. `[/Script/UnrealEd.ProjectPackagingSettings]:bGenerateChunks=True`.
 static bool IsChunkOverrideEnabled()
 {
-	static bool bIsChunkOverrideEnabledByCmdLine = FParse::Param(FCommandLine::Get(), TEXT("OverrideShaderChunks"));
-	return bIsChunkOverrideEnabledByCmdLine || CVarShaderLibraryPakChunkOverride.GetValueOnAnyThread();
+	static bool bIsChunkOverrideEnabled = []()
+		{
+			if (IsGenerateChunksEnabled())
+			{
+				const bool bIsChunkOverrideEnabledByCmdLine = FParse::Param(FCommandLine::Get(), TEXT("OverrideShaderChunks"));
+				return bIsChunkOverrideEnabledByCmdLine || CVarShaderLibraryPakChunkOverride.GetValueOnAnyThread();
+			}
+			return false;
+		}();
+	return bIsChunkOverrideEnabled;
 }
 
 static bool IsChunkOverrideDebugEnabled()
@@ -225,10 +254,13 @@ public:
 		// Copy editor archive to this archive
 		TArray<uint8> FlatShaderCode;
 		bool bRanOutOfRoom = false;
-		constexpr int64 kNoShaderSizeLimit = -1;
+		constexpr int64 kShaderSizeLimit = MAX_int32; // Limited by max capacity of TArray<>::Num()
 		constexpr int64 kNoShaderCountLimit = -1;
-		ShaderCodeArchive.CopyToArchiveAndClear(*this, FlatShaderCode, bRanOutOfRoom, kNoShaderSizeLimit, kNoShaderCountLimit, ECookShaderLibrarySource::CurrentCook);
-		check(!bRanOutOfRoom);
+		ShaderCodeArchive.CopyToArchiveAndClear(*this, FlatShaderCode, bRanOutOfRoom, kShaderSizeLimit, kNoShaderCountLimit, ECookShaderLibrarySource::CurrentCook);
+		if (bRanOutOfRoom)
+		{
+			return false;
+		}
 
 		// Write header
 		FLargeMemoryWriter HeaderWriter;
@@ -690,7 +722,7 @@ void FIoStoreShaderLibraryProcessor::ProcessShaderLibraries(const FIoStoreArgume
 	if (IsChunkOverrideEnabled())
 	{
 #if WITH_EDITORONLY_DATA
-		UE_LOG(LogIoStore, Display, TEXT("Shader library chunk override is enabled (r.ShaderCodeLibrary.PakChunkOverride=true)"));
+		UE_LOG(LogIoStore, Display, TEXT("Shader library chunk override is enabled"));
 
 		//@lh-todo:
 		//  We need to generate all target files for ShaderCodeLibrary chunks. The shader library headers for both IAD and non-IAD
@@ -751,12 +783,12 @@ void FIoStoreShaderLibraryProcessor::ProcessShaderLibraries(const FIoStoreArgume
 			});
 		}
 #else
-		UE_LOG(LogIoStore, Warning, TEXT("Shader library chunk override is enabled (r.ShaderCodeLibrary.PakChunkOverride=true), but not supported (WITH_EDITORONLY_DATA=0)"));
+		UE_LOG(LogIoStore, Warning, TEXT("Shader library chunk override is enabled, but not supported (WITH_EDITORONLY_DATA=0)"));
 #endif // WITH_EDITORONLY_DATA
 	}
 	else
 	{
-		UE_LOG(LogIoStore, Display, TEXT("Shader library chunk override is disabled (r.ShaderCodeLibrary.PakChunkOverride=false)"));
+		UE_LOG(LogIoStore, Display, TEXT("Shader library chunk override is disabled"));
 	}
 
 	{
diff --git a/Engine/Source/Runtime/RenderCore/Private/ShaderLibrary/EditorShaderCodeArchive.cpp b/Engine/Source/Runtime/RenderCore/Private/ShaderLibrary/EditorShaderCodeArchive.cpp
--- a/Engine/Source/Runtime/RenderCore/Private/ShaderLibrary/EditorShaderCodeArchive.cpp
+++ b/Engine/Source/Runtime/RenderCore/Private/ShaderLibrary/EditorShaderCodeArchive.cpp
@@ -252,7 +252,7 @@ void FEditorShaderCodeArchive::CopyToArchiveAndClear(FSerializedShaderArchive& T
 				{
 					TargetShaderEntry.Offset = TargetFlatShaderCode.Num();
 					check(SourceShaderEntry.Size == SourceShaderCode.GetSize());
-					if ((MaxShaderSize > 0 && TargetFlatShaderCode.Num() + (int64)SourceShaderCode.GetSize() > MaxShaderSize) ||
+					if ((MaxShaderSize > 0 && (int64)TargetFlatShaderCode.Num() + (int64)SourceShaderCode.GetSize() > MaxShaderSize) ||
 						(MaxShaderCount > 0 && NumShadersSentWithCode + 1 > MaxShaderCount))
 					{
 						// We have to stop here to avoid overflowing the shader limit. Send the shaders we have accumulated
