[Shaders] Log shadermap duplication stats when chunk override is enabled.

This logs the number of unique shadermaps and their copies in the old and new chunks when shader chunk override is enabled.
Here is an example for PCD3D_SM6:
```
LogIoStore: Display: Shader library collection PCD3D_SM6-PCD3D_SM6 shadermap stats:
Unique shader maps:              44338
Shader map copies in old chunks: [REDACTED_PHONE] : 1 ratio)
Shader map copies in new chunks: [REDACTED_PHONE] : 1 ratio)
```

#rnx

#Synced-CL 50638327

diff --git a/Engine/Source/Developer/IoStoreUtilities/Private/IoStoreShaderLibraryProcessor.cpp b/Engine/Source/Developer/IoStoreUtilities/Private/IoStoreShaderLibraryProcessor.cpp
--- a/Engine/Source/Developer/IoStoreUtilities/Private/IoStoreShaderLibraryProcessor.cpp
+++ b/Engine/Source/Developer/IoStoreUtilities/Private/IoStoreShaderLibraryProcessor.cpp
@@ -897,6 +897,8 @@ void FIoStoreShaderLibraryProcessor::ProcessShaderLibraries(const FIoStoreArgume
 	{
 		for (const FString& FormatAndPlatformName : FormatAndPlatformNames)
 		{
+			LogChunkOverrideStats(FName(FormatAndPlatformName));
+
 			const FString DebugLibFolder = FSerializedShaderArchive::MakeDebugDirectory() / FormatAndPlatformName;
 			IFileManager::Get().MakeDirectory(*DebugLibFolder, true);
 			const FString DebugFilename = DebugLibFolder / FString::Printf(TEXT("ShaderChunkOverride-%s.txt"), *FormatAndPlatformName);
@@ -1023,6 +1025,36 @@ void FIoStoreShaderLibraryProcessor::AddShaderMapsToPakChunkDebugInfo(
 	}
 }
 
+void FIoStoreShaderLibraryProcessor::LogChunkOverrideStats(FName FormatAndPlatformName)
+{
+	if (const FMonolithicArchiveDebugInfo* DebugInfo = ChunkAssignmentDebugInfo.Find(FormatAndPlatformName))
+	{
+		const int32 NumUniqueShaderMaps = DebugInfo->ShaderMaps.Num();
+
+		int32 NumShaderMapCopiesInOldChunks = 0;
+		int32 NumShaderMapCopiesInNewChunks = 0;
+
+		for (const TPair<FSHAHash, FChunkAssignmentInfo>& Pair : DebugInfo->ShaderMaps)
+		{
+			NumShaderMapCopiesInOldChunks += Pair.Value.OldPakChunks.Num();
+			NumShaderMapCopiesInNewChunks += Pair.Value.NewPakChunks.Num();
+		}
+
+		UE_LOG(LogIoStore, Display,
+			TEXT("Shader library collection %s shadermap stats:\n")
+			TEXT("  Unique shader maps:              %d\n")
+			TEXT("  Shader map copies in old chunks: %d (%.3f : 1 ratio)\n")
+			TEXT("  Shader map copies in new chunks: %d (%.3f : 1 ratio)"),
+			*FormatAndPlatformName.ToString(),
+			NumUniqueShaderMaps,
+			NumShaderMapCopiesInOldChunks,
+			NumUniqueShaderMaps > 0 ? static_cast<double>(NumShaderMapCopiesInOldChunks) / static_cast<double>(NumUniqueShaderMaps) : 0.0,
+			NumShaderMapCopiesInNewChunks,
+			NumUniqueShaderMaps > 0 ? static_cast<double>(NumShaderMapCopiesInNewChunks) / static_cast<double>(NumUniqueShaderMaps) : 0.0
+		);
+	}
+}
+
 bool FIoStoreShaderLibraryProcessor::DumpChunkOverrideDebugInfo(FName FormatAndPlatformName, FArchive& Ar)
 {
 	if (!ChunkAssignmentDebugInfo.Contains(FormatAndPlatformName))
diff --git a/Engine/Source/Developer/IoStoreUtilities/Private/IoStoreShaderLibraryProcessor.h b/Engine/Source/Developer/IoStoreUtilities/Private/IoStoreShaderLibraryProcessor.h
--- a/Engine/Source/Developer/IoStoreUtilities/Private/IoStoreShaderLibraryProcessor.h
+++ b/Engine/Source/Developer/IoStoreUtilities/Private/IoStoreShaderLibraryProcessor.h
@@ -178,6 +178,9 @@ private:
 	void AddShaderMapsToPakChunkDebugInfo(const FContainerTargetFile& TargetFile, const FSerializedShaderArchiveBuffer& SerializedShaders,
 		FName FormatAndPlatformName, int32 ChunkAssignmentType);
 
+	// Logs statistics about how many shadermaps have been duplicated and deduplicated.
+	void LogChunkOverrideStats(FName FormatAndPlatformName);
+
 	// Writes out a debug file with information about all the shaders that have been re-assigned
 	// to a difference chunk (if r.ShaderCodeLibrary.PakChunkOverride is enabled).
 	bool DumpChunkOverrideDebugInfo(FName FormatAndPlatformName, FArchive& Ar);
