UnrealPak Patch-in-place defragmentation to limit container size growth

#rnx
#Synced-CL 51008151

diff --git a/Engine/Source/Developer/IoStoreUtilities/Private/IoStoreUtilities.cpp b/Engine/Source/Developer/IoStoreUtilities/Private/IoStoreUtilities.cpp
--- a/Engine/Source/Developer/IoStoreUtilities/Private/IoStoreUtilities.cpp
+++ b/Engine/Source/Developer/IoStoreUtilities/Private/IoStoreUtilities.cpp
@@ -2299,8 +2299,9 @@ void LogWriterResults(const FIoStoreWriterSettings& GeneralIoWriterSettings, con
 		uint64 UncompressedContainerSize = 0;
 		uint64 CompressedContainerSize = 0;
 		uint64 PaddingSize = 0;
-		uint64 PatchInPlacePinnedSize = 0;
-		int64  PatchInPlaceContainerSizeChange = 0;
+		uint64 PatchInPlaceSizeBeforeLayout = 0;
+		uint64 PatchInPlacePinnedSizeBeforeDefrag = 0;
+		uint64 PatchInPlacePinnedSizeAfterDefrag = 0;
 	};
 
 	// Container name length can vary wildly, so find the longest one so that we can correctly align the output grids
@@ -2319,11 +2320,11 @@ void LogWriterResults(const FIoStoreWriterSettings& GeneralIoWriterSettings, con
 	FString PatchInPlaceHeader = TEXT("");
 	if (GeneralIoWriterSettings.bUsePatchInPlaceLayout)
 	{
-		PatchInPlaceHeader = FString::Printf(TEXT("%20s %20s"), TEXT("PIP Pinned (MiB)"), TEXT("PIP Overhead (MiB)"));
+		PatchInPlaceHeader = FString::Printf(TEXT(" %20s %20s %20s"), TEXT("PIP Pinned (MiB)"), TEXT("PIP Defrag (MiB)"), TEXT("PIP Overhead (MiB)"));
 	}
 	UE_LOG(LogIoStore, Display, TEXT("%-*s %10s %15s %15s %20s %20s%s"), ColumnWidth, TEXT("Container"), TEXT("Flags"), TEXT("Chunk(s) #"), TEXT("TOC (KiB)"), TEXT("Raw Size (MiB)"), TEXT("Size (MiB)"), *PatchInPlaceHeader);
 
-	const int32 SummaryBorderLength = ColumnWidth + 5 + 10 + 15 + 15 + 20 + 20 + (GeneralIoWriterSettings.bUsePatchInPlaceLayout ? 2 + 20 + 20 : 0); // +5 to account for spaces between columns
+	const int32 SummaryBorderLength = ColumnWidth + 5 + 10 + 15 + 15 + 20 + 20 + (GeneralIoWriterSettings.bUsePatchInPlaceLayout ? 3 + 20 + 20 + 20 : 0); // +5 to account for spaces between columns
 	UE_LOG(LogIoStore, Display, TEXT("%s"), *FString::ChrN(SummaryBorderLength, '-'));
 	
 	FContainerStats TotalStats;
@@ -2343,9 +2344,14 @@ void LogWriterResults(const FIoStoreWriterSettings& GeneralIoWriterSettings, con
 		FString PatchInPlaceInfo = TEXT("");
 		if (GeneralIoWriterSettings.bUsePatchInPlaceLayout)
 		{
-			PatchInPlaceInfo = FString::Printf(TEXT(" %20.2lf %20.2lf"),
-				(double)Result.PatchInPlacePinnedSize / 1024.0 / 1024.0,
-				(double)Result.PatchInPlaceContainerSizeChange / 1024.0 / 1024.0);
+			const int64 Overhead = (int64)Result.CompressedContainerSize - Result.PatchInPlaceSizeBeforeLayout;
+			const int64 DefragUnpinned = (int64)Result.PatchInPlacePinnedSizeAfterDefrag - Result.PatchInPlacePinnedSizeBeforeDefrag;
+
+			PatchInPlaceInfo = FString::Printf(TEXT(" %20.2lf %20.2lf %20.2lf"),
+				(double)Result.PatchInPlacePinnedSizeAfterDefrag / 1024.0 / 1024.0,
+				(double)DefragUnpinned / 1024.0 / 1024.0,
+				(double)Overhead / 1024.0 / 1024.0
+			);
 		}
 
 		FString ContainerSettings = FString::Printf(TEXT("%s/%s/%s/%s/%s"),
@@ -2378,8 +2384,9 @@ void LogWriterResults(const FIoStoreWriterSettings& GeneralIoWriterSettings, con
 		TotalStats.UncompressedContainerSize += Result.UncompressedContainerSize;
 		TotalStats.CompressedContainerSize += Result.CompressedContainerSize;
 		TotalStats.PaddingSize += Result.PaddingSize;
-		TotalStats.PatchInPlacePinnedSize += Result.PatchInPlacePinnedSize;
-		TotalStats.PatchInPlaceContainerSizeChange += Result.PatchInPlaceContainerSizeChange;
+		TotalStats.PatchInPlaceSizeBeforeLayout += Result.PatchInPlaceSizeBeforeLayout;
+		TotalStats.PatchInPlacePinnedSizeBeforeDefrag += Result.PatchInPlacePinnedSizeBeforeDefrag;
+		TotalStats.PatchInPlacePinnedSizeAfterDefrag += Result.PatchInPlacePinnedSizeAfterDefrag;
 	}
 
 	UE_LOG(LogIoStore, Display, TEXT("%s"), *FString::ChrN(SummaryBorderLength, '-'));
@@ -2398,9 +2405,14 @@ void LogWriterResults(const FIoStoreWriterSettings& GeneralIoWriterSettings, con
 	FString PatchInPlaceTotal = TEXT("");
 	if (GeneralIoWriterSettings.bUsePatchInPlaceLayout)
 	{
-		PatchInPlaceTotal = FString::Printf(TEXT(" %20.2lf %20.2lf"),
-			(double)TotalStats.PatchInPlacePinnedSize / 1024.0 / 1024.0,
-			(double)TotalStats.PatchInPlaceContainerSizeChange / 1024.0 / 1024.0);
+		const int64 TotalOverhead = (int64)TotalStats.CompressedContainerSize - TotalStats.PatchInPlaceSizeBeforeLayout;
+		const int64 DefragUnpinned = (int64)TotalStats.PatchInPlacePinnedSizeAfterDefrag - TotalStats.PatchInPlacePinnedSizeBeforeDefrag;
+
+		PatchInPlaceTotal = FString::Printf(TEXT(" %20.2lf %20.2lf %20.2lf"),
+			(double)TotalStats.PatchInPlacePinnedSizeAfterDefrag / 1024.0 / 1024.0,
+			(double)DefragUnpinned / 1024.0 / 1024.0,
+			(double)TotalOverhead / 1024.0 / 1024.0
+		);
 	}
 
 	UE_LOG(LogIoStore, Display, TEXT("%-*s %10s %15llu %15.2lf %20.2lf %20.2lf%s"), ColumnWidth,
@@ -2443,6 +2455,18 @@ void LogWriterResults(const FIoStoreWriterSettings& GeneralIoWriterSettings, con
 		UE_LOG(LogIoStore, Display, TEXT("%-*s %16d %16d %16d %16.2lf %16.2lf"), ColumnWidth, *Result.ContainerName, Result.TocEntryCount, Result.ModifiedChunksCount, Result.AddedChunksCount, Result.ModifiedChunksSize / 1024.0 / 1024.0, Result.AddedChunksSize / 1024.0 / 1024.0);
 	}
 
+	if (GeneralIoWriterSettings.bUsePatchInPlaceLayout && GeneralIoWriterSettings.PatchInPlaceMaximumTotalGrowth > 0)
+	{
+		const int64 TotalOverhead = (int64)TotalStats.CompressedContainerSize - TotalStats.PatchInPlaceSizeBeforeLayout;
+
+		if (TotalOverhead > (int64)GeneralIoWriterSettings.PatchInPlaceMaximumTotalGrowth)
+		{
+			UE_LOG(LogIoStore, Error, TEXT("Patch-in-place total growth %lld bytes exceeds limit %llu bytes. Packaged data is valid."),
+				TotalOverhead,
+				GeneralIoWriterSettings.PatchInPlaceMaximumTotalGrowth);
+		}
+	}
+
 	UE_LOG(LogIoStore, Display, TEXT("%s"), *FString::ChrN(PatchReportBorderLength, '-'));
 }
 
@@ -9100,6 +9124,19 @@ bool ParseContainerGenerationArguments(FIoStoreArguments& Arguments, FIoStoreWri
 
 	WriterSettings.bUsePatchInPlaceLayout = FParse::Param(FCommandLine::Get(), TEXT("UsePatchInPlaceLayout"));
 
+	if (WriterSettings.bUsePatchInPlaceLayout)
+	{
+		// default values, chosen by experimentation
+		WriterSettings.PatchInPlaceParetoFrontCoarseSamples = 50;
+		WriterSettings.PatchInPlaceParetoFrontFineSamples = 1500;
+
+		ParseSizeArgument(FCommandLine::Get(), TEXT("-patchInPlaceMaximumTotalGrowth="), WriterSettings.PatchInPlaceMaximumTotalGrowth, 0);
+		ParseSizeArgument(FCommandLine::Get(), TEXT("-patchInPlaceMaximumPerContainerGrowth="), WriterSettings.PatchInPlaceMaximumPerContainerGrowth, 0);
+		FParse::Value(FCommandLine::Get(), TEXT("-patchInPlaceParetoFrontCoarseSamples="), WriterSettings.PatchInPlaceParetoFrontCoarseSamples);
+		FParse::Value(FCommandLine::Get(), TEXT("-patchInPlaceParetoFrontFineSamples="), WriterSettings.PatchInPlaceParetoFrontFineSamples);
+		FParse::Value(FCommandLine::Get(), TEXT("-patchInPlaceParetoFrontHeuristic="), WriterSettings.PatchInPlaceParetoFrontHeuristic);
+	}
+
 	UE_LOG(LogIoStore, Display, TEXT("Using memory mapping alignment '%ld'"), WriterSettings.MemoryMappingAlignment);
 	UE_LOG(LogIoStore, Display, TEXT("Using compression block size '%ld'"), WriterSettings.CompressionBlockSize);
 	UE_LOG(LogIoStore, Display, TEXT("Using compression block alignment '%ld'"), WriterSettings.CompressionBlockAlignment);
@@ -9139,6 +9176,23 @@ bool ParseContainerGenerationArguments(FIoStoreArguments& Arguments, FIoStoreWri
 	if (WriterSettings.bUsePatchInPlaceLayout)
 	{
 		UE_LOG(LogIoStore, Display, TEXT("Using patch-in-place ordering"));
+
+		if (WriterSettings.PatchInPlaceMaximumTotalGrowth > 0)
+		{
+			UE_LOG(LogIoStore, Display, TEXT("Limiting patch-in-place total containers growth to '%llu' bytes"),
+				WriterSettings.PatchInPlaceMaximumTotalGrowth);
+		}
+
+		if (WriterSettings.PatchInPlaceMaximumPerContainerGrowth > 0)
+		{
+			UE_LOG(LogIoStore, Display, TEXT("Limiting patch-in-place per container growth to '%llu' bytes"),
+				WriterSettings.PatchInPlaceMaximumPerContainerGrowth);
+			UE_LOG(LogIoStore, Display, TEXT("Using patch-in-place Pareto front configuration %llu/%llu/%u"),
+				WriterSettings.PatchInPlaceParetoFrontCoarseSamples,
+				WriterSettings.PatchInPlaceParetoFrontFineSamples,
+				WriterSettings.PatchInPlaceParetoFrontHeuristic
+			);
+		}
 	}
 
 	FString CommandListFile;
diff --git a/Engine/Source/Developer/IoStoreUtilities/Private/IoStoreWriter.cpp b/Engine/Source/Developer/IoStoreUtilities/Private/IoStoreWriter.cpp
--- a/Engine/Source/Developer/IoStoreUtilities/Private/IoStoreWriter.cpp
+++ b/Engine/Source/Developer/IoStoreUtilities/Private/IoStoreWriter.cpp
@@ -788,17 +788,22 @@ public:
 		const FString& InSourceContainerPath,
 		const FString& InTargetContainerPath,
 		const FIoStoreTocResourceView& InSourceToc,
-		uint64 InMemoryMappingAlignment)
+		const FIoStoreWriterSettings& InSettings
+	)
 		: ReferenceContainerPath(InReferenceContainerPath)
 		, SourceContainerPath(InSourceContainerPath)
 		, TargetContainerPath(InTargetContainerPath)
 		, SourceToc(InSourceToc)
-		, MemoryMappingAlignment(InMemoryMappingAlignment)
+		, Settings(InSettings)
 		, ReferenceToc()
 	{
 	}
 
-	bool PrepareForLayout()
+	bool PrepareForLayout(
+		const uint64 SizeBeforeLayout,
+		uint64& OutPinnedSizeBeforeDefrag,
+		uint64& OutPinnedSizeAfterDefrag
+	)
 	{
 		TRACE_CPUPROFILER_EVENT_SCOPE(PrepareForLayout);
 
@@ -818,35 +823,126 @@ public:
 		PinChunksEntries();
 
 		// only do the layout if there are entries to be kept in place
-		return !EntriesToKeep.IsEmpty();
-	}
+		if (EntriesToKeep.IsEmpty())
+		{
+			return false;
+		}
 
-	void RunLayout(FIoStoreTocResource& TargetToc, int32& OutTargetPartitionNum, uint64& OutPinnedSize, uint64& OutTotalPadding)
-	{
-		TRACE_CPUPROFILER_EVENT_SCOPE(PatchInPlaceLayout);
+		OutPinnedSizeBeforeDefrag = 0;
+		for (const FTocEntry& Entry : EntriesToKeep)
+		{
+			OutPinnedSizeBeforeDefrag += Entry.Size;
+		}
+		OutPinnedSizeAfterDefrag = OutPinnedSizeBeforeDefrag;
 
-		UE_LOG(LogIoStore, Display, TEXT("Patch-in-place '%s' -> '%s'"), *ReferenceContainerPath, *TargetContainerPath);
+		const FBFDLayoutResult ResultBeforeDefrag = BestFitDecreasing(EntriesToKeep, EntriesToInsert, Settings.MemoryMappingAlignment);
+
+		// continue if we fit into growth limit
+		if (Settings.PatchInPlaceMaximumPerContainerGrowth == 0 ||
+			ResultBeforeDefrag.SizeAfter <= SizeBeforeLayout + Settings.PatchInPlaceMaximumPerContainerGrowth)
+		{
+			return true;
+		}
+
+		TMap<FIoChunkId, int32> UnpinCandidates;
+		BuildUnpinCandidates(UnpinCandidates, EntriesToKeep, EntriesToInsert, Settings.MemoryMappingAlignment, Settings.PatchInPlaceParetoFrontHeuristic);
+
+		TArray<FParetoPoint> ParetoFront;
+
+		// probe whole unpin set with coarse steps to collect initial pareto front
+		ProbeParetoFront(
+			EntriesToKeep,
+			EntriesToInsert,
+			UnpinCandidates,
+			SizeBeforeLayout + Settings.PatchInPlaceMaximumPerContainerGrowth,
+			Settings.MemoryMappingAlignment,
+			ParetoFront,
+			0,
+			UnpinCandidates.Num(),
+			Settings.PatchInPlaceParetoFrontCoarseSamples);
+
+		// refine pareto front if found, re-sample whole set otherwise
+		ProbeParetoFront(
+			EntriesToKeep,
+			EntriesToInsert,
+			UnpinCandidates,
+			SizeBeforeLayout + Settings.PatchInPlaceMaximumPerContainerGrowth,
+			Settings.MemoryMappingAlignment,
+			ParetoFront,
+			ParetoFront.Num() > 1 ? ParetoFront[0].UnpinCount : 0,
+			ParetoFront.Num() > 1 ? ParetoFront.Last().UnpinCount : UnpinCandidates.Num(),
+			Settings.PatchInPlaceParetoFrontFineSamples);
+
+		// chose smallest unpin penalty size
+		const FParetoPoint* BestPoint = Algo::MinElementBy(ParetoFront, [](const FParetoPoint& X){return X.LayoutResult.BytesWritten;});
+
+		if (!BestPoint)
+		{
+			UE_LOG(LogIoStore, Warning, TEXT("Patch-in-place '%s' container grew above limit (%llu bytes before pip -> %llu bytes after pip), failed to fit into growth limit, disabling patch-in-place"),
+				*ReferenceContainerPath,
+				SizeBeforeLayout,
+				ResultBeforeDefrag.SizeAfter);
+			return false;
+		}
+
+		// preserve EntriesToKeep order
+		EntriesToKeep.RemoveAll([&UnpinCandidates, &BestPoint, this](const FTocEntry& Entry)
+		{
+			const int32* UnpinIndex = UnpinCandidates.Find(Entry.Id);
+			if (!UnpinIndex || *UnpinIndex >= BestPoint->UnpinCount)
+			{
+				return false;
+			}
 
-		OutPinnedSize = 0;
+			EntriesToInsert.Add(FTocEntry(Entry.Id, UINT64_MAX, Entry.OriginalOffset, Entry.Size));
+			return true;
+		});
+
+		OutPinnedSizeAfterDefrag = 0;
 		for (const FTocEntry& Entry : EntriesToKeep)
 		{
-			OutPinnedSize += Entry.Size;
+			OutPinnedSizeAfterDefrag += Entry.Size;
 		}
 
-		uint64 InsertedSize = 0;
-		for (const FTocEntry& Entry : EntriesToInsert)
+		// redo BFD on the new set
+		const FBFDLayoutResult ResultAfterDefrag = BestFitDecreasing(EntriesToKeep, EntriesToInsert, Settings.MemoryMappingAlignment);
+
+		UE_LOG(LogIoStore, Display,
+			TEXT("Patch-in-place '%s' container grew above limit (%llu bytes before pip -> %llu bytes after pip), defragmented by unpinning %lli bytes, new container size %llu"),
+			*ReferenceContainerPath,
+			SizeBeforeLayout,
+			ResultBeforeDefrag.SizeAfter,
+			(int64)ResultAfterDefrag.BytesWritten - (int64)ResultBeforeDefrag.BytesWritten,
+			ResultAfterDefrag.SizeAfter
+		);
+
+		return true;
+	}
+
+	void RunLayout(
+		FIoStoreTocResource& TargetToc,
+		int32& OutTargetPartitionNum,
+		uint64& OutTotalPadding
+	)
+	{
+		TRACE_CPUPROFILER_EVENT_SCOPE(PatchInPlaceLayout);
+
+		UE_LOG(LogIoStore, Display, TEXT("Patch-in-place '%s' -> '%s'"), *ReferenceContainerPath, *TargetContainerPath);
+
+		TArray<FTocEntry> TargetEntries(MoveTemp(EntriesToKeep));
+		TargetEntries.Append(MoveTemp(EntriesToInsert));
+		Algo::SortBy(TargetEntries, [](const FTocEntry& X) {return X.Offset;});
+		if (bParanoidChecks)
 		{
-			InsertedSize += Entry.Size;
+			EnsureSortedAndNonOverlapping(TargetEntries);
 		}
 
-		int32 SourcePartitionNum = 0;
+		FitIntoPartitions(TargetEntries, SourceToc.Header.PartitionSize, OutTotalPadding);
 
 		TArray<FCopyOp> CopyOps = {};
-
-		BestFitDecreasing(EntriesToKeep, EntriesToInsert, MemoryMappingAlignment);
-		FitIntoPartitions(EntriesToKeep, SourceToc.Header.PartitionSize, OutTotalPadding);
-		BuildTargetToc(TargetToc, CopyOps, EntriesToKeep, SourcePartitionNum, OutTargetPartitionNum);
-		WriteChunksToTarget(EntriesToKeep, CopyOps, SourcePartitionNum, OutTargetPartitionNum);
+		int32 SourcePartitionNum = 0;
+		BuildTargetToc(TargetToc, CopyOps, TargetEntries, SourcePartitionNum, OutTargetPartitionNum);
+		WriteChunksToTarget(TargetEntries, CopyOps, SourcePartitionNum, OutTargetPartitionNum);
 
 		for (int32 Index = 0; Index < OutTargetPartitionNum; ++Index)
 		{
@@ -859,7 +955,7 @@ private:
 	const FString SourceContainerPath;
 	const FString TargetContainerPath;
 	const FIoStoreTocResourceView SourceToc;
-	const uint64 MemoryMappingAlignment;
+	const FIoStoreWriterSettings& Settings;
 
 	FIoStoreTocResourceView ReferenceToc;
 	FIoStoreTocResourceStorage ReferenceTocStorage;
@@ -868,8 +964,8 @@ private:
 	TMap<FIoChunkId, int32> SourceTocChunkIdToIndex;
 
 	struct FTocEntry;
-	TArray<FTocEntry> EntriesToKeep;
-	TArray<FTocEntry> EntriesToInsert;
+	TArray<FTocEntry> EntriesToKeep; // sorted by offset increasing
+	TArray<FTocEntry> EntriesToInsert; // sorted by size decreasing
 
 	static void BuildTocMap(const FIoStoreTocResourceView& Toc, TMap<FIoChunkId, int32>& OutResult)
 	{
@@ -882,12 +978,13 @@ private:
 
 	struct FTocEntry
 	{
-		FIoChunkId Id;
+		const FIoChunkId Id;
 		uint64 Offset;
-		uint64 Size;
+		const uint64 OriginalOffset;
+		const uint64 Size;
 
-		FTocEntry(FIoChunkId InId, uint64 InOffset, uint64 InSize)
-			: Id(InId), Offset(InOffset), Size(InSize)
+		FTocEntry(FIoChunkId InId, uint64 InOffset, uint64 InOriginalOffset, uint64 InSize)
+			: Id(InId), Offset(InOffset), OriginalOffset(InOriginalOffset), Size(InSize)
 		{
 		}
 
@@ -900,7 +997,7 @@ private:
 	struct FLayoutGap
 	{
 		uint64 Begin;
-		uint64 End;
+		uint64 End; // inclusive
 
 		FLayoutGap(uint64 InBegin, uint64 InEnd)
 			: Begin(InBegin), End(InEnd)
@@ -909,7 +1006,7 @@ private:
 		
 		uint64 Size() const
 		{
-			return End - Begin;
+			return End - Begin + 1;
 		}
 	};
 
@@ -947,6 +1044,38 @@ private:
 		}
 	};
 
+	struct FBFDLayoutResult
+	{
+		uint64 BytesWritten;
+		uint64 SizeAfter;
+
+		FBFDLayoutResult(uint64 InBytesWritten, uint64 InSizeAfter)
+			: BytesWritten(InBytesWritten)
+			, SizeAfter(InSizeAfter)
+		{
+		}
+	};
+
+	struct FParetoPoint
+	{
+		int32 UnpinCount; // exclusive, unpin all elements in unpin array [0, UnpinCount)
+		FBFDLayoutResult LayoutResult;
+
+		FParetoPoint(int32 InUnpinCount, FBFDLayoutResult InLayouResult)
+			: UnpinCount(InUnpinCount)
+			, LayoutResult(InLayouResult)
+		{
+		}
+		
+		bool Dominates(const FParetoPoint& Other) const
+		{
+			return (LayoutResult.BytesWritten < Other.LayoutResult.BytesWritten &&
+					LayoutResult.SizeAfter <= Other.LayoutResult.SizeAfter) ||
+					(LayoutResult.BytesWritten <= Other.LayoutResult.BytesWritten &&
+					LayoutResult.SizeAfter < Other.LayoutResult.SizeAfter);
+		}
+	};
+
 	void PinChunksEntries()
 	{
 		for (int32 ChunkIndex = 0; ChunkIndex < SourceToc.ChunkIds.Num(); ChunkIndex++)
@@ -956,7 +1085,7 @@ private:
 			const int32* ReferenceTocIndex = ReferenceTocChunkIdToIndex.Find(SourceChunk.Id);
 			if (!ReferenceTocIndex)
 			{
-				EntriesToInsert.Add(FTocEntry(SourceChunk.Id, UINT64_MAX, SourceChunk.SizeOnDisk));
+				EntriesToInsert.Add(FTocEntry(SourceChunk.Id, UINT64_MAX, SourceChunk.OffsetOnDisk, SourceChunk.SizeOnDisk));
 				continue;
 			}
 
@@ -964,13 +1093,16 @@ private:
 			if (SourceChunk.ChunkHash == ReferenceChunk.ChunkHash &&
 				SourceChunk.SizeOnDisk == ReferenceChunk.SizeOnDisk)
 			{
-				EntriesToKeep.Add(FTocEntry(SourceChunk.Id, ReferenceChunk.OffsetOnDisk, SourceChunk.SizeOnDisk));
+				EntriesToKeep.Add(FTocEntry(SourceChunk.Id, ReferenceChunk.OffsetOnDisk, SourceChunk.OffsetOnDisk, SourceChunk.SizeOnDisk));
 			}
 			else
 			{
-				EntriesToInsert.Add(FTocEntry(SourceChunk.Id, UINT64_MAX, SourceChunk.SizeOnDisk));
+				EntriesToInsert.Add(FTocEntry(SourceChunk.Id, UINT64_MAX, SourceChunk.OffsetOnDisk, SourceChunk.SizeOnDisk));
 			}
 		}
+
+		Algo::SortBy(EntriesToKeep, [](const FTocEntry& X){return X.Offset;});
+		Algo::SortBy(EntriesToInsert, [](const FTocEntry& X){return X.Size;}, TGreater<uint64>());
 	}
 
 	static uint64 GapsBinarySearchProjection(const FLayoutGap& Gap)
@@ -987,7 +1119,7 @@ private:
 		}
 		return Result;
 	}
-	
+
 	static void EnsureSortedAndNonOverlapping(TConstArrayView<FTocEntry> Entries)
 	{
 		for (int32 Index = 1; Index < Entries.Num(); ++Index)
@@ -1015,29 +1147,117 @@ private:
 		}
 	}
 
-	static void BestFitDecreasing(TArray<FTocEntry>& EntriesToKeep, TArray<FTocEntry>& EntriesToInsert, const uint64 MemoryMappingAlignment)
+	static void BuildUnpinCandidates(
+		TMap<FIoChunkId, int32>& OutUnpinCandidates, // chunk id -> index of all unpin candidates 
+		const TConstArrayView<FTocEntry>& EntriesToKeep,
+		const TConstArrayView<FTocEntry>& EntriesToInsert,
+		const uint64 MemoryMappingAlignment,
+		const uint8 ParetoFrontHeuristic
+	)
 	{
-		TRACE_CPUPROFILER_EVENT_SCOPE(BestFitDecreasing);
+		struct FUnpinCandidate
+		{
+			FIoChunkId Id;
+			double Contribution;
+		};
 
-		Algo::SortBy(EntriesToKeep, [](const FTocEntry& X){return X.Offset;});
-		Algo::SortBy(EntriesToInsert, [](const FTocEntry& X){return X.Size;}, TGreater<uint64>());
+		TArray<FUnpinCandidate> UnpinCandidates;
+		UnpinCandidates.Reserve(EntriesToKeep.Num());
+
+		for (int32 i = 0; i < EntriesToKeep.Num(); ++i)
+		{
+			const FTocEntry& Entry = EntriesToKeep[i];
+
+			// [previous entry] [gap before] [current entry] [gap after] [next entry]
+			const uint64 PrevEnd = i > 0 ? AlignOffset(EntriesToKeep[i - 1].End(), MemoryMappingAlignment) : 0;
+			const uint64 CurrBegin = Entry.Offset;
+			const uint64 CurrEnd = Entry.End(); // TODO do we want to align this?
+			const uint64 NextBegin = i + 1 < EntriesToKeep.Num() ? EntriesToKeep[i + 1].Offset : CurrEnd;
+
+			const uint64 GapBefore = CurrBegin > PrevEnd ? CurrBegin - PrevEnd : 0;
+			const uint64 GapAfter  = NextBegin > CurrEnd ? NextBegin - CurrEnd : 0;
+			const uint64 EntrySize = Entry.Size;
+
+			double Contribution;
+			switch (ParetoFrontHeuristic)
+			{
+			default:
+			case 0: // Prioritize based on how "useful" the gap for incoming dataset
+				{
+					uint64 Offset = PrevEnd;
+					uint64 UsableSpace = 0;
+
+					for (const FTocEntry& Insert : EntriesToInsert)
+					{
+						if (Offset + Insert.Size >= NextBegin)
+						{
+							continue;
+						}
+
+						UsableSpace += Insert.Size;
+						Offset += AlignOffset(Offset + Insert.Size, MemoryMappingAlignment);
+
+						if (Offset >= NextBegin)
+						{
+							break;
+						}
+					}
+
+					Contribution = (double)UsableSpace / (double)EntrySize;
+					break;
+				}
+			case 1: // Prioritize most gap space regardless of entry size
+				{
+					Contribution = GapBefore + GapAfter;
+					break;
+				}
+			case 2: // Prioritize freed up space per byte written
+				{
+					Contribution = (double)(GapBefore + GapAfter) / (double)EntrySize;
+					break;
+				}
+			}
+
+			UnpinCandidates.Add(FUnpinCandidate(Entry.Id, Contribution));
+		}
+
+		Algo::SortBy(UnpinCandidates, [](const FUnpinCandidate& X){return X.Contribution;}, TGreater<double>());
+
+		OutUnpinCandidates.Reserve(UnpinCandidates.Num());
+		for (int32 i = 0; i < UnpinCandidates.Num(); i++)
+		{
+			OutUnpinCandidates.Add(UnpinCandidates[i].Id, i);
+		}
+	}
+
+	static FBFDLayoutResult BestFitDecreasing(
+		const TConstArrayView<FTocEntry>& EntriesToKeep,
+		TArray<FTocEntry>& EntriesToInsert,
+		const uint64 MemoryMappingAlignment)
+	{
+		TRACE_CPUPROFILER_EVENT_SCOPE(BestFitDecreasing);
 
 		if (bParanoidChecks)
 		{
 			EnsureSortedAndNonOverlapping(EntriesToKeep);
+
+			for (int32 Index = 1; Index < EntriesToInsert.Num(); ++Index)
+			{
+				verify(EntriesToInsert[Index].Size <= EntriesToInsert[Index - 1].Size);
+			}
 		}
 
 		TArray<FLayoutGap> Gaps;
 		if (!EntriesToKeep.IsEmpty() && EntriesToKeep[0].Offset > 0)
 		{
-			Gaps.Add(FLayoutGap(0, EntriesToKeep[0].Offset));
+			Gaps.Add(FLayoutGap(0, EntriesToKeep[0].Offset - 1));
 		}
-		
+
 		for (int32 Index = 1; Index < EntriesToKeep.Num(); ++Index)
 		{
 			const uint64 PrevEnd = AlignOffset(EntriesToKeep[Index - 1].End(), MemoryMappingAlignment);
 			const uint64 CurrentBegin = EntriesToKeep[Index].Offset;
-			if (CurrentBegin > PrevEnd + 1)
+			if (CurrentBegin >= PrevEnd + 1)
 			{
 				Gaps.Add(FLayoutGap(PrevEnd, CurrentBegin - 1));
 			}
@@ -1046,11 +1266,15 @@ private:
 		Algo::SortBy(Gaps, GapsBinarySearchProjection);
 
 		uint64 AppendOffset = !EntriesToKeep.IsEmpty() ? EntriesToKeep.Last().End() : 0;
+		uint64 BytesWritten = 0;
 
 		// TODO support compression block alignment?
+		// TODO try place entries together via Entry.OriginalOffset to preserve parts of original layout
 		for (FTocEntry& Entry : EntriesToInsert)
 		{
-			const int32 BestGapIndex = Algo::UpperBoundBy(Gaps, Entry.Size, GapsBinarySearchProjection);
+			BytesWritten += Entry.Size;
+
+			const int32 BestGapIndex = Algo::LowerBoundBy(Gaps, Entry.Size, GapsBinarySearchProjection);
 
 			// TODO make it faster
 			if (BestGapIndex != Gaps.Num())
@@ -1066,7 +1290,7 @@ private:
 				Entry.Offset = Gap.Begin;
 				Gap.Begin = AlignOffset(Gap.Begin + Entry.Size, MemoryMappingAlignment);
 
-				if (Gap.Begin < Gap.End)
+				if (Gap.Begin <= Gap.End)
 				{
 					Gaps.Insert(Gap, Algo::UpperBoundBy(Gaps, Gap.Size(), GapsBinarySearchProjection));
 
@@ -1083,14 +1307,89 @@ private:
 			}
 		}
 
-		EntriesToKeep.Append(MoveTemp(EntriesToInsert));
+		return FBFDLayoutResult(BytesWritten, AppendOffset);
+	}
 
-		Algo::SortBy(EntriesToKeep, [](const FTocEntry& X) {return X.Offset;});
+	// Simulates BFD by unpinning elements starting from [0, FromUnpinCount to [0, ToUnpinCount]
+	// with sample count provided, collects pareto front from samples and sorts by unpin count.
+	static void ProbeParetoFront(
+		const TConstArrayView<FTocEntry>& EntriesToKeep,
+		const TConstArrayView<FTocEntry>& EntriesToInsert,
+		const TMap<FIoChunkId, int32>& UnpinCandidates,
+		const uint64 MaxContainerSize,
+		const uint64 MemoryMappingAlignment,
+		TArray<FParetoPoint>& ParetoFront,
+		const int32 FromUnpinCount, // inclusive
+		const int32 ToUnpinCount, // exclusive
+		const int32 WithSampleCount
+	)
+	{
+		const int32 IterationCount = FMath::DivideAndRoundUp(ToUnpinCount - FromUnpinCount, WithSampleCount);
+		if (IterationCount == 0)
+		{
+			return;
+		}
 
-		if (bParanoidChecks)
+		TArray<FParetoPoint> Results;
+		Results.SetNumUninitialized(IterationCount);
+
+		ParallelFor(IterationCount, [&](int32 Index)
 		{
-			EnsureSortedAndNonOverlapping(EntriesToKeep);
+			// ensure ToUnpinCount is sampled
+			const int32 UnpinCount = (Index + 1 == IterationCount) ? ToUnpinCount : FMath::Lerp(FromUnpinCount, ToUnpinCount, Index / (double)(IterationCount - 1));
+
+			TArray<FTocEntry> EntriesToInsertUnpinned;
+			EntriesToInsertUnpinned.Reserve(EntriesToInsert.Num() + UnpinCount);
+			EntriesToInsertUnpinned.Append(EntriesToInsert);
+			TArray<FTocEntry> EntriesToKeepUnpinned = EntriesToKeep.FilterByPredicate([&](const FTocEntry& Entry)
+			{
+				const int32* UnpinIndex = UnpinCandidates.Find(Entry.Id);
+				if (!UnpinIndex || *UnpinIndex >= UnpinCount)
+				{
+					return true;
+				}
+
+				EntriesToInsertUnpinned.Add(FTocEntry(Entry.Id, UINT64_MAX, Entry.OriginalOffset, Entry.Size));
+				return false;
+			});
+
+			Algo::SortBy(EntriesToInsertUnpinned, [](const FTocEntry& X){return X.Size;}, TGreater<uint64>());
+
+			const FBFDLayoutResult Result = BestFitDecreasing(EntriesToKeepUnpinned, EntriesToInsertUnpinned, MemoryMappingAlignment);
+
+			Results[Index] = FParetoPoint(UnpinCount, Result);
+		});
+
+		for (const FParetoPoint& Point : Results)
+		{
+			if (Point.LayoutResult.SizeAfter <= MaxContainerSize)
+			{
+				AddParetoPoint(ParetoFront, Point);
+			}
 		}
+
+		Algo::SortBy(ParetoFront, [](const FParetoPoint& X){return X.UnpinCount;});
+	}
+
+	static bool AddParetoPoint(TArray<FParetoPoint>& ParetoFront, const FParetoPoint& NewPoint)
+	{
+		// Check if any points dominates the new point
+		for (const FParetoPoint& CurrPoint : ParetoFront)
+		{
+			if (CurrPoint.Dominates(NewPoint))
+			{
+				return false;
+			}
+		}
+
+		// Remove all points dominated by the new point
+		ParetoFront.RemoveAll([&NewPoint](const FParetoPoint& CurrPoint)
+		{
+			return NewPoint.Dominates(CurrPoint);
+		});
+
+		ParetoFront.Add(NewPoint);
+		return true;
 	}
 
 	static void FitIntoPartitions(TArray<FTocEntry>& Chunks, const uint64 PartitionSize, uint64& OutTotalPadding)
@@ -1334,7 +1633,7 @@ private:
 
 					Source->ReadAt(TempBuffer.GetData(), Size, SourceOffset + Offset);
 					Source->ReadAt(TempBuffer2.GetData(), Size, TargetOffset + Offset);
-					
+
 					verify(FMemory::Memcmp(TempBuffer.GetData(), TempBuffer2.GetData(), Size) == 0);
 				}
 			}
@@ -1894,6 +2193,15 @@ public:
 
 	void DoPatchInPlaceLayout()
 	{
+		PatchInPlaceSizeBeforeLayout = 0;
+		for (FPartition& Partition : Partitions)
+		{
+			PatchInPlaceSizeBeforeLayout += Partition.Offset;
+		}
+
+		PatchInPlacePinnedSizeBeforeDefrag = 0;
+		PatchInPlacePinnedSizeAfterDefrag = 0;
+
 		if (!ReferenceChunkDatabase.IsValid())
 		{
 			UE_LOG(LogIoStore, Warning, TEXT("Trying to use patch-in-place without valid reference database"));
@@ -1908,7 +2216,6 @@ public:
 		}
 
 		const FIoStoreWriterSettings& WriterSettings = WriterContext->GetSettings();
-
 		if (WriterSettings.CompressionBlockAlignment)
 		{
 			UE_LOG(LogIoStore, Warning, TEXT("Ignoring CompressionBlockAlignment in for container '%s' due to patch-in-place enabled"), *ContainerPathAndBaseFileName);
@@ -1927,40 +2234,30 @@ public:
 			TocFilePath,
 			ContainerPathAndBaseFileName + TEXT("_patchinplace.utoc"),
 			TocBuilder.GetTocResource().AsView(),
-			WriterSettings.MemoryMappingAlignment
+			WriterSettings
 		);
 
-		if (!Layout.PrepareForLayout())
+		if (!Layout.PrepareForLayout(PatchInPlaceSizeBeforeLayout, PatchInPlacePinnedSizeBeforeDefrag, PatchInPlacePinnedSizeAfterDefrag))
 		{
 			return;
 		}
 
-		int64 ContainerSizeBeforePatchInPlace = 0;
-		for (FPartition& Partition : Partitions)
-		{
-			ContainerSizeBeforePatchInPlace += Partition.Offset;
-		}
-
 		// release file handles for patch-in-place
 		Partitions.Empty();
 
-		Layout.RunLayout(TargetToc, TargetPartitionNum, PatchInPlacePinnedSize, TotalPaddingSize);
+		Layout.RunLayout(TargetToc, TargetPartitionNum, TotalPaddingSize);
 
 		SourceToc = MoveTemp(TargetToc);
 
 		// recreate file handles
-		int64 ContainerSizeAfterPatchInPlace = 0;
 		for (int32 i = 0; i < TargetPartitionNum; ++i)
 		{
 			FPartition& Partition = Partitions.AddDefaulted_GetRef();
 			Partition.Index = Partitions.Num() - 1;
 			CreatePartitionContainerFile(Partition, true);
 			Partition.Offset = Partition.ContainerFileHandle->Tell();
-			ContainerSizeAfterPatchInPlace += Partition.Offset;
 		}
 
-		PatchInPlaceContainerSizeChange = ContainerSizeAfterPatchInPlace - ContainerSizeBeforePatchInPlace;
-
 		if (WriterContext->GetSettings().bEnableFileRegions)
 		{
 			int32 PreviousPartitionIndex = -1;
@@ -2122,8 +2419,9 @@ public:
 		Result.CompressionMethod = EnumHasAnyFlags(ContainerSettings.ContainerFlags, EIoContainerFlags::Compressed)
 			? WriterSettings.CompressionMethod
 			: NAME_None;
-		Result.PatchInPlacePinnedSize = PatchInPlacePinnedSize;
-		Result.PatchInPlaceContainerSizeChange = PatchInPlaceContainerSizeChange;
+		Result.PatchInPlaceSizeBeforeLayout = PatchInPlaceSizeBeforeLayout;
+		Result.PatchInPlacePinnedSizeBeforeDefrag = PatchInPlacePinnedSizeBeforeDefrag;
+		Result.PatchInPlacePinnedSizeAfterDefrag = PatchInPlacePinnedSizeAfterDefrag;
 		Result.ModifiedChunksCount = 0;
 		Result.AddedChunksCount = 0;
 		Result.ModifiedChunksSize= 0;
@@ -2913,8 +3211,9 @@ private:
 	uint64						TotalPaddingSize = 0;
 	uint64						UncompressedContainerSize = 0; // this is the size the container would be if it were uncompressed.
 	uint64						CompressedContainerSize = 0; // this is the size of the container with the given compression (which may be none).
-	uint64						PatchInPlacePinnedSize = 0;
-	int64						PatchInPlaceContainerSizeChange = 0;
+	uint64						PatchInPlaceSizeBeforeLayout = 0;
+	uint64						PatchInPlacePinnedSizeBeforeDefrag = 0;
+	uint64						PatchInPlacePinnedSizeAfterDefrag = 0;
 	uint64						MaxPartitionSize = 0;
 	int32						CurrentPartitionIndex = 0;
 	bool						bHasMemoryMappedEntry = false;
diff --git a/Engine/Source/Developer/IoStoreUtilities/Private/IoStoreWriter.h b/Engine/Source/Developer/IoStoreUtilities/Private/IoStoreWriter.h
--- a/Engine/Source/Developer/IoStoreUtilities/Private/IoStoreWriter.h
+++ b/Engine/Source/Developer/IoStoreUtilities/Private/IoStoreWriter.h
@@ -22,6 +22,11 @@ struct FIoStoreWriterSettings
 	int32 CompressionMinSizeToConsiderDDC = 0;
 	uint64 MemoryMappingAlignment = 0;
 	uint64 MaxPartitionSize = 0;
+	uint64 PatchInPlaceMaximumTotalGrowth = 0;        // fail build if total overhead goes over this value and bUsePatchInPlaceLayout is enabled, 0 to ignore
+	uint64 PatchInPlaceMaximumPerContainerGrowth = 0; // defragment container if it's overhead goes above this value by moving some entries, 0 to ignore
+	uint64 PatchInPlaceParetoFrontCoarseSamples = 0;  // amount of coarse samples to use for pareto front probing during defrag
+	uint64 PatchInPlaceParetoFrontFineSamples = 0;    // amount of fine samples to fine tune pareto front during defrag
+	uint8 PatchInPlaceParetoFrontHeuristic = 0;       // heuristic to choose for pareto front, see FPatchInPlaceLayout::BuildUnpinCandidates
 	TMap<FString, uint64> MaxPartitionSizeOverride;
 	bool bEnableFileRegions = false;
 	bool bContainerMetaEnabled = false;
@@ -43,8 +48,9 @@ struct FIoStoreWriterResult
 	int64 DirectoryIndexSize = 0;
 	uint64 TotalEntryCompressedSize = 0; // sum of the compressed size of entries excluding encryption alignment.
 	uint64 ReferenceCacheMissBytes = 0; // number of compressed bytes excluding alignment that could have been from refcache but weren't.
-	uint64 PatchInPlacePinnedSize = 0; // compressed size of all entries which offset is matching reference db offset.
-	int64 PatchInPlaceContainerSizeChange = 0; // difference of value of CompressedContainerSize before and after patch-in-place layout.
+	uint64 PatchInPlaceSizeBeforeLayout = 0;
+	uint64 PatchInPlacePinnedSizeBeforeDefrag = 0;
+	uint64 PatchInPlacePinnedSizeAfterDefrag = 0;
 	uint64 AddedChunksCount = 0;
 	uint64 AddedChunksSize = 0;
 	uint64 ModifiedChunksCount = 0;
