Patch-in-place takes into account memory alignment requirement for each entry.
Bump patch-in-place limits to 256mb per container / 1024mb total.
Misc quality of life improvements for patch-in-place.

#redirectsaresafe

#Synced-CL 51084377

diff --git a/Engine/Source/Developer/IoStoreUtilities/Private/IoStoreUtilities.cpp b/Engine/Source/Developer/IoStoreUtilities/Private/IoStoreUtilities.cpp
--- a/Engine/Source/Developer/IoStoreUtilities/Private/IoStoreUtilities.cpp
+++ b/Engine/Source/Developer/IoStoreUtilities/Private/IoStoreUtilities.cpp
@@ -2461,9 +2461,9 @@ void LogWriterResults(const FIoStoreWriterSettings& GeneralIoWriterSettings, con
 
 		if (TotalOverhead > (int64)GeneralIoWriterSettings.PatchInPlaceMaximumTotalGrowth)
 		{
-			UE_LOG(LogIoStore, Error, TEXT("Patch-in-place total growth %lld bytes exceeds limit %llu bytes. Packaged data is valid."),
-				TotalOverhead,
-				GeneralIoWriterSettings.PatchInPlaceMaximumTotalGrowth);
+			UE_LOG(LogIoStore, Warning, TEXT("Patch-in-place total growth %.2f MiB exceeds limit %.2f MiB. Packaged data is valid."),
+				(double)TotalOverhead / 1024.0 / 1024.0,
+				(double)GeneralIoWriterSettings.PatchInPlaceMaximumTotalGrowth / 1024.0 / 1024.0);
 		}
 	}
 
@@ -9193,6 +9193,11 @@ bool ParseContainerGenerationArguments(FIoStoreArguments& Arguments, FIoStoreWri
 				WriterSettings.PatchInPlaceParetoFrontHeuristic
 			);
 		}
+
+		if (WriterSettings.CompressionBlockAlignment)
+		{
+			UE_LOG(LogIoStore, Warning, TEXT("Ignoring CompressionBlockAlignment due to patch-in-place enabled"));
+		}
 	}
 
 	FString CommandListFile;
diff --git a/Engine/Source/Developer/IoStoreUtilities/Private/IoStoreWriter.cpp b/Engine/Source/Developer/IoStoreUtilities/Private/IoStoreWriter.cpp
--- a/Engine/Source/Developer/IoStoreUtilities/Private/IoStoreWriter.cpp
+++ b/Engine/Source/Developer/IoStoreUtilities/Private/IoStoreWriter.cpp
@@ -847,6 +847,25 @@ public:
 		TMap<FIoChunkId, int32> UnpinCandidates;
 		BuildUnpinCandidates(UnpinCandidates, EntriesToKeep, EntriesToInsert, Settings.MemoryMappingAlignment, Settings.PatchInPlaceParetoFrontHeuristic);
 
+		if (bParanoidChecks)
+		{
+			// unpinning all should result in valid pareto front that doesn't go over existing size
+			TArray<FParetoPoint> ParetoFrontTemp;
+			ProbeParetoFront(
+				EntriesToKeep,
+				EntriesToInsert,
+				UnpinCandidates,
+				SizeBeforeLayout,
+				Settings.MemoryMappingAlignment,
+				ParetoFrontTemp,
+				UnpinCandidates.Num() - 1,
+				UnpinCandidates.Num(),
+				1);
+			ensure(!ParetoFrontTemp.IsEmpty());
+			ensure(ParetoFrontTemp.Last().UnpinCount == UnpinCandidates.Num());
+			ensure(ParetoFrontTemp.Last().LayoutResult.SizeAfter <= SizeBeforeLayout);
+		}
+
 		TArray<FParetoPoint> ParetoFront;
 
 		// probe whole unpin set with coarse steps to collect initial pareto front
@@ -878,10 +897,12 @@ public:
 
 		if (!BestPoint)
 		{
-			UE_LOG(LogIoStore, Warning, TEXT("Patch-in-place '%s' container grew above limit (%llu bytes before pip -> %llu bytes after pip), failed to fit into growth limit, disabling patch-in-place"),
+			OutPinnedSizeBeforeDefrag = 0;
+			OutPinnedSizeAfterDefrag = 0;
+			UE_LOG(LogIoStore, Warning, TEXT("Patch-in-place '%s' container grew above limit (%.2f MiB before pip -> %.2f MiB after pip), failed to fit into growth limit, disabling patch-in-place"),
 				*ReferenceContainerPath,
-				SizeBeforeLayout,
-				ResultBeforeDefrag.SizeAfter);
+				(double)SizeBeforeLayout / 1024.0 / 1024.0,
+				(double)ResultBeforeDefrag.SizeAfter / 1024.0 / 1024.0);
 			return false;
 		}
 
@@ -894,10 +915,12 @@ public:
 				return false;
 			}
 
-			EntriesToInsert.Add(FTocEntry(Entry.Id, UINT64_MAX, Entry.OriginalOffset, Entry.Size));
+			EntriesToInsert.Add(Entry.Unpin());
 			return true;
 		});
 
+		Algo::SortBy(EntriesToInsert, [](const FTocEntry& X){return X.Size;}, TGreater<uint64>());
+
 		OutPinnedSizeAfterDefrag = 0;
 		for (const FTocEntry& Entry : EntriesToKeep)
 		{
@@ -908,12 +931,12 @@ public:
 		const FBFDLayoutResult ResultAfterDefrag = BestFitDecreasing(EntriesToKeep, EntriesToInsert, Settings.MemoryMappingAlignment);
 
 		UE_LOG(LogIoStore, Display,
-			TEXT("Patch-in-place '%s' container grew above limit (%llu bytes before pip -> %llu bytes after pip), defragmented by unpinning %lli bytes, new container size %llu"),
+			TEXT("Patch-in-place '%s' container grew above limit (%.2f MiB before pip -> %.2f MiB after pip), defragmented by unpinning %.2f MiB, new container size %.2f MiB"),
 			*ReferenceContainerPath,
-			SizeBeforeLayout,
-			ResultBeforeDefrag.SizeAfter,
-			(int64)ResultAfterDefrag.BytesWritten - (int64)ResultBeforeDefrag.BytesWritten,
-			ResultAfterDefrag.SizeAfter
+			(double)SizeBeforeLayout / 1024.0 / 1024.0,
+			(double)ResultBeforeDefrag.SizeAfter / 1024.0 / 1024.0,
+			(double)((int64)ResultAfterDefrag.BytesWritten - (int64)ResultBeforeDefrag.BytesWritten) / 1024.0 / 1024.0,
+			(double)ResultAfterDefrag.SizeAfter / 1024.0 / 1024.0
 		);
 
 		return true;
@@ -982,9 +1005,10 @@ private:
 		uint64 Offset;
 		const uint64 OriginalOffset;
 		const uint64 Size;
+		const bool bMemoryMapped;
 
-		FTocEntry(FIoChunkId InId, uint64 InOffset, uint64 InOriginalOffset, uint64 InSize)
-			: Id(InId), Offset(InOffset), OriginalOffset(InOriginalOffset), Size(InSize)
+		FTocEntry(FIoChunkId InId, uint64 InOffset, uint64 InOriginalOffset, uint64 InSize, const bool bInMemoryMapped)
+			: Id(InId), Offset(InOffset), OriginalOffset(InOriginalOffset), Size(InSize), bMemoryMapped(bInMemoryMapped)
 		{
 		}
 
@@ -992,21 +1016,50 @@ private:
 		{
 			return Offset + Size;
 		}
+
+		FTocEntry Unpin() const
+		{
+			return FTocEntry(Id, UINT64_MAX, OriginalOffset, Size, bMemoryMapped);
+		}
+
+		static FTocEntry Unpinned(const FIoStoreTocChunkInfo& ChunkInfo)
+		{
+			return FTocEntry(ChunkInfo.Id, UINT64_MAX, ChunkInfo.OffsetOnDisk, ChunkInfo.SizeOnDisk, ChunkInfo.bIsMemoryMapped);
+		}
+
+		static FTocEntry Pinned(const FIoStoreTocChunkInfo& ChunkInfo, const uint64 AtOffset)
+		{
+			return FTocEntry(ChunkInfo.Id, AtOffset, ChunkInfo.OffsetOnDisk, ChunkInfo.SizeOnDisk, ChunkInfo.bIsMemoryMapped);
+		}
 	};
 
 	struct FLayoutGap
 	{
-		uint64 Begin;
-		uint64 End; // inclusive
+		uint64 Begin; // inclusive, unaligned
+		uint64 End; // exclusive
+		bool bTail; // gap is in the tail after all pinned entries
 
-		FLayoutGap(uint64 InBegin, uint64 InEnd)
-			: Begin(InBegin), End(InEnd)
+		FLayoutGap(uint64 InBegin, uint64 InEnd, bool bInTail)
+			: Begin(InBegin), End(InEnd), bTail(bInTail)
 		{
 		}
-		
+
 		uint64 Size() const
 		{
-			return End - Begin + 1;
+			return End >= Begin ? End - Begin : 0;
+		}
+
+		bool Fits(const FTocEntry& Entry, const uint64 MemoryMappingAlignment) const
+		{
+			const uint64 NewBegin = AlignOffset(Begin, Entry.bMemoryMapped, MemoryMappingAlignment);
+			return NewBegin + Entry.Size <= End;
+		}
+
+		uint64 Place(const FTocEntry& Entry, const uint64 MemoryMappingAlignment)
+		{
+			const uint64 NewBegin = AlignOffset(Begin, Entry.bMemoryMapped, MemoryMappingAlignment);
+			Begin = NewBegin + Entry.Size;
+			return NewBegin;
 		}
 	};
 
@@ -1085,7 +1138,7 @@ private:
 			const int32* ReferenceTocIndex = ReferenceTocChunkIdToIndex.Find(SourceChunk.Id);
 			if (!ReferenceTocIndex)
 			{
-				EntriesToInsert.Add(FTocEntry(SourceChunk.Id, UINT64_MAX, SourceChunk.OffsetOnDisk, SourceChunk.SizeOnDisk));
+				EntriesToInsert.Add(FTocEntry::Unpinned(SourceChunk));
 				continue;
 			}
 
@@ -1093,11 +1146,11 @@ private:
 			if (SourceChunk.ChunkHash == ReferenceChunk.ChunkHash &&
 				SourceChunk.SizeOnDisk == ReferenceChunk.SizeOnDisk)
 			{
-				EntriesToKeep.Add(FTocEntry(SourceChunk.Id, ReferenceChunk.OffsetOnDisk, SourceChunk.OffsetOnDisk, SourceChunk.SizeOnDisk));
+				EntriesToKeep.Add(FTocEntry::Pinned(SourceChunk, ReferenceChunk.OffsetOnDisk));
 			}
 			else
 			{
-				EntriesToInsert.Add(FTocEntry(SourceChunk.Id, UINT64_MAX, SourceChunk.OffsetOnDisk, SourceChunk.SizeOnDisk));
+				EntriesToInsert.Add(FTocEntry::Unpinned(SourceChunk));
 			}
 		}
 
@@ -1110,10 +1163,10 @@ private:
 		return Gap.Size();
 	}
 
-	static uint64 AlignOffset(const uint64 Offset, const uint64 MemoryMappingAlignment)
+	static uint64 AlignOffset(const uint64 Offset, const bool bMemoryMapped, const uint64 MemoryMappingAlignment)
 	{
 		uint64 Result = Offset;
-		if (MemoryMappingAlignment > 0)
+		if (bMemoryMapped && MemoryMappingAlignment > 0)
 		{
 			Result = Align(Result, MemoryMappingAlignment);
 		}
@@ -1169,9 +1222,9 @@ private:
 			const FTocEntry& Entry = EntriesToKeep[i];
 
 			// [previous entry] [gap before] [current entry] [gap after] [next entry]
-			const uint64 PrevEnd = i > 0 ? AlignOffset(EntriesToKeep[i - 1].End(), MemoryMappingAlignment) : 0;
+			const uint64 PrevEnd = i > 0 ? EntriesToKeep[i - 1].End() : 0; // don't align
 			const uint64 CurrBegin = Entry.Offset;
-			const uint64 CurrEnd = Entry.End(); // TODO do we want to align this?
+			const uint64 CurrEnd = Entry.End(); // don't align
 			const uint64 NextBegin = i + 1 < EntriesToKeep.Num() ? EntriesToKeep[i + 1].Offset : CurrEnd;
 
 			const uint64 GapBefore = CurrBegin > PrevEnd ? CurrBegin - PrevEnd : 0;
@@ -1184,20 +1237,20 @@ private:
 			default:
 			case 0: // Prioritize based on how "useful" the gap for incoming dataset
 				{
-					uint64 Offset = PrevEnd;
+					FLayoutGap Gap(PrevEnd, NextBegin, false);
 					uint64 UsableSpace = 0;
 
 					for (const FTocEntry& Insert : EntriesToInsert)
 					{
-						if (Offset + Insert.Size >= NextBegin)
+						if (!Gap.Fits(Insert, MemoryMappingAlignment))
 						{
 							continue;
 						}
 
 						UsableSpace += Insert.Size;
-						Offset += AlignOffset(Offset + Insert.Size, MemoryMappingAlignment);
+						Gap.Place(Insert, MemoryMappingAlignment);
 
-						if (Offset >= NextBegin)
+						if (Gap.Size() == 0)
 						{
 							break;
 						}
@@ -1248,24 +1301,19 @@ private:
 		}
 
 		TArray<FLayoutGap> Gaps;
-		if (!EntriesToKeep.IsEmpty() && EntriesToKeep[0].Offset > 0)
-		{
-			Gaps.Add(FLayoutGap(0, EntriesToKeep[0].Offset - 1));
-		}
-
-		for (int32 Index = 1; Index < EntriesToKeep.Num(); ++Index)
+		for (int32 Index = 0; Index <= EntriesToKeep.Num(); ++Index)
 		{
-			const uint64 PrevEnd = AlignOffset(EntriesToKeep[Index - 1].End(), MemoryMappingAlignment);
-			const uint64 CurrentBegin = EntriesToKeep[Index].Offset;
-			if (CurrentBegin >= PrevEnd + 1)
+			// don't align offsets here, will be aligned later
+			const uint64 PrevEnd = Index > 0 ? EntriesToKeep[Index - 1].End() : 0;
+			const uint64 CurrentBegin = Index < EntriesToKeep.Num() ? EntriesToKeep[Index].Offset : UINT64_MAX;
+			if (CurrentBegin > PrevEnd)
 			{
-				Gaps.Add(FLayoutGap(PrevEnd, CurrentBegin - 1));
+				Gaps.Add(FLayoutGap(PrevEnd, CurrentBegin, Index == EntriesToKeep.Num()));
 			}
 		}
 
 		Algo::SortBy(Gaps, GapsBinarySearchProjection);
 
-		uint64 AppendOffset = !EntriesToKeep.IsEmpty() ? EntriesToKeep.Last().End() : 0;
 		uint64 BytesWritten = 0;
 
 		// TODO support compression block alignment?
@@ -1274,40 +1322,44 @@ private:
 		{
 			BytesWritten += Entry.Size;
 
-			const int32 BestGapIndex = Algo::LowerBoundBy(Gaps, Entry.Size, GapsBinarySearchProjection);
+			// gaps are sorted by size, find first that fits by size
+			const int32 StartIndex = Algo::LowerBoundBy(Gaps, Entry.Size, GapsBinarySearchProjection);
+			verify(StartIndex < Gaps.Num());
 
-			// TODO make it faster
-			if (BestGapIndex != Gaps.Num())
+			// then find one that actually fits given alignment requirement
+			const int32 RelativeIndex = TConstArrayView<FLayoutGap>(Gaps)
+					.RightChop(StartIndex)
+					.IndexOfByPredicate([&](const FLayoutGap& Gap)
 			{
-				FLayoutGap Gap = Gaps[BestGapIndex];
-				if (bParanoidChecks)
-				{
-					verify(Gap.Size() >= Entry.Size);
-				}
+				return Gap.Fits(Entry, MemoryMappingAlignment);
+			});
+			verify(RelativeIndex != INDEX_NONE);
 
-				Gaps.RemoveAt(BestGapIndex, EAllowShrinking::No);
+			const int32 BestGapIndex = StartIndex + RelativeIndex; 
 
-				Entry.Offset = Gap.Begin;
-				Gap.Begin = AlignOffset(Gap.Begin + Entry.Size, MemoryMappingAlignment);
+			FLayoutGap Gap = Gaps[BestGapIndex];
+			if (bParanoidChecks)
+			{
+				verify(Gap.Size() >= Entry.Size);
+				verify(Gap.Fits(Entry, MemoryMappingAlignment));
+			}
 
-				if (Gap.Begin <= Gap.End)
-				{
-					Gaps.Insert(Gap, Algo::UpperBoundBy(Gaps, Gap.Size(), GapsBinarySearchProjection));
+			Gaps.RemoveAt(BestGapIndex, EAllowShrinking::No);
+			Entry.Offset = Gap.Place(Entry, MemoryMappingAlignment);
 
-					if (bParanoidChecks)
-					{
-						verify(Algo::IsSortedBy(Gaps, GapsBinarySearchProjection));
-					}
-				}
-			}
-			else
+			if (Gap.Size() > 0)
 			{
-				Entry.Offset = AppendOffset;
-				AppendOffset = AlignOffset(AppendOffset + Entry.Size, MemoryMappingAlignment);
+				Gaps.Insert(Gap, Algo::UpperBoundBy(Gaps, Gap.Size(), GapsBinarySearchProjection));
+
+				if (bParanoidChecks)
+				{
+					verify(Algo::IsSortedBy(Gaps, GapsBinarySearchProjection));
+				}
 			}
 		}
 
-		return FBFDLayoutResult(BytesWritten, AppendOffset);
+		verify(Gaps.Last().bTail == true);
+		return FBFDLayoutResult(BytesWritten, Gaps.Last().Begin);
 	}
 
 	// Simulates BFD by unpinning elements starting from [0, FromUnpinCount to [0, ToUnpinCount]
@@ -1349,7 +1401,7 @@ private:
 					return true;
 				}
 
-				EntriesToInsertUnpinned.Add(FTocEntry(Entry.Id, UINT64_MAX, Entry.OriginalOffset, Entry.Size));
+				EntriesToInsertUnpinned.Add(Entry.Unpin());
 				return false;
 			});
 
@@ -1632,7 +1684,7 @@ private:
 					const uint64 Size = FMath::Min(Op.Size - Offset, TempBufferSize);
 
 					Source->ReadAt(TempBuffer.GetData(), Size, SourceOffset + Offset);
-					Source->ReadAt(TempBuffer2.GetData(), Size, TargetOffset + Offset);
+					Target->ReadAt(TempBuffer2.GetData(), Size, TargetOffset + Offset);
 
 					verify(FMemory::Memcmp(TempBuffer.GetData(), TempBuffer2.GetData(), Size) == 0);
 				}
@@ -2216,10 +2268,6 @@ public:
 		}
 
 		const FIoStoreWriterSettings& WriterSettings = WriterContext->GetSettings();
-		if (WriterSettings.CompressionBlockAlignment)
-		{
-			UE_LOG(LogIoStore, Warning, TEXT("Ignoring CompressionBlockAlignment in for container '%s' due to patch-in-place enabled"), *ContainerPathAndBaseFileName);
-		}
 
 		FIoStoreTocResource& SourceToc = TocBuilder.GetTocResource();
 		// needed for FIoStoreTocResource::GetTocChunkInfo to work 
