[Shaders] Allow putting all IAD shader library headers into a primary container via INI configuration to ensure they are loaded before any IAD container is mounted.

#rnx

#Synced-CL 51102399

diff --git a/Engine/Source/Developer/IoStoreUtilities/Private/IoStoreShaderLibraryProcessor.cpp b/Engine/Source/Developer/IoStoreUtilities/Private/IoStoreShaderLibraryProcessor.cpp
--- a/Engine/Source/Developer/IoStoreUtilities/Private/IoStoreShaderLibraryProcessor.cpp
+++ b/Engine/Source/Developer/IoStoreUtilities/Private/IoStoreShaderLibraryProcessor.cpp
@@ -56,6 +56,14 @@ static bool IsGenerateChunksEnabled()
 	return false;
 }
 
+// Returns the configured container name to be used for on-demand shader library TOCs.
+static FString GetPrimaryContainerForOnDemandShaderLibTOC()
+{
+	FString PrimaryContainerName;
+	GConfig->GetString(TEXT("ShaderCodeLibrary.PakChunkOverride"), TEXT("PrimaryContainerForOnDemandShaderLibTOC"), PrimaryContainerName, GEngineIni);
+	return PrimaryContainerName;
+}
+
 // Returns true if shader code chunk override is enabled, via CVar `r.ShaderCodeLibrary.PakChunkOverride` or commandline `-OverrideShaderChunks`.
 // This also requires that the project being packaged is configured to generate chunks, i.e. `[/Script/UnrealEd.ProjectPackagingSettings]:bGenerateChunks=True`.
 static bool IsChunkOverrideEnabled()
@@ -838,8 +846,31 @@ void FIoStoreShaderLibraryProcessor::ProcessShaderLibraries(const FIoStoreArgume
 				return bIsOnDemandChunk == bOnDemandOnly ? ChunkId : INDEX_NONE;
 			};
 
-		auto FindNonOnDemandContainerTarget = [&ContainerTargets, &GetConditionalChunkId](const FContainerTargetSpec& InContainerTarget, int32 InChunkId) -> FContainerTargetSpec*
+		// First, try to find a pakchunk that was specified as primary target for IAD shader library headers to be loaded before any of such containers are mounted
+		FContainerTargetSpec* PrimaryContainerForOnDemandShaderLibTOC = [&ContainerTargets]() ->FContainerTargetSpec*
 			{
+				const FString PrimaryContainerName = GetPrimaryContainerForOnDemandShaderLibTOC();
+				if (!PrimaryContainerName.IsEmpty())
+				{
+					for (FContainerTargetSpec* ContainerTarget : ContainerTargets)
+					{
+						if (ContainerTarget->Name.ToString().Contains(PrimaryContainerName))
+						{
+							return ContainerTarget;
+						}
+					}
+				}
+				return nullptr;
+			}();
+
+		auto FindNonOnDemandContainerTarget = [&ContainerTargets, &GetConditionalChunkId, PrimaryContainerForOnDemandShaderLibTOC](const FContainerTargetSpec& InContainerTarget, int32 InChunkId) -> FContainerTargetSpec*
+			{
+				// First, try to find a pakchunk that was specified as primary target for IAD shader library headers to be loaded before any of such containers are mounted
+				if (PrimaryContainerForOnDemandShaderLibTOC)
+				{
+					return PrimaryContainerForOnDemandShaderLibTOC;
+				}
+
 				// Try to find non-OnDemand container with same chunk ID (e.g. "Chunk100" for input container "Chunk100iad")
 				for (FContainerTargetSpec* ContainerTarget : ContainerTargets)
 				{
