Zero fill gaps between entries in a patch-in-place container, this is required to have stable container content after patching.
Fixed total padding calculation in patch-in-place.

#rnx

#Synced-CL 51084377

diff --git a/Engine/Source/Developer/IoStoreUtilities/Private/IoStoreWriter.cpp b/Engine/Source/Developer/IoStoreUtilities/Private/IoStoreWriter.cpp
--- a/Engine/Source/Developer/IoStoreUtilities/Private/IoStoreWriter.cpp
+++ b/Engine/Source/Developer/IoStoreUtilities/Private/IoStoreWriter.cpp
@@ -962,10 +962,10 @@ public:
 
 		FitIntoPartitions(TargetEntries, SourceToc.Header.PartitionSize, OutTotalPadding);
 
-		TArray<FCopyOp> CopyOps = {};
+		TArray<FFileOp> FileOps = {};
 		int32 SourcePartitionNum = 0;
-		BuildTargetToc(TargetToc, CopyOps, TargetEntries, SourcePartitionNum, OutTargetPartitionNum);
-		WriteChunksToTarget(TargetEntries, CopyOps, SourcePartitionNum, OutTargetPartitionNum);
+		BuildTargetToc(TargetToc, FileOps, TargetEntries, SourcePartitionNum, OutTargetPartitionNum);
+		WriteChunksToTarget(TargetEntries, FileOps, SourcePartitionNum, OutTargetPartitionNum);
 
 		for (int32 Index = 0; Index < OutTargetPartitionNum; ++Index)
 		{
@@ -1063,16 +1063,18 @@ private:
 		}
 	};
 
-	struct FCopyOp
+	struct FFileOp
 	{
-		uint64 SourceOffset;
-		uint64 TargetOffset;
+		const uint64 SourceOffset;
+		const uint64 TargetOffset;
 		uint64 Size;
+		const bool bZeroFill;
 
-		FCopyOp(uint64 InSourceOffset, uint64 InTargetOffset, uint64 InSize)
+		FFileOp(uint64 InSourceOffset, uint64 InTargetOffset, uint64 InSize, bool bInZeroFill)
 			: SourceOffset(InSourceOffset)
 			, TargetOffset(InTargetOffset)
 			, Size(InSize)
+			, bZeroFill(bInZeroFill)
 		{
 		}
 
@@ -1095,6 +1097,22 @@ private:
 		{
 			return PartitionSize > 0 ? TargetOffset % PartitionSize : TargetOffset;
 		}
+
+		bool CanCoalesce(const uint64 SourceReadOffset, const uint64 TargetWriteOffset, const uint64 PartitionSize) const
+		{
+			return (SourceOffset + Size == SourceReadOffset) &&
+				(TargetOffset + Size == TargetWriteOffset) &&
+				(bZeroFill == false) &&
+				(PartitionSize > 0 ? ( // only coalesce with-in same partition
+					SourceOffset / PartitionSize == SourceReadOffset / PartitionSize &&
+					TargetOffset / PartitionSize == TargetWriteOffset / PartitionSize
+				) : true);
+		}
+
+		void Coalesce(const uint64 OtherSize)
+		{
+			Size += OtherSize;
+		}
 	};
 
 	struct FBFDLayoutResult
@@ -1465,12 +1483,35 @@ private:
 
 			Chunk.Offset += Padding;
 
-			OutTotalPadding = Chunk.Offset - PreviousEnd;
+			OutTotalPadding += Chunk.Offset - PreviousEnd;
 			PreviousEnd = Chunk.Offset + Chunk.Size;
 		}
 	}
 
-	void BuildTargetToc(FIoStoreTocResource& TargetToc, TArray<FCopyOp>& OutCopyOps, TConstArrayView<FTocEntry> Chunks, int32& OutSourcePartitionNum, int32& OutTargetPartitionNum)
+	// emit zero fill of a provided region, splitting ops between partitions
+	static void ZeroFillRegion(TArray<FFileOp>& OutOps, const uint64 Begin, const uint64 End, const uint64 MaxPartitionSize)
+	{
+		if (Begin >= End)
+		{
+			return;
+		}
+
+		if (MaxPartitionSize == 0 || MaxPartitionSize == UINT64_MAX)
+		{
+			OutOps.Add(FFileOp(0, Begin, End - Begin, true));
+			return;
+		}
+
+		for (uint64 Position = Begin; Position < End;)
+		{
+			const uint64 NextPartitionBegin = ((Position / MaxPartitionSize) + 1) * MaxPartitionSize;
+			const uint64 OpEnd = FMath::Min(NextPartitionBegin, End);
+			OutOps.Add(FFileOp(0, Position, OpEnd - Position, true));
+			Position = OpEnd;
+		}
+	}
+
+	void BuildTargetToc(FIoStoreTocResource& TargetToc, TArray<FFileOp>& OutOps, TConstArrayView<FTocEntry> Chunks, int32& OutSourcePartitionNum, int32& OutTargetPartitionNum)
 	{
 		TRACE_CPUPROFILER_EVENT_SCOPE(BuildTargetToc);
 
@@ -1494,12 +1535,12 @@ private:
 		const uint64 MaxPartitionSize = SourceToc.Header.PartitionSize;
 
 		uint64 UncompressedOffset = 0;
-		TUniquePtr<FArchive> TargetHandle;
+		uint64 TargetOffset = 0;
 
 		int32 TargetChunkIndex = 0;
 		int32 TargetBlockIndex = 0;
 
-		OutCopyOps.Reserve(Chunks.Num());
+		OutOps.Reserve(Chunks.Num() * 2); // likely have same amount of gaps
 
 		int32 SourcePartitionNum = 0;
 		int32 TargetPartitionNum = 0;
@@ -1512,6 +1553,13 @@ private:
 			{
 				continue;
 			}
+
+			if (Chunk.Offset > TargetOffset)
+			{
+				// padding is defined as zero filled given expected small amount of padding bytes 
+				ZeroFillRegion(OutOps, TargetOffset, Chunk.Offset, MaxPartitionSize);
+			}
+
 			const int32 SourceIndex = *SourceTocChunkIdToIndex.Find(Chunk.Id);
 
 			TargetToc.ChunkIds[TargetChunkIndex] = Chunk.Id;
@@ -1550,60 +1598,74 @@ private:
 
 			verify(UncompressedOffset <= MaxPossibleOffset);
 			verify(UncompressedSize <= MaxPossibleOffset);
+			verify(DiskSize == Chunk.Size);
 
 			TargetToc.ChunkOffsetLengths[TargetChunkIndex].SetOffset(UncompressedOffset);
 			TargetToc.ChunkOffsetLengths[TargetChunkIndex].SetLength(UncompressedSize);
 
 			// try to coalesce with previous operation
-			if (OutCopyOps.Num() > 0 && 
-				OutCopyOps.Last().SourceOffset + OutCopyOps.Last().Size == SourceReadOffset &&
-				OutCopyOps.Last().TargetOffset + OutCopyOps.Last().Size == Chunk.Offset &&
-				// only coalesce with-in same partition
-				(MaxPartitionSize > 0 ? (
-					OutCopyOps.Last().SourceOffset / MaxPartitionSize == SourceReadOffset / MaxPartitionSize &&
-					OutCopyOps.Last().TargetOffset / MaxPartitionSize == Chunk.Offset / MaxPartitionSize
-					) : true)
-				)
-			{
-				OutCopyOps.Last().Size += Chunk.Size;
+			if (OutOps.Num() > 0 && OutOps.Last().CanCoalesce(SourceReadOffset, Chunk.Offset, MaxPartitionSize))
+			{
+				OutOps.Last().Coalesce(Chunk.Size);
 			}
 			else
 			{
-				const FCopyOp& Op = OutCopyOps.Add_GetRef(FCopyOp(SourceReadOffset, Chunk.Offset, Chunk.Size));
+				const FFileOp& Op = OutOps.Add_GetRef(FFileOp(SourceReadOffset, Chunk.Offset, Chunk.Size, false));
 				SourcePartitionNum = FMath::Max(Op.SourcePartitionIndex(MaxPartitionSize) + 1, SourcePartitionNum);
 				TargetPartitionNum = FMath::Max(Op.TargetPartitionIndex(MaxPartitionSize) + 1, TargetPartitionNum);
 			}
 
-			verify(DiskSize == Chunk.Size);
+			TargetOffset = Chunk.Offset + Chunk.Size;
 
 			UncompressedOffset += (SourceLastBlockIndex - SourceFirstBlockIndex + 1) * CompressionBlockSize;
 			TargetChunkIndex++;
 		}
 
-		Algo::SortBy(OutCopyOps, [](const FCopyOp& X) {return X.SourceOffset;});
+		Algo::SortBy(OutOps, &FFileOp::TargetOffset);
 
 		OutSourcePartitionNum = SourcePartitionNum;
 		OutTargetPartitionNum = TargetPartitionNum;
+
+		if (bParanoidChecks && !Chunks.IsEmpty())
+		{
+			// copy ops should write to every single byte in all partitions and not overlap
+			uint64 ExpectedOffset = 0;
+			for (const FFileOp& Op : OutOps)
+			{
+				verify(Op.TargetOffset == ExpectedOffset);
+				verify(Op.Size > 0);
+				ExpectedOffset = Op.TargetOffset + Op.Size;
+			}
+
+			const uint64 ExpectedTotalSize = Chunks.Last().Offset + Chunks.Last().Size;
+			verify(ExpectedOffset == ExpectedTotalSize);
+		}
 	}
 
-	void WriteChunksToTarget(TConstArrayView<FTocEntry> Chunks, TConstArrayView<FCopyOp> CopyOps, const int32 SourcePartitionNum, const int32 TargetPartitionNum)
+	void WriteChunksToTarget(TConstArrayView<FTocEntry> Chunks, TConstArrayView<FFileOp> FileOps, const int32 SourcePartitionNum, const int32 TargetPartitionNum)
 	{
 		TRACE_CPUPROFILER_EVENT_SCOPE(WriteChunksToTarget);
 
-		if (Chunks.IsEmpty() || SourcePartitionNum == 0 || TargetPartitionNum == 0 || CopyOps.IsEmpty())
+		if (Chunks.IsEmpty() || SourcePartitionNum == 0 || TargetPartitionNum == 0 || FileOps.IsEmpty())
 		{
 			return;
 		}
 
 		const uint64 MaxPartitionSize = SourceToc.Header.PartitionSize;
 		const uint64 LastPartitionSizeLogical = Chunks.Last().Offset + Chunks.Last().Size;
-		const uint64 LastPartitionSize = MaxPartitionSize > 0 ? LastPartitionSizeLogical % MaxPartitionSize : LastPartitionSizeLogical;
-
 		if (LastPartitionSizeLogical == 0)
 		{
 			return;
 		}
 
+		const uint64 LastPartitionSize = 
+			MaxPartitionSize > 0
+			? (LastPartitionSizeLogical % MaxPartitionSize == 0
+				? MaxPartitionSize // edge case, LastPartitionSizeLogical precisely matches partition size
+				: LastPartitionSizeLogical % MaxPartitionSize)
+			: LastPartitionSizeLogical;
+		ensure(LastPartitionSize != 0);
+
 		IPlatformFile& Ipf = IPlatformFile::GetPlatformPhysical();
 		Ipf.CreateDirectoryTree(*FPaths::GetPath(TargetContainerPath));
 
@@ -1642,9 +1704,9 @@ private:
 		constexpr uint64 TempBufferSize = 4 * 1024 * 1024; 
 		FIoBuffer TempBuffer(TempBufferSize);
 
-		for (const FCopyOp& Op : CopyOps) 
+		for (const FFileOp& Op : FileOps) 
 		{
-			IFileHandle* Source = SourcePartitions[Op.SourcePartitionIndex(MaxPartitionSize)].Get();
+			IFileHandle* Source = Op.bZeroFill ? nullptr : SourcePartitions[Op.SourcePartitionIndex(MaxPartitionSize)].Get();
 			IFileHandle* Target = TargetPartitions[Op.TargetPartitionIndex(MaxPartitionSize)].Get();
 
 			const uint64 SourceOffset = Op.SourcePartitionOffset(MaxPartitionSize);
@@ -1656,11 +1718,21 @@ private:
 
 				if (bParanoidChecks)
 				{
-					verify(SourceOffset + Offset + Size <= (uint64)Source->Size());
+					if (!Op.bZeroFill)
+					{
+						verify(SourceOffset + Offset + Size <= (uint64)Source->Size());
+					}
 					verify(TargetOffset + Offset + Size <= (uint64)Target->Size());
 				}
 
-				Source->ReadAt(TempBuffer.GetData(), Size, SourceOffset + Offset);
+				if (Op.bZeroFill)
+				{
+					FMemory::Memzero(TempBuffer.GetData(), Size);
+				}
+				else
+				{
+					Source->ReadAt(TempBuffer.GetData(), Size, SourceOffset + Offset);
+				}
 
 				Target->Seek(TargetOffset + Offset);
 				Target->Write(TempBuffer.GetData(), Size);
@@ -1671,9 +1743,9 @@ private:
 		{
 			FIoBuffer TempBuffer2(TempBufferSize);
 
-			for (const FCopyOp& Op : CopyOps) 
+			for (const FFileOp& Op : FileOps) 
 			{
-				IFileHandle* Source = SourcePartitions[Op.SourcePartitionIndex(MaxPartitionSize)].Get();
+				IFileHandle* Source = Op.bZeroFill ? nullptr : SourcePartitions[Op.SourcePartitionIndex(MaxPartitionSize)].Get();
 				IFileHandle* Target = TargetPartitions[Op.TargetPartitionIndex(MaxPartitionSize)].Get();
 
 				const uint64 SourceOffset = Op.SourcePartitionOffset(MaxPartitionSize);
@@ -1683,7 +1755,14 @@ private:
 				{
 					const uint64 Size = FMath::Min(Op.Size - Offset, TempBufferSize);
 
-					Source->ReadAt(TempBuffer.GetData(), Size, SourceOffset + Offset);
+					if (Op.bZeroFill)
+					{
+						FMemory::Memzero(TempBuffer.GetData(), Size);
+					}
+					else
+					{
+						Source->ReadAt(TempBuffer.GetData(), Size, SourceOffset + Offset);
+					}
 					Target->ReadAt(TempBuffer2.GetData(), Size, TargetOffset + Offset);
 
 					verify(FMemory::Memcmp(TempBuffer.GetData(), TempBuffer2.GetData(), Size) == 0);
