Partially preserving order for inserted entries in patch-in-place

#rnx
#Synced-CL 51328467

diff --git a/Engine/Source/Developer/IoStoreUtilities/Private/IoStoreUtilities.cpp b/Engine/Source/Developer/IoStoreUtilities/Private/IoStoreUtilities.cpp
--- a/Engine/Source/Developer/IoStoreUtilities/Private/IoStoreUtilities.cpp
+++ b/Engine/Source/Developer/IoStoreUtilities/Private/IoStoreUtilities.cpp
@@ -2302,6 +2302,7 @@ void LogWriterResults(const FIoStoreWriterSettings& GeneralIoWriterSettings, con
 		uint64 PatchInPlaceSizeBeforeLayout = 0;
 		uint64 PatchInPlacePinnedSizeBeforeDefrag = 0;
 		uint64 PatchInPlacePinnedSizeAfterDefrag = 0;
+		int64 PatchInPlaceLocalityChange = 0;
 	};
 
 	// Container name length can vary wildly, so find the longest one so that we can correctly align the output grids
@@ -2320,11 +2321,11 @@ void LogWriterResults(const FIoStoreWriterSettings& GeneralIoWriterSettings, con
 	FString PatchInPlaceHeader = TEXT("");
 	if (GeneralIoWriterSettings.bUsePatchInPlaceLayout)
 	{
-		PatchInPlaceHeader = FString::Printf(TEXT(" %20s %20s %20s"), TEXT("PIP Pinned (MiB)"), TEXT("PIP Defrag (MiB)"), TEXT("PIP Overhead (MiB)"));
+		PatchInPlaceHeader = FString::Printf(TEXT(" %20s %20s %20s %20s"), TEXT("PIP Pinned (MiB)"), TEXT("PIP Defrag (MiB)"), TEXT("PIP Locality (MiB)"), TEXT("PIP Overhead (MiB)"));
 	}
 	UE_LOG(LogIoStore, Display, TEXT("%-*s %10s %15s %15s %20s %20s%s"), ColumnWidth, TEXT("Container"), TEXT("Flags"), TEXT("Chunk(s) #"), TEXT("TOC (KiB)"), TEXT("Raw Size (MiB)"), TEXT("Size (MiB)"), *PatchInPlaceHeader);
 
-	const int32 SummaryBorderLength = ColumnWidth + 5 + 10 + 15 + 15 + 20 + 20 + (GeneralIoWriterSettings.bUsePatchInPlaceLayout ? 3 + 20 + 20 + 20 : 0); // +5 to account for spaces between columns
+	const int32 SummaryBorderLength = ColumnWidth + 5 + 10 + 15 + 15 + 20 + 20 + (GeneralIoWriterSettings.bUsePatchInPlaceLayout ? 4 + 20 + 20 + 20 + 20 : 0); // +5 to account for spaces between columns
 	UE_LOG(LogIoStore, Display, TEXT("%s"), *FString::ChrN(SummaryBorderLength, '-'));
 	
 	FContainerStats TotalStats;
@@ -2347,9 +2348,10 @@ void LogWriterResults(const FIoStoreWriterSettings& GeneralIoWriterSettings, con
 			const int64 Overhead = (int64)Result.CompressedContainerSize - Result.PatchInPlaceSizeBeforeLayout;
 			const int64 DefragUnpinned = (int64)Result.PatchInPlacePinnedSizeAfterDefrag - Result.PatchInPlacePinnedSizeBeforeDefrag;
 
-			PatchInPlaceInfo = FString::Printf(TEXT(" %20.2lf %20.2lf %20.2lf"),
+			PatchInPlaceInfo = FString::Printf(TEXT(" %20.2lf %20.2lf %20.2lf %20.2lf"),
 				(double)Result.PatchInPlacePinnedSizeAfterDefrag / 1024.0 / 1024.0,
 				(double)DefragUnpinned / 1024.0 / 1024.0,
+				(double)Result.PatchInPlaceLocalityChange / 1024.0 / 1024.0,
 				(double)Overhead / 1024.0 / 1024.0
 			);
 		}
@@ -2387,6 +2389,7 @@ void LogWriterResults(const FIoStoreWriterSettings& GeneralIoWriterSettings, con
 		TotalStats.PatchInPlaceSizeBeforeLayout += Result.PatchInPlaceSizeBeforeLayout;
 		TotalStats.PatchInPlacePinnedSizeBeforeDefrag += Result.PatchInPlacePinnedSizeBeforeDefrag;
 		TotalStats.PatchInPlacePinnedSizeAfterDefrag += Result.PatchInPlacePinnedSizeAfterDefrag;
+		TotalStats.PatchInPlaceLocalityChange += Result.PatchInPlaceLocalityChange;
 	}
 
 	UE_LOG(LogIoStore, Display, TEXT("%s"), *FString::ChrN(SummaryBorderLength, '-'));
@@ -2408,9 +2411,10 @@ void LogWriterResults(const FIoStoreWriterSettings& GeneralIoWriterSettings, con
 		const int64 TotalOverhead = (int64)TotalStats.CompressedContainerSize - TotalStats.PatchInPlaceSizeBeforeLayout;
 		const int64 DefragUnpinned = (int64)TotalStats.PatchInPlacePinnedSizeAfterDefrag - TotalStats.PatchInPlacePinnedSizeBeforeDefrag;
 
-		PatchInPlaceTotal = FString::Printf(TEXT(" %20.2lf %20.2lf %20.2lf"),
+		PatchInPlaceTotal = FString::Printf(TEXT(" %20.2lf %20.2lf %20.2lf %20.2lf"),
 			(double)TotalStats.PatchInPlacePinnedSizeAfterDefrag / 1024.0 / 1024.0,
 			(double)DefragUnpinned / 1024.0 / 1024.0,
+			(double)TotalStats.PatchInPlaceLocalityChange / 1024.0 / 1024.0,
 			(double)TotalOverhead / 1024.0 / 1024.0
 		);
 	}
@@ -9130,10 +9134,16 @@ bool ParseContainerGenerationArguments(FIoStoreArguments& Arguments, FIoStoreWri
 		WriterSettings.PatchInPlaceParetoFrontCoarseSamples = 50;
 		WriterSettings.PatchInPlaceParetoFrontFineSamples = 1500;
 
+		// typical read granularity for QLC NAND is 16kb, TLC is 8KB
+		WriterSettings.PatchInPlaceReadGranularity = 16 * 1024;
+
 		ParseSizeArgument(FCommandLine::Get(), TEXT("-patchInPlaceMaximumTotalGrowth="), WriterSettings.PatchInPlaceMaximumTotalGrowth, 0);
 		ParseSizeArgument(FCommandLine::Get(), TEXT("-patchInPlaceMaximumPerContainerGrowth="), WriterSettings.PatchInPlaceMaximumPerContainerGrowth, 0);
 		FParse::Value(FCommandLine::Get(), TEXT("-patchInPlaceParetoFrontCoarseSamples="), WriterSettings.PatchInPlaceParetoFrontCoarseSamples);
 		FParse::Value(FCommandLine::Get(), TEXT("-patchInPlaceParetoFrontFineSamples="), WriterSettings.PatchInPlaceParetoFrontFineSamples);
+		FParse::Value(FCommandLine::Get(), TEXT("-patchInPlaceReadGranularity="), WriterSettings.PatchInPlaceReadGranularity);
+		FParse::Value(FCommandLine::Get(), TEXT("-patchInPlaceMaxChainSize="), WriterSettings.PatchInPlaceMaxChainSize);
+		FParse::Value(FCommandLine::Get(), TEXT("-patchInPlaceMaxChainLength="), WriterSettings.PatchInPlaceMaxChainLength);
 		FParse::Value(FCommandLine::Get(), TEXT("-patchInPlaceParetoFrontHeuristic="), WriterSettings.PatchInPlaceParetoFrontHeuristic);
 	}
 
diff --git a/Engine/Source/Developer/IoStoreUtilities/Private/IoStoreWriter.cpp b/Engine/Source/Developer/IoStoreUtilities/Private/IoStoreWriter.cpp
--- a/Engine/Source/Developer/IoStoreUtilities/Private/IoStoreWriter.cpp
+++ b/Engine/Source/Developer/IoStoreUtilities/Private/IoStoreWriter.cpp
@@ -835,7 +835,7 @@ public:
 		}
 		OutPinnedSizeAfterDefrag = OutPinnedSizeBeforeDefrag;
 
-		const FBFDLayoutResult ResultBeforeDefrag = BestFitDecreasing(EntriesToKeep, EntriesToInsert, Settings.MemoryMappingAlignment);
+		const FBFDLayoutResult ResultBeforeDefrag = BestFitDecreasing(EntriesToKeep, EntriesToInsert, Settings);
 
 		// continue if we fit into growth limit
 		if (Settings.PatchInPlaceMaximumPerContainerGrowth == 0 ||
@@ -845,7 +845,7 @@ public:
 		}
 
 		TMap<FIoChunkId, int32> UnpinCandidates;
-		BuildUnpinCandidates(UnpinCandidates, EntriesToKeep, EntriesToInsert, Settings.MemoryMappingAlignment, Settings.PatchInPlaceParetoFrontHeuristic);
+		BuildUnpinCandidates(UnpinCandidates, EntriesToKeep, EntriesToInsert, Settings);
 
 		if (bParanoidChecks)
 		{
@@ -856,7 +856,7 @@ public:
 				EntriesToInsert,
 				UnpinCandidates,
 				SizeBeforeLayout,
-				Settings.MemoryMappingAlignment,
+				Settings,
 				ParetoFrontTemp,
 				UnpinCandidates.Num() - 1,
 				UnpinCandidates.Num(),
@@ -874,7 +874,7 @@ public:
 			EntriesToInsert,
 			UnpinCandidates,
 			SizeBeforeLayout + Settings.PatchInPlaceMaximumPerContainerGrowth,
-			Settings.MemoryMappingAlignment,
+			Settings,
 			ParetoFront,
 			0,
 			UnpinCandidates.Num(),
@@ -886,7 +886,7 @@ public:
 			EntriesToInsert,
 			UnpinCandidates,
 			SizeBeforeLayout + Settings.PatchInPlaceMaximumPerContainerGrowth,
-			Settings.MemoryMappingAlignment,
+			Settings,
 			ParetoFront,
 			ParetoFront.Num() > 1 ? ParetoFront[0].UnpinCount : 0,
 			ParetoFront.Num() > 1 ? ParetoFront.Last().UnpinCount : UnpinCandidates.Num(),
@@ -919,7 +919,7 @@ public:
 			return true;
 		});
 
-		Algo::SortBy(EntriesToInsert, [](const FTocEntry& X){return X.Size;}, TGreater<uint64>());
+		Algo::SortBy(EntriesToInsert, &FTocEntry::OriginalOffset);
 
 		OutPinnedSizeAfterDefrag = 0;
 		for (const FTocEntry& Entry : EntriesToKeep)
@@ -928,7 +928,7 @@ public:
 		}
 
 		// redo BFD on the new set
-		const FBFDLayoutResult ResultAfterDefrag = BestFitDecreasing(EntriesToKeep, EntriesToInsert, Settings.MemoryMappingAlignment);
+		const FBFDLayoutResult ResultAfterDefrag = BestFitDecreasing(EntriesToKeep, EntriesToInsert, Settings);
 
 		UE_LOG(LogIoStore, Display,
 			TEXT("Patch-in-place '%s' container grew above limit (%.2f MiB before pip -> %.2f MiB after pip), defragmented by unpinning %.2f MiB, new container size %.2f MiB"),
@@ -945,7 +945,8 @@ public:
 	void RunLayout(
 		FIoStoreTocResource& TargetToc,
 		int32& OutTargetPartitionNum,
-		uint64& OutTotalPadding
+		uint64& OutTotalPadding,
+		int64& OutLocalityChange
 	)
 	{
 		TRACE_CPUPROFILER_EVENT_SCOPE(PatchInPlaceLayout);
@@ -962,6 +963,18 @@ public:
 
 		FitIntoPartitions(TargetEntries, SourceToc.Header.PartitionSize, OutTotalPadding);
 
+		uint64 LocalityBefore = 0;
+		uint64 LocalityAfter = 0;
+		CalculateLocalityChange(TargetEntries, Settings, LocalityBefore, LocalityAfter);
+		OutLocalityChange = (int64)LocalityAfter - (int64)LocalityBefore;
+
+		UE_LOG(LogIoStore, Display, TEXT("Access locality for '%s' with granularity %llu, original %.2f MiB, new %.2f MiB, diff %.2f MiB"),
+			*TargetContainerPath,
+			Settings.PatchInPlaceReadGranularity,
+			(double)LocalityBefore / 1024.0 / 1024.0,
+			(double)LocalityAfter / 1024.0 / 1024.0,
+			(double)OutLocalityChange / 1024.0 / 1024.0);
+
 		TArray<FFileOp> FileOps = {};
 		int32 SourcePartitionNum = 0;
 		BuildTargetToc(TargetToc, FileOps, TargetEntries, SourcePartitionNum, OutTargetPartitionNum);
@@ -987,8 +1000,8 @@ private:
 	TMap<FIoChunkId, int32> SourceTocChunkIdToIndex;
 
 	struct FTocEntry;
-	TArray<FTocEntry> EntriesToKeep; // sorted by offset increasing
-	TArray<FTocEntry> EntriesToInsert; // sorted by size decreasing
+	TArray<FTocEntry> EntriesToKeep; // sorted by Offset increasing
+	TArray<FTocEntry> EntriesToInsert; // sorted by OriginalOffset increasing
 
 	static void BuildTocMap(const FIoStoreTocResourceView& Toc, TMap<FIoChunkId, int32>& OutResult)
 	{
@@ -1012,6 +1025,11 @@ private:
 		{
 		}
 
+		uint64 OriginalEnd() const
+		{
+			return OriginalOffset + Size;
+		}
+
 		uint64 End() const
 		{
 			return Offset + Size;
@@ -1033,6 +1051,46 @@ private:
 		}
 	};
 
+	struct FInsertionChain
+	{
+		int32 Index;
+		int32 Count;
+		uint64 Size;
+
+		FInsertionChain(const int32 InIndex, const uint64 InSize)
+			: Index(InIndex), Count(1), Size(InSize)
+		{
+		}
+
+		void Coalesce(const int32 InIndex, const uint64 InSize)
+		{
+			ensure(Index + Count == InIndex);
+			Count++;
+			Size += InSize;
+		}
+
+		uint64 CalculateEnd(uint64 AtOffset, const TConstArrayView<FTocEntry>& EntriesToInsert, const uint64 MemoryMappingAlignment) const
+		{
+			for (int32 I = Index; I < Index + Count; ++I)
+			{
+				AtOffset = AlignOffset(AtOffset, EntriesToInsert[I].bMemoryMapped, MemoryMappingAlignment) + EntriesToInsert[I].Size;
+			}
+
+			return AtOffset;
+		}
+
+		uint64 InsertEntries(uint64 AtOffset, const TArrayView<FTocEntry>& EntriesToInsert, const uint64 MemoryMappingAlignment) const
+		{
+			for (int32 I = Index; I < Index + Count; ++I)
+			{
+				EntriesToInsert[I].Offset = AlignOffset(AtOffset, EntriesToInsert[I].bMemoryMapped, MemoryMappingAlignment);
+				AtOffset = EntriesToInsert[I].End();
+			}
+
+			return AtOffset;
+		}
+	};
+
 	struct FLayoutGap
 	{
 		uint64 Begin; // inclusive, unaligned
@@ -1172,8 +1230,8 @@ private:
 			}
 		}
 
-		Algo::SortBy(EntriesToKeep, [](const FTocEntry& X){return X.Offset;});
-		Algo::SortBy(EntriesToInsert, [](const FTocEntry& X){return X.Size;}, TGreater<uint64>());
+		Algo::SortBy(EntriesToKeep, &FTocEntry::Offset);
+		Algo::SortBy(EntriesToInsert, &FTocEntry::OriginalOffset);
 	}
 
 	static uint64 GapsBinarySearchProjection(const FLayoutGap& Gap)
@@ -1191,29 +1249,58 @@ private:
 		return Result;
 	}
 
-	static void EnsureSortedAndNonOverlapping(TConstArrayView<FTocEntry> Entries)
+	static void EnsureSortedAndNonOverlapping(TConstArrayView<FTocEntry> Entries, bool bUseOriginalOffset = false)
 	{
-		for (int32 Index = 1; Index < Entries.Num(); ++Index)
+		if (bUseOriginalOffset)
 		{
-			verify(Entries[Index].Offset > Entries[Index - 1].Offset);
-		}
+			for (int32 Index = 1; Index < Entries.Num(); ++Index)
+			{
+				verify(Entries[Index].OriginalOffset > Entries[Index - 1].OriginalOffset);
+			}
 
-		for (int32 I = 0; I < Entries.Num(); ++I)
-		{
-			for (int32 J = 0; J < Entries.Num(); ++J)
+			for (int32 I = 0; I < Entries.Num(); ++I)
 			{
-				if (I == J)
+				for (int32 J = 0; J < Entries.Num(); ++J)
 				{
-					continue;
+					if (I == J)
+					{
+						continue;
+					}
+
+					const uint64 L1 = Entries[I].OriginalOffset;
+					const uint64 R1 = Entries[I].OriginalEnd() - 1;
+
+					const uint64 L2 = Entries[J].OriginalOffset;
+					const uint64 R2 = Entries[J].OriginalEnd() - 1;
+
+					verify(R1 < L2 || L1 > R2);
 				}
+			}
+		}
+		else
+		{
+			for (int32 Index = 1; Index < Entries.Num(); ++Index)
+			{
+				verify(Entries[Index].Offset > Entries[Index - 1].Offset);
+			}
+
+			for (int32 I = 0; I < Entries.Num(); ++I)
+			{
+				for (int32 J = 0; J < Entries.Num(); ++J)
+				{
+					if (I == J)
+					{
+						continue;
+					}
 
-				const uint64 L1 = Entries[I].Offset;
-				const uint64 R1 = Entries[I].End() - 1;
+					const uint64 L1 = Entries[I].Offset;
+					const uint64 R1 = Entries[I].End() - 1;
 
-				const uint64 L2 = Entries[J].Offset;
-				const uint64 R2 = Entries[J].End() - 1;
+					const uint64 L2 = Entries[J].Offset;
+					const uint64 R2 = Entries[J].End() - 1;
 
-				verify(R1 < L2 || L1 > R2);
+					verify(R1 < L2 || L1 > R2);
+				}
 			}
 		}
 	}
@@ -1222,8 +1309,7 @@ private:
 		TMap<FIoChunkId, int32>& OutUnpinCandidates, // chunk id -> index of all unpin candidates 
 		const TConstArrayView<FTocEntry>& EntriesToKeep,
 		const TConstArrayView<FTocEntry>& EntriesToInsert,
-		const uint64 MemoryMappingAlignment,
-		const uint8 ParetoFrontHeuristic
+		const FIoStoreWriterSettings& Settings
 	)
 	{
 		struct FUnpinCandidate
@@ -1250,7 +1336,7 @@ private:
 			const uint64 EntrySize = Entry.Size;
 
 			double Contribution;
-			switch (ParetoFrontHeuristic)
+			switch (Settings.PatchInPlaceParetoFrontHeuristic)
 			{
 			default:
 			case 0: // Prioritize based on how "useful" the gap for incoming dataset
@@ -1260,13 +1346,13 @@ private:
 
 					for (const FTocEntry& Insert : EntriesToInsert)
 					{
-						if (!Gap.Fits(Insert, MemoryMappingAlignment))
+						if (!Gap.Fits(Insert, Settings.MemoryMappingAlignment))
 						{
 							continue;
 						}
 
 						UsableSpace += Insert.Size;
-						Gap.Place(Insert, MemoryMappingAlignment);
+						Gap.Place(Insert, Settings.MemoryMappingAlignment);
 
 						if (Gap.Size() == 0)
 						{
@@ -1292,7 +1378,7 @@ private:
 			UnpinCandidates.Add(FUnpinCandidate(Entry.Id, Contribution));
 		}
 
-		Algo::SortBy(UnpinCandidates, [](const FUnpinCandidate& X){return X.Contribution;}, TGreater<double>());
+		Algo::SortBy(UnpinCandidates, &FUnpinCandidate::Contribution, TGreater<double>());
 
 		OutUnpinCandidates.Reserve(UnpinCandidates.Num());
 		for (int32 i = 0; i < UnpinCandidates.Num(); i++)
@@ -1304,18 +1390,14 @@ private:
 	static FBFDLayoutResult BestFitDecreasing(
 		const TConstArrayView<FTocEntry>& EntriesToKeep,
 		TArray<FTocEntry>& EntriesToInsert,
-		const uint64 MemoryMappingAlignment)
+		const FIoStoreWriterSettings& Settings)
 	{
 		TRACE_CPUPROFILER_EVENT_SCOPE(BestFitDecreasing);
 
 		if (bParanoidChecks)
 		{
 			EnsureSortedAndNonOverlapping(EntriesToKeep);
-
-			for (int32 Index = 1; Index < EntriesToInsert.Num(); ++Index)
-			{
-				verify(EntriesToInsert[Index].Size <= EntriesToInsert[Index - 1].Size);
-			}
+			EnsureSortedAndNonOverlapping(EntriesToInsert, true);
 		}
 
 		TArray<FLayoutGap> Gaps;
@@ -1332,38 +1414,64 @@ private:
 
 		Algo::SortBy(Gaps, GapsBinarySearchProjection);
 
+		const uint64 MemoryMappingAlignment = Settings.MemoryMappingAlignment;
+		const uint64 Granularity = Settings.PatchInPlaceReadGranularity;
+		const uint64 MaxChainSize = Settings.PatchInPlaceMaxChainSize;
+		const int32 MaxChainLength = Settings.PatchInPlaceMaxChainLength;
+
+		TArray<FInsertionChain> Chains;
+		for (int32 I = 0; I < EntriesToInsert.Num(); ++I)
+		{
+			const FTocEntry* Prev = I > 0 ? &EntriesToInsert[I - 1] : nullptr;
+			const FTocEntry& Curr = EntriesToInsert[I];
+
+			ensure(I == 0 || Prev->OriginalEnd() <= Curr.OriginalOffset);
+
+			const bool bCanCoalesce = 
+				Granularity > 0 &&
+				I > 0 &&
+				// heuristic, coalesce if previous end and current start were close enough to potentially end up in the same read block
+				Curr.OriginalOffset - Prev->OriginalEnd() < Granularity &&
+				(MaxChainLength == 0 || Chains.Last().Count < MaxChainLength) &&
+				(MaxChainSize == 0 || Chains.Last().Size + Curr.Size <= MaxChainSize);
+
+			if (bCanCoalesce)
+			{
+				Chains.Last().Coalesce(I, Curr.Size);
+			}
+			else
+			{
+				Chains.Add(FInsertionChain(I, Curr.Size));
+			}
+		}
+
+		Algo::SortBy(Chains, &FInsertionChain::Size, TGreater<uint64>());
+
 		uint64 BytesWritten = 0;
 
 		// TODO support compression block alignment?
-		// TODO try place entries together via Entry.OriginalOffset to preserve parts of original layout
-		for (FTocEntry& Entry : EntriesToInsert)
+		for (const FInsertionChain& Chain : Chains)
 		{
-			BytesWritten += Entry.Size;
+			BytesWritten += Chain.Size;
 
-			// gaps are sorted by size, find first that fits by size
-			const int32 StartIndex = Algo::LowerBoundBy(Gaps, Entry.Size, GapsBinarySearchProjection);
+			// Binary search for first gap that fits by total size (lower bound, ignoring alignment padding)
+			const int32 StartIndex = Algo::LowerBoundBy(Gaps, Chain.Size, GapsBinarySearchProjection);
 			verify(StartIndex < Gaps.Num());
 
-			// then find one that actually fits given alignment requirement
+			// Linear scan for a gap that actually fits the whole chain with alignment of each element
 			const int32 RelativeIndex = TConstArrayView<FLayoutGap>(Gaps)
-					.RightChop(StartIndex)
-					.IndexOfByPredicate([&](const FLayoutGap& Gap)
+				.RightChop(StartIndex)
+				.IndexOfByPredicate([&](const FLayoutGap& Gap)
 			{
-				return Gap.Fits(Entry, MemoryMappingAlignment);
+				return Chain.CalculateEnd(Gap.Begin, EntriesToInsert, MemoryMappingAlignment) <= Gap.End;
 			});
 			verify(RelativeIndex != INDEX_NONE);
 
-			const int32 BestGapIndex = StartIndex + RelativeIndex; 
+			const int32 BestGapIndex = StartIndex + RelativeIndex;
 
 			FLayoutGap Gap = Gaps[BestGapIndex];
-			if (bParanoidChecks)
-			{
-				verify(Gap.Size() >= Entry.Size);
-				verify(Gap.Fits(Entry, MemoryMappingAlignment));
-			}
-
 			Gaps.RemoveAt(BestGapIndex, EAllowShrinking::No);
-			Entry.Offset = Gap.Place(Entry, MemoryMappingAlignment);
+			Gap.Begin = Chain.InsertEntries(Gap.Begin, EntriesToInsert, MemoryMappingAlignment);
 
 			if (Gap.Size() > 0)
 			{
@@ -1387,7 +1495,7 @@ private:
 		const TConstArrayView<FTocEntry>& EntriesToInsert,
 		const TMap<FIoChunkId, int32>& UnpinCandidates,
 		const uint64 MaxContainerSize,
-		const uint64 MemoryMappingAlignment,
+		const FIoStoreWriterSettings& Settings,
 		TArray<FParetoPoint>& ParetoFront,
 		const int32 FromUnpinCount, // inclusive
 		const int32 ToUnpinCount, // exclusive
@@ -1423,9 +1531,9 @@ private:
 				return false;
 			});
 
-			Algo::SortBy(EntriesToInsertUnpinned, [](const FTocEntry& X){return X.Size;}, TGreater<uint64>());
+			Algo::SortBy(EntriesToInsertUnpinned, &FTocEntry::OriginalOffset);
 
-			const FBFDLayoutResult Result = BestFitDecreasing(EntriesToKeepUnpinned, EntriesToInsertUnpinned, MemoryMappingAlignment);
+			const FBFDLayoutResult Result = BestFitDecreasing(EntriesToKeepUnpinned, EntriesToInsertUnpinned, Settings);
 
 			Results[Index] = FParetoPoint(UnpinCount, Result);
 		});
@@ -1488,6 +1596,62 @@ private:
 		}
 	}
 
+	// returns amount of bytes that are prefetched when reading all entries in their provided order
+	static void CalculateLocalityChange(
+		TConstArrayView<FTocEntry> Chunks,
+		const FIoStoreWriterSettings& Settings,
+		uint64& OutLocalityBefore,
+		uint64& OutLocalityAfter
+	)
+	{
+		OutLocalityBefore = 0;
+		OutLocalityAfter = 0;
+
+		const uint64 Granularity = Settings.PatchInPlaceReadGranularity;
+		if (Chunks.Num() < 2 || Granularity == 0)
+		{
+			return;
+		}
+
+		// assuming everything is accessed as in OriginalOffset order
+		TArray<int32> AccessOrder;
+		AccessOrder.SetNumUninitialized(Chunks.Num());
+		for (int32 I = 0; I < Chunks.Num(); ++I)
+		{
+			AccessOrder[I] = I;
+		}
+		Algo::SortBy(AccessOrder, [&Chunks](const int32 X){return Chunks[X].OriginalOffset;});
+
+		for (int32 I = 1; I < AccessOrder.Num(); ++I)
+		{
+			const FTocEntry& Prev = Chunks[AccessOrder[I - 1]];
+			const FTocEntry& Curr = Chunks[AccessOrder[I]];
+
+			if (Prev.Size == 0)
+			{
+				continue;
+			}
+
+			const uint64 OriginalPrevEndPage = (Prev.OriginalEnd() - 1) / Granularity;
+			const uint64 OriginalPrevNextPageOffset = (OriginalPrevEndPage + 1) * Granularity;
+			const uint64 OriginalCurrBeginPage = Curr.OriginalOffset / Granularity;
+
+			if (OriginalPrevEndPage == OriginalCurrBeginPage)
+			{
+				OutLocalityBefore += FMath::Min(Curr.Size, OriginalPrevNextPageOffset - Curr.OriginalOffset);
+			}
+
+			const uint64 PrevEndPage = (Prev.End() - 1) / Granularity;
+			const uint64 PrevNextPageOffset = (PrevEndPage + 1) * Granularity;
+			const uint64 CurrBeginPage = Curr.Offset / Granularity;
+
+			if (PrevEndPage == CurrBeginPage)
+			{
+				OutLocalityAfter += FMath::Min(Curr.Size, PrevNextPageOffset - Curr.Offset);
+			}
+		}
+	}
+
 	// emit zero fill of a provided region, splitting ops between partitions
 	static void ZeroFillRegion(TArray<FFileOp>& OutOps, const uint64 Begin, const uint64 End, const uint64 MaxPartitionSize)
 	{
@@ -2372,7 +2536,7 @@ public:
 		// release file handles for patch-in-place
 		Partitions.Empty();
 
-		Layout.RunLayout(TargetToc, TargetPartitionNum, TotalPaddingSize);
+		Layout.RunLayout(TargetToc, TargetPartitionNum, TotalPaddingSize, PatchInPlaceLocalityChange);
 
 		SourceToc = MoveTemp(TargetToc);
 
@@ -2549,6 +2713,7 @@ public:
 		Result.PatchInPlaceSizeBeforeLayout = PatchInPlaceSizeBeforeLayout;
 		Result.PatchInPlacePinnedSizeBeforeDefrag = PatchInPlacePinnedSizeBeforeDefrag;
 		Result.PatchInPlacePinnedSizeAfterDefrag = PatchInPlacePinnedSizeAfterDefrag;
+		Result.PatchInPlaceLocalityChange = PatchInPlaceLocalityChange;
 		Result.ModifiedChunksCount = 0;
 		Result.AddedChunksCount = 0;
 		Result.ModifiedChunksSize= 0;
@@ -3341,6 +3506,7 @@ private:
 	uint64						PatchInPlaceSizeBeforeLayout = 0;
 	uint64						PatchInPlacePinnedSizeBeforeDefrag = 0;
 	uint64						PatchInPlacePinnedSizeAfterDefrag = 0;
+	int64						PatchInPlaceLocalityChange = 0;
 	uint64						MaxPartitionSize = 0;
 	int32						CurrentPartitionIndex = 0;
 	bool						bHasMemoryMappedEntry = false;
diff --git a/Engine/Source/Developer/IoStoreUtilities/Private/IoStoreWriter.h b/Engine/Source/Developer/IoStoreUtilities/Private/IoStoreWriter.h
--- a/Engine/Source/Developer/IoStoreUtilities/Private/IoStoreWriter.h
+++ b/Engine/Source/Developer/IoStoreUtilities/Private/IoStoreWriter.h
@@ -26,6 +26,9 @@ struct FIoStoreWriterSettings
 	uint64 PatchInPlaceMaximumPerContainerGrowth = 0; // defragment container if it's overhead goes above this value by moving some entries, 0 to ignore
 	uint64 PatchInPlaceParetoFrontCoarseSamples = 0;  // amount of coarse samples to use for pareto front probing during defrag
 	uint64 PatchInPlaceParetoFrontFineSamples = 0;    // amount of fine samples to fine tune pareto front during defrag
+	uint64 PatchInPlaceReadGranularity = 0;           // group inserted entries with-in read granularity bins based on offsets from file open order, 0 to ignore
+	uint64 PatchInPlaceMaxChainSize = 0;              // largest possible insertion chain size, 0 to ignore
+	int32 PatchInPlaceMaxChainLength = 0;             // largest possible insertion chain length, 0 to ignore
 	uint8 PatchInPlaceParetoFrontHeuristic = 0;       // heuristic to choose for pareto front, see FPatchInPlaceLayout::BuildUnpinCandidates
 	TMap<FString, uint64> MaxPartitionSizeOverride;
 	bool bEnableFileRegions = false;
@@ -51,6 +54,7 @@ struct FIoStoreWriterResult
 	uint64 PatchInPlaceSizeBeforeLayout = 0;
 	uint64 PatchInPlacePinnedSizeBeforeDefrag = 0;
 	uint64 PatchInPlacePinnedSizeAfterDefrag = 0;
+	int64 PatchInPlaceLocalityChange = 0;
 	uint64 AddedChunksCount = 0;
 	uint64 AddedChunksSize = 0;
 	uint64 ModifiedChunksCount = 0;
