Fix compilation issues

#rnx
#Synced-CL 51328467

diff --git a/Engine/Source/Developer/IoStoreUtilities/Private/IoStoreWriter.cpp b/Engine/Source/Developer/IoStoreUtilities/Private/IoStoreWriter.cpp
--- a/Engine/Source/Developer/IoStoreUtilities/Private/IoStoreWriter.cpp
+++ b/Engine/Source/Developer/IoStoreUtilities/Private/IoStoreWriter.cpp
@@ -1069,22 +1069,22 @@ private:
 			Size += InSize;
 		}
 
-		uint64 CalculateEnd(uint64 AtOffset, const TConstArrayView<FTocEntry>& EntriesToInsert, const uint64 MemoryMappingAlignment) const
+		uint64 CalculateEnd(uint64 AtOffset, const TConstArrayView<FTocEntry>& InEntriesToInsert, const uint64 MemoryMappingAlignment) const
 		{
 			for (int32 I = Index; I < Index + Count; ++I)
 			{
-				AtOffset = AlignOffset(AtOffset, EntriesToInsert[I].bMemoryMapped, MemoryMappingAlignment) + EntriesToInsert[I].Size;
+				AtOffset = AlignOffset(AtOffset, InEntriesToInsert[I].bMemoryMapped, MemoryMappingAlignment) + InEntriesToInsert[I].Size;
 			}
 
 			return AtOffset;
 		}
 
-		uint64 InsertEntries(uint64 AtOffset, const TArrayView<FTocEntry>& EntriesToInsert, const uint64 MemoryMappingAlignment) const
+		uint64 InsertEntries(uint64 AtOffset, const TArrayView<FTocEntry>& InEntriesToInsert, const uint64 MemoryMappingAlignment) const
 		{
 			for (int32 I = Index; I < Index + Count; ++I)
 			{
-				EntriesToInsert[I].Offset = AlignOffset(AtOffset, EntriesToInsert[I].bMemoryMapped, MemoryMappingAlignment);
-				AtOffset = EntriesToInsert[I].End();
+				InEntriesToInsert[I].Offset = AlignOffset(AtOffset, InEntriesToInsert[I].bMemoryMapped, MemoryMappingAlignment);
+				AtOffset = InEntriesToInsert[I].End();
 			}
 
 			return AtOffset;
