UnrealPak: Disable the fragile double throttling on memory while compressing and writing the containers.

#rnx

#Synced-CL 51335264

diff --git a/Engine/Source/Developer/IoStoreUtilities/Private/IoStoreUtilities.cpp b/Engine/Source/Developer/IoStoreUtilities/Private/IoStoreUtilities.cpp
--- a/Engine/Source/Developer/IoStoreUtilities/Private/IoStoreUtilities.cpp
+++ b/Engine/Source/Developer/IoStoreUtilities/Private/IoStoreUtilities.cpp
@@ -2557,6 +2557,13 @@ public:
 		RetirerThread.Wait();
 	}
 
+	void DisableMemoryThrottling()
+	{
+		MaxSourceBufferMemory = 0;
+		UE_LOG(LogIoStore, Display, TEXT("Disabled WriteRequestManager memory throttling, MaxConcurrentSourceReads=%d (%d cores)"),
+			MaxConcurrentSourceReads, FPlatformMisc::NumberOfCoresIncludingHyperthreads());
+	}
+
 	IIoStoreWriteRequest* Read(const FContainerTargetFile& InTargetFile)
 	{
 		if (InTargetFile.SourceBuffer.IsSet())
@@ -2890,12 +2897,15 @@ private:
 			SourceReadCompletedEvent->Wait();
 		}
 
-		uint64 LocalUsedBufferMemory = UsedBufferMemory.Load();
-		while (LocalUsedBufferMemory > 0 && LocalUsedBufferMemory + SourceBufferSize > MaxSourceBufferMemory)
+		if (MaxSourceBufferMemory)
 		{
-			TRACE_CPUPROFILER_EVENT_SCOPE(WaitForBufferMemory);
-			MemoryAvailableEvent->Wait();
-			LocalUsedBufferMemory = UsedBufferMemory.Load();
+			uint64 LocalUsedBufferMemory = UsedBufferMemory.Load();
+			while (LocalUsedBufferMemory > 0 && LocalUsedBufferMemory + SourceBufferSize > MaxSourceBufferMemory)
+			{
+				TRACE_CPUPROFILER_EVENT_SCOPE(WaitForBufferMemory);
+				MemoryAvailableEvent->Wait();
+				LocalUsedBufferMemory = UsedBufferMemory.Load();
+			}
 		}
 
 		UsedBufferMemory.AddExchange(SourceBufferSize);
@@ -4896,20 +4906,25 @@ static int32 CreateTarget(const FIoStoreArguments& Arguments, const FIoStoreWrit
 		GlobalIoStoreWriter->Append(CreateIoChunkId(0, 0, EIoChunkType::ScriptObjects), ScriptObjectsBuffer, WriteOptions);
 	}
 
-	UE_LOG(LogIoStore, Display, TEXT("Serializing container(s)..."));
+	// Finalize the layout and disable memory throttling for source reads before IoStoreWriterContext->FinalizeWrites that 
+	// has its own memory throttling logic.
+	IoStoreWriterContext->FinalizeLayout();
+	WriteRequestManager.DisableMemoryThrottling();
+
+	UE_LOG(LogIoStore, Display, TEXT("Writing container(s)..."));
 	{
-		IOSTORE_CPU_SCOPE(Serializing);
+		IOSTORE_CPU_SCOPE(WritingContainers);
 
-		TFuture<void> FlushTask = Async(EAsyncExecution::Thread, [&IoStoreWriterContext]()
+		TFuture<void> FinalizeTask = Async(EAsyncExecution::Thread, [&IoStoreWriterContext]()
 		{
-			IoStoreWriterContext->Flush();
+			IoStoreWriterContext->FinalizeWrites();
 		});
 
 		bool bPrintProgress = true;
 		while (bPrintProgress)
 		{
-			FlushTask.WaitFor(FTimespan::FromSeconds(2.0));
-			bPrintProgress = !FlushTask.IsReady();
+			FinalizeTask.WaitFor(FTimespan::FromSeconds(2.0));
+			bPrintProgress = !FinalizeTask.IsReady();
 
 			TStringBuilder<1024> ProgressStringBuilder;
 
@@ -5584,7 +5599,8 @@ int32 CreateContentPatch(const FIoStoreArguments& Arguments, const FIoStoreWrite
 		});
 	}
 
-	IoStoreWriterContext->Flush();
+	IoStoreWriterContext->FinalizeLayout();
+	IoStoreWriterContext->FinalizeWrites();
 	TArray<FIoStoreWriterResult> Results;
 	for (TSharedPtr<IIoStoreWriter> IoStoreWriter : IoStoreWriters)
 	{
diff --git a/Engine/Source/Developer/IoStoreUtilities/Private/IoStoreWriter.cpp b/Engine/Source/Developer/IoStoreUtilities/Private/IoStoreWriter.cpp
--- a/Engine/Source/Developer/IoStoreUtilities/Private/IoStoreWriter.cpp
+++ b/Engine/Source/Developer/IoStoreUtilities/Private/IoStoreWriter.cpp
@@ -491,7 +491,8 @@ public:
 
 	TSharedPtr<IIoStoreWriter> CreateContainer(const TCHAR* InContainerPathAndBaseFileName, const FIoContainerSettings& InContainerSettings);
 
-	void Flush();
+	void FinalizeLayout();
+	void FinalizeWrites();
 
 	FIoStoreWriterContext::FProgress GetProgress() const
 	{
@@ -654,9 +655,14 @@ TSharedPtr<IIoStoreWriter> FIoStoreWriterContext::CreateContainer(const TCHAR* I
 	return Impl->CreateContainer(InContainerPath, InContainerSettings);
 }
 
-void FIoStoreWriterContext::Flush()
+void FIoStoreWriterContext::FinalizeLayout()
 {
-	Impl->Flush();
+	Impl->FinalizeLayout();
+}
+
+void FIoStoreWriterContext::FinalizeWrites()
+{
+	Impl->FinalizeWrites();
 }
 
 FIoStoreWriterContext::FProgress FIoStoreWriterContext::GetProgress() const
@@ -3529,9 +3535,9 @@ TSharedPtr<IIoStoreWriter> FIoStoreWriterContextImpl::CreateContainer(const TCHA
 	return IoStoreWriter;
 }
 
-void FIoStoreWriterContextImpl::Flush()
+void FIoStoreWriterContextImpl::FinalizeLayout()
 {
-	TRACE_CPUPROFILER_EVENT_SCOPE(FIoStoreWriterContext::Flush);
+	TRACE_CPUPROFILER_EVENT_SCOPE(FIoStoreWriterContext::FinalizeLayout);
 	TArray<FIoStoreWriteQueueEntry*> AllEntries;
 	for (TSharedPtr<FIoStoreWriter> IoStoreWriter : IoStoreWriters)
 	{
@@ -3552,8 +3558,12 @@ void FIoStoreWriterContextImpl::Flush()
 			IoStoreWriter->FinalizeLayout();
 		}
 	}
-	// Update list of all entries after having the finilized layouts of each container
-	AllEntries.Reset();
+}
+
+void FIoStoreWriterContextImpl::FinalizeWrites()
+{
+	TRACE_CPUPROFILER_EVENT_SCOPE(FIoStoreWriterContext::FinalizeWrites);
+	TArray<FIoStoreWriteQueueEntry*> AllEntries;
 	for (TSharedPtr<FIoStoreWriter> IoStoreWriter : IoStoreWriters)
 	{
 		AllEntries.Append(IoStoreWriter->Entries);
@@ -3616,7 +3626,7 @@ UE::DerivedData::FCacheKey FIoStoreWriterContextImpl::MakeDDCKey(FIoStoreWriteQu
 	CacheKeySuffix << IoStoreDDCVersion;
 	CacheKeySuffix << Entry->ChunkHash;
 	CacheKeySuffix.Append(FCompression::GetCompressorDDCSuffix(Entry->CompressionMethod));
-	CacheKeySuffix.Appendf(TEXT("%llu_%d_%d_%d"),
+	CacheKeySuffix.Appendf(TEXT("%llu_%llu_%d_%d"),
 		WriterSettings.CompressionBlockSize,
 		CompressionBufferSize,
 		WriterSettings.CompressionMinBytesSaved,
diff --git a/Engine/Source/Developer/IoStoreUtilities/Private/IoStoreWriter.h b/Engine/Source/Developer/IoStoreUtilities/Private/IoStoreWriter.h
--- a/Engine/Source/Developer/IoStoreUtilities/Private/IoStoreWriter.h
+++ b/Engine/Source/Developer/IoStoreUtilities/Private/IoStoreWriter.h
@@ -109,7 +109,8 @@ public:
 
 	[[nodiscard]] FIoStatus Initialize(const FIoStoreWriterSettings& InWriterSettings);
 	TSharedPtr<class IIoStoreWriter> CreateContainer(const TCHAR* InContainerPath, const FIoContainerSettings& InContainerSettings);
-	void Flush();
+	void FinalizeLayout();
+	void FinalizeWrites();
 	FProgress GetProgress() const;
 	uint32 GetErrors() const;
 	void SaveContainerMeta(const FString& FilePath);
