[Shaders] Improve iteration over shader libraries during packaging.

- Build the dependency information that is derived from `UChunkDependencyInfo` and needed for the shader libraries ahead of time instead of putting this into every call to `AssignShaderArchives()`.
- Use new class `FContainerTargetFileIterator` to cache target files to avoid iterating over potentially thousand of container target files multiple times, since we're only interested in the ShaderCodeLibrary entries.
- Sort the cached target files using the tree level from the chunk dependency hierarchy to ensure parent chunks are always iterated before their child nodes.

#rnx

#Synced-CL 51496351

diff --git a/Engine/Source/Developer/IoStoreUtilities/Private/IoStoreShaderLibraryProcessor.cpp b/Engine/Source/Developer/IoStoreUtilities/Private/IoStoreShaderLibraryProcessor.cpp
--- a/Engine/Source/Developer/IoStoreUtilities/Private/IoStoreShaderLibraryProcessor.cpp
+++ b/Engine/Source/Developer/IoStoreUtilities/Private/IoStoreShaderLibraryProcessor.cpp
@@ -641,20 +641,118 @@ bool FIoStoreShaderLibraryProcessor::ConvertToIoStoreShaderLibrary(FContainerTar
 	return true;
 }
 
-// Invokes TargetFileCallback on each ShaderCodeLibrary entry in all container targets.
-static void ForEachShaderCodeLibraryTargetFile(const TArrayView<FContainerTargetSpec*>& ContainerTargets, const TFunction<void(FContainerTargetFile&)>& TargetFileCallback)
+/**
+ * Helper class to iterate over each target file of a specific chunk type in all container targets.
+ * For a container with thousands of target files, there might only be a handful of ShaderCodeLibrary target files and we need to iterate through them in several passes.
+ * So this class caches them and optionally sorts them to speed up iteration later on.
+ */
+class FContainerTargetFileIterator
 {
-	for (FContainerTargetSpec* ContainerTarget : ContainerTargets)
+	TArray<FContainerTargetFile*> CachedTargetFiles;
+
+	/**
+	 * Builds a map that maps from ChunkID to its depth in the chunk dependency tree.
+	 * The depth of the root's ChunkID is 0. The depth of an arbitrary ChunkID in the tree is the number of links to traverse to travel from the root to the ChunkID.
+	 */
+	static void BuildChunkdIdToDependencyTreeDepthMap(const FChunkDependencyTreeNode& TreeNode, TMap<int32, int32>& OutChunkIdToTreeDepthMap, int32 TreeDepth = 0)
+	{
+		OutChunkIdToTreeDepthMap.Add(TreeNode.ChunkID) = TreeDepth;
+		for (const FChunkDependencyTreeNode& ChildTreeNode : TreeNode.ChildNodes)
+		{
+			FContainerTargetFileIterator::BuildChunkdIdToDependencyTreeDepthMap(ChildTreeNode, OutChunkIdToTreeDepthMap, TreeDepth + 1);
+		}
+	}
+
+public:
+	FContainerTargetFileIterator() = default;
+
+	FContainerTargetFileIterator(const TArrayView<FContainerTargetSpec*>& InContainerTargets, EContainerChunkType InChunkType)
+	{
+		ResetCache(InContainerTargets, InChunkType);
+	}
+
+	void ResetCache(const TArrayView<FContainerTargetSpec*>& InContainerTargets, EContainerChunkType InChunkType)
 	{
-		for (FContainerTargetFile& TargetFile : ContainerTarget->TargetFiles)
+		// Cache all target files for the input chunk type.
+		CachedTargetFiles.Reset(CachedTargetFiles.Num());
+		for (FContainerTargetSpec* ContainerTarget : InContainerTargets)
 		{
-			if (TargetFile.ChunkType == EContainerChunkType::ShaderCodeLibrary)
+			for (FContainerTargetFile& TargetFile : ContainerTarget->TargetFiles)
 			{
-				TargetFileCallback(TargetFile);
+				if (TargetFile.ChunkType == InChunkType)
+				{
+					CachedTargetFiles.Add(&TargetFile);
+				}
 			}
 		}
 	}
-}
+
+	/**
+	 * Sorts this->CachedTargetFiles based on the dependency tree information in InContainerTargetToChunkDependencies.
+	 * Each CachedTargetFile is sorted by its container's depth in the dependency tree, from root (depth 0) to leaf.
+	 * When traversing the CachedTargetFiles in this order, we are therefore guaranteed to process the CachedTargetFiles
+	 * in the container of a parent AssetManager Chunk before processing any in the container of a child AssetManager
+	 * Chunk. Parent chunks are always loaded whenever child chunks are loaded. This allows us to remove shaders from
+	 * child chunks based on the full knowledge of what shaders exist in their parent chunks.
+	 */
+	void SortByChunkDependency(const TMap<const FContainerTargetSpec*, FContainerChunkDependency>& InContainerTargetToChunkDependencies)
+	{
+		auto PrintCachedTargetFiles = [this]() -> FString
+			{
+				TStringBuilder<1024> OutString;
+				for (const FContainerTargetFile* TargetFile : CachedTargetFiles)
+				{
+					OutString.Appendf(TEXT("\n  %s"), *TargetFile->NormalizedSourcePath);
+				}
+				return FString(OutString);
+			};
+
+		UE_LOG(LogIoStore, Verbose, TEXT("SortByChunkDependency: Target files *before* sorting:%s"), *PrintCachedTargetFiles());
+
+		if (UChunkDependencyInfo* DependencyInfo = GetMutableDefault<UChunkDependencyInfo>())
+		{
+			// Get chunk dependency tree to sort target files by their level in the tree hiearchy
+			if (const FChunkDependencyTreeNode* ChunkDepTreeNode = DependencyInfo->GetOrBuildChunkDependencyGraph())
+			{
+				TMap<int32, int32> ChunkIdToTreeDepthMap;
+				FContainerTargetFileIterator::BuildChunkdIdToDependencyTreeDepthMap(*ChunkDepTreeNode, ChunkIdToTreeDepthMap);
+
+				auto FindChunkIdTreeDepthForTargetFile = [&InContainerTargetToChunkDependencies, &ChunkIdToTreeDepthMap](const FContainerTargetFile& TargetFile) -> int32
+					{
+						if (const FContainerChunkDependency* ContainerChunkDependency = InContainerTargetToChunkDependencies.Find(TargetFile.ContainerTarget))
+						{
+							if (const int32* TreeDepth = ChunkIdToTreeDepthMap.Find(ContainerChunkDependency->ChunkId))
+							{
+								return *TreeDepth;
+							}
+						}
+
+						// TargetFiles that are not in a container (should be impossible, but not important to this algorithm)
+						// or that are in a container which is not part of the AssetManager's chunk dependency tree (could be possible, depending on the project's config data)
+						// can be processed in any order since they don't have parents. We arbitrarily process them first by setting their level to INDEX_NONE.
+						return INDEX_NONE;
+					};
+
+				// Use StableSort to prevent indeterminism during staging.
+				CachedTargetFiles.StableSort([&FindChunkIdTreeDepthForTargetFile](FContainerTargetFile& Lhs, FContainerTargetFile& Rhs) -> bool {
+					const int32 LhsTreeDepth = FindChunkIdTreeDepthForTargetFile(Lhs);
+					const int32 RhsTreeDepth = FindChunkIdTreeDepthForTargetFile(Rhs);
+					return LhsTreeDepth < RhsTreeDepth;
+				});
+
+				UE_LOG(LogIoStore, Verbose, TEXT("SortByChunkDependency: Target files *after* sorting:%s"), *PrintCachedTargetFiles());
+			}
+		}
+	}
+
+	void ForEachTargetFile(const TFunction<void(FContainerTargetFile&)>& TargetFileCallback)
+	{
+		for (FContainerTargetFile* TargetFile : CachedTargetFiles)
+		{
+			TargetFileCallback(*TargetFile);
+		}
+	}
+};
 
 // Returns true if the specified target file refers to a Global shadermap library
 static bool IsTargetFileGlobalShaderLibrary(const FContainerTargetFile& TargetFile)
@@ -729,11 +827,12 @@ static bool GetPackagedProjectNameAndTargetFormats(TArrayView<FContainerTargetSp
 	}
 
 	// Heuristic to find all shader formats relevant for the container targets.
-	ForEachShaderCodeLibraryTargetFile(InContainerTargets, [&OutProjectName, &OutFormatAndPlatformNames](FContainerTargetFile& TargetFile) -> void {
-		FShaderLibraryNameInfo LibraryNameInfo;
-		LibraryNameInfo.ParseFromFilename(*TargetFile.NormalizedSourcePath);
-		OutFormatAndPlatformNames.AddUnique(LibraryNameInfo.GetFormatAndPlatformName());
-	});
+	FContainerTargetFileIterator{ InContainerTargets, EContainerChunkType::ShaderCodeLibrary }
+		.ForEachTargetFile([&OutProjectName, &OutFormatAndPlatformNames](FContainerTargetFile& TargetFile) -> void {
+			FShaderLibraryNameInfo LibraryNameInfo;
+			LibraryNameInfo.ParseFromFilename(*TargetFile.NormalizedSourcePath);
+			OutFormatAndPlatformNames.AddUnique(LibraryNameInfo.GetFormatAndPlatformName());
+		});
 
 	return !OutProjectName.IsEmpty() && !OutFormatAndPlatformNames.IsEmpty();
 }
@@ -750,6 +849,8 @@ void FIoStoreShaderLibraryProcessor::ProcessShaderLibraries(const FIoStoreArgume
 	ShaderOodleCompressor = Arguments.ShaderOodleCompressor;
 	ShaderOodleLevel = Arguments.ShaderOodleLevel;
 
+	FContainerTargetFileIterator TargetFileIterator;
+
 	if (IsChunkOverrideEnabled())
 	{
 #if WITH_EDITORONLY_DATA
@@ -792,10 +893,21 @@ void FIoStoreShaderLibraryProcessor::ProcessShaderLibraries(const FIoStoreArgume
 			}
 		}
 
+		// Cache all target files and sort them by chunk dependency to ensure we deduplicate redundant shadermaps as tight as possible. 
+		{
+			for (FContainerTargetSpec* ContainerTarget : ContainerTargets)
+			{
+				BuildContainerTargetChunkDependencies(ContainerTarget);
+			}
+
+			TargetFileIterator.ResetCache(ContainerTargets, EContainerChunkType::ShaderCodeLibrary);
+			TargetFileIterator.SortByChunkDependency(ContainerTargetToChunkDependencies);
+		}
+
 		// Preload all shader archives, so we can re-assign shaders to their targets from monolithic libraries
 		{
 			IOSTORE_CPU_SCOPE(PreloadShaderLibraries);
-			ForEachShaderCodeLibraryTargetFile(ContainerTargets, [this, &Arguments, &Output](FContainerTargetFile& TargetFile) -> void {
+			TargetFileIterator.ForEachTargetFile([this, &Arguments, &Output](FContainerTargetFile& TargetFile) -> void {
 				if (!IsTargetFileExcludedFromChunkReassignment(TargetFile) && !PreloadShaderArchive(TargetFile))
 				{
 					UE_LOG(LogIoStore, Warning, TEXT("Failed preloading shader library '%s'"), *TargetFile.NormalizedSourcePath);
@@ -806,7 +918,7 @@ void FIoStoreShaderLibraryProcessor::ProcessShaderLibraries(const FIoStoreArgume
 		// Now assign shader archives to their final target files
 		{
 			IOSTORE_CPU_SCOPE(AssignShaderArchives);
-			ForEachShaderCodeLibraryTargetFile(ContainerTargets, [this, &Arguments, &Output](FContainerTargetFile& TargetFile) -> void {
+			TargetFileIterator.ForEachTargetFile([this, &Arguments, &Output](FContainerTargetFile& TargetFile) -> void {
 				if (!IsTargetFileExcludedFromChunkReassignment(TargetFile) && !AssignShaderArchives(TargetFile))
 				{
 					UE_LOG(LogIoStore, Warning, TEXT("Failed assigning shader archives for target library '%s'"), *TargetFile.NormalizedSourcePath);
@@ -820,11 +932,13 @@ void FIoStoreShaderLibraryProcessor::ProcessShaderLibraries(const FIoStoreArgume
 	else
 	{
 		UE_LOG(LogIoStore, Display, TEXT("Shader library chunk override is disabled"));
+
+		TargetFileIterator.ResetCache(ContainerTargets, EContainerChunkType::ShaderCodeLibrary);
 	}
 
 	{
 		IOSTORE_CPU_SCOPE(ConvertShaderLibraries);
-		ForEachShaderCodeLibraryTargetFile(ContainerTargets, [this, &Arguments, &Output](FContainerTargetFile& TargetFile) -> void {
+		TargetFileIterator.ForEachTargetFile([this, &Arguments, &Output](FContainerTargetFile& TargetFile) -> void {
 			if (!ProcessShaderLibraryFromTargetFile(Arguments, TargetFile, Output))
 			{
 				UE_LOG(LogIoStore, Warning, TEXT("Failed converting shader library '%s'"), *TargetFile.NormalizedSourcePath);
@@ -1030,6 +1144,51 @@ bool FIoStoreShaderLibraryProcessor::ProcessShaderLibraryFromTargetFile(const FI
 
 #if WITH_EDITORONLY_DATA
 
+// Print the set of chunk dependency IDs to string for logging, e.g. "None", "0", "0, 1" etc.
+static FString PrintChunkDependenciesToString(const TSet<int32>& InChunkDependencies)
+{
+	TStringBuilder<1024> OutString;
+	if (InChunkDependencies.IsEmpty())
+	{
+		OutString.Append(TEXT("None"));
+	}
+	else
+	{
+		for (int32 ChunkId : InChunkDependencies)
+		{
+			if (OutString.Len() > 0)
+			{
+				OutString.Append(TEXT(", "));
+			}
+			OutString.Appendf(TEXT("%d"), ChunkId);
+		}
+	}
+	return FString(OutString);
+}
+
+void FIoStoreShaderLibraryProcessor::BuildContainerTargetChunkDependencies(const FContainerTargetSpec* ContainerTarget)
+{
+	if (UChunkDependencyInfo* DependencyInfo = GetMutableDefault<UChunkDependencyInfo>())
+	{
+		FContainerChunkDependency& ContainerChunkDependency = ContainerTargetToChunkDependencies.Add(ContainerTarget);
+
+		FString ContainerTargetChunkIdSuffix;
+		if (FShaderLibraryNameInfo::ParsePakChunkId(ContainerTarget->Name.ToString(), ContainerChunkDependency.ChunkId, ContainerTargetChunkIdSuffix))
+		{
+			// Get chunk dependency tree to determine what shadermaps are already included in a parent chunk
+			DependencyInfo->GetOrBuildChunkDependencyGraph(ContainerChunkDependency.ChunkId);
+			DependencyInfo->GetChunkDependencies(ContainerChunkDependency.ChunkId, ContainerChunkDependency.ChunkDependencies);
+			ContainerChunkDependency.ChunkDependencies.Remove(ContainerChunkDependency.ChunkId);
+			ContainerChunkDependency.bIsExternalContainer = !ContainerTargetChunkIdSuffix.IsEmpty();
+
+			// Log information about chunk dependencies for shader library
+			UE_LOG(LogIoStore, Display, TEXT("Processing chunk dependency for container '%s': ChunkId=%d, bIsExternalContainer=%s, ChunkDependencies=(%s)"),
+				*ContainerTarget->Name.ToString(), ContainerChunkDependency.ChunkId, ContainerChunkDependency.bIsExternalContainer ? TEXT("Yes") : TEXT("No"),
+				*PrintChunkDependenciesToString(ContainerChunkDependency.ChunkDependencies));
+		}
+	}
+}
+
 bool FIoStoreShaderLibraryProcessor::PreloadShaderArchive(const FContainerTargetFile& TargetFile)
 {
 	// Preload serialized shader archive for target file
@@ -1173,28 +1332,6 @@ bool FIoStoreShaderLibraryProcessor::DumpChunkOverrideDebugInfo(FName FormatAndP
 	return true;
 }
 
-// Print the set of chunk dependency IDs to string for logging, e.g. "None", "0", "0, 1" etc.
-static FString PrintChunkDependenciesToString(const TSet<int32>& InChunkDependencies)
-{
-	TStringBuilder<1024> OutString;
-	if (InChunkDependencies.IsEmpty())
-	{
-		OutString.Append(TEXT("None"));
-	}
-	else
-	{
-		for (int32 ChunkId : InChunkDependencies)
-		{
-			if (OutString.Len() > 0)
-			{
-				OutString.Append(TEXT(", "));
-			}
-			OutString.Appendf(TEXT("%d"), ChunkId);
-		}
-	}
-	return FString(OutString);
-}
-
 bool FIoStoreShaderLibraryProcessor::AssignShaderArchives(const FContainerTargetFile& TargetFile)
 {
 	// Convert shader format and platform name from *.ushaderbytecode filename
@@ -1216,33 +1353,13 @@ bool FIoStoreShaderLibraryProcessor::AssignShaderArchives(const FContainerTarget
 	}
 
 	// Gather all chunk dependencies to deduplicate redundant shadermap entries
-	FContainerChunkDependency& ContainerChunkDependency = ContainerTargetToChunkDependencies.FindOrAdd(ContainerTarget);
-	if (ContainerChunkDependency.ChunkId == INDEX_NONE)
-	{
-		FString ContainerTargetChunkIdSuffix;
-		if (FShaderLibraryNameInfo::ParsePakChunkId(ContainerTarget->Name.ToString(), ContainerChunkDependency.ChunkId, ContainerTargetChunkIdSuffix))
-		{
-			if (UChunkDependencyInfo* DependencyInfo = GetMutableDefault<UChunkDependencyInfo>())
-			{
-				// Get chunk dependency graph to determine what shadermaps are already included in a parent chunk
-				DependencyInfo->GetOrBuildChunkDependencyGraph(ContainerChunkDependency.ChunkId);
-				DependencyInfo->GetChunkDependencies(ContainerChunkDependency.ChunkId, ContainerChunkDependency.ChunkDependencies);
-				ContainerChunkDependency.ChunkDependencies.Remove(ContainerChunkDependency.ChunkId);
-				ContainerChunkDependency.bIsExternalContainer = !ContainerTargetChunkIdSuffix.IsEmpty();
-
-				// Log information about chunk dependencies for shader library
-				UE_LOG(LogIoStore, Display, TEXT("Processing chunk dependency for container '%s': ChunkId=%d, bIsExternalContainer=%s, ChunkDependencies=(%s)"),
-					*ContainerTarget->Name.ToString(), ContainerChunkDependency.ChunkId, ContainerChunkDependency.bIsExternalContainer ? TEXT("Yes") : TEXT("No"),
-					*PrintChunkDependenciesToString(ContainerChunkDependency.ChunkDependencies));
-			}
-		}
-	}
+	FContainerChunkDependency* ContainerChunkDependency = ContainerTargetToChunkDependencies.Find(ContainerTarget);
 
 	// Lambda to determine what shadermaps to exclude from current target file to deduplicate them across their chunk dependency tree
-	auto ExcludeShaderMap = [this, &ContainerChunkDependency, ContainerTarget](const FSHAHash& ShaderMapHash) -> bool
+	auto ExcludeShaderMap = [this, ContainerChunkDependency, ContainerTarget](const FSHAHash& ShaderMapHash) -> bool
 		{
 			// Check if shadermap hash is already stored in any of the current container's chunk dependencies
-			for (int32 ChunkDependencyId : ContainerChunkDependency.ChunkDependencies)
+			for (int32 ChunkDependencyId : ContainerChunkDependency->ChunkDependencies)
 			{
 				if (const FContainerShaderArchives* ContainerInfo = this->ChunkIdToContainerInfo.Find(ChunkDependencyId))
 				{
@@ -1267,7 +1384,12 @@ bool FIoStoreShaderLibraryProcessor::AssignShaderArchives(const FContainerTarget
 		check(MonolithicShaderArchive);
 
 		// Create new chunk for packages.
-		TUniquePtr<FEditorShaderCodeArchive> NewChunkShaderArchive = MonolithicShaderArchive->CreateNamedChunk(LibraryNameInfo.ChunkName, PackagesInChunk, ExcludeShaderMap);
+		// All containers were added to ContainerTargetToChunkDependencies by the ProcessShaderLibraries's call to BuildContainerTargetChunkDependencies before calling AssignShaderArchives.
+		FShaderMapAssetAssociations::FShaderMapFilterFunction ShaderMapFilterFunction =
+			ensureMsgf(ContainerChunkDependency != nullptr, TEXT("Failed to find container chunk dependency information for container target '%s'"), *ContainerTarget->Name.ToString())
+				? ExcludeShaderMap
+				: FShaderMapAssetAssociations::FShaderMapFilterFunction(nullptr);
+		TUniquePtr<FEditorShaderCodeArchive> NewChunkShaderArchive = MonolithicShaderArchive->CreateNamedChunk(LibraryNameInfo.ChunkName, PackagesInChunk, ShaderMapFilterFunction);
 		check(NewChunkShaderArchive);
 
 		// Convert constructed intermediate shader archive to final serialized archive with flat buffer
@@ -1287,12 +1409,9 @@ bool FIoStoreShaderLibraryProcessor::AssignShaderArchives(const FContainerTarget
 			// Link shader archive to current container chunk dependency if this is not an external container, i.e. no OnDemand or Optional ID suffix.
 			// This information can be retrieved by subsequent shader archives in `ExcludeShaderMap` using their chunk dependency IDs
 			// to check what shadermaps are already contained in their parent chunks to deduplicate them.
-			// @lh-todo:
-			//  This currently relies on the containers arriving in the order of their dependency, which means this is potentially missing out on a tighter deduplication.
-			//  For the exact order, we need to iterate through all containers according to the dependency graph.
-			if (ContainerChunkDependency.ChunkId != INDEX_NONE && !ContainerChunkDependency.bIsExternalContainer)
+			if (ContainerChunkDependency && ContainerChunkDependency->ChunkId != INDEX_NONE && !ContainerChunkDependency->bIsExternalContainer)
 			{
-				FContainerShaderArchives& ContainerInfo = ChunkIdToContainerInfo.FindOrAdd(ContainerChunkDependency.ChunkId);
+				FContainerShaderArchives& ContainerInfo = ChunkIdToContainerInfo.FindOrAdd(ContainerChunkDependency->ChunkId);
 				ContainerInfo.AssignedSerializedShaders.Add(AssignedShaderArchive.Get());
 			}
 
diff --git a/Engine/Source/Developer/IoStoreUtilities/Private/IoStoreShaderLibraryProcessor.h b/Engine/Source/Developer/IoStoreUtilities/Private/IoStoreShaderLibraryProcessor.h
--- a/Engine/Source/Developer/IoStoreUtilities/Private/IoStoreShaderLibraryProcessor.h
+++ b/Engine/Source/Developer/IoStoreUtilities/Private/IoStoreShaderLibraryProcessor.h
@@ -82,6 +82,20 @@ struct FShaderAssociationInfo
 	}
 };
 
+//@lh-todo
+// Move this information into FContainerSourceSpec and FContainerTargetSpec.
+// This information should be determined in `ParseContainerGenerationArguments()` when the source containers are built,
+// and then transferred to the target containers in `InitializeContainerTargetsAndPackages()`.
+// More testing beyond shaders for this is required, so keep this information private to the shader library processor for now.
+// Same goes for `IsTargetFileGlobalShaderLibrary()`, which should be determined as soon as the source container is generated
+// instead of parsing the filename later on to extract those flags.
+struct FContainerChunkDependency
+{
+	int32 ChunkId = INDEX_NONE;
+	bool bIsExternalContainer = false; // Is the container loaded externally? E.g. pakchunk100iad, pakchunk100ondemand, etc.
+	TSet<int32> ChunkDependencies;
+};
+
 struct FIoStoreShaderLibraryOutput
 {
 	TArray<FShaderInfo*> ShaderInfos;
@@ -115,20 +129,6 @@ class FIoStoreShaderLibraryProcessor
 	TMap<FSHAHash, TArray<FIoChunkId>> ShaderMapToGroupChunkIds; // Maps from shader map hash to its set of shader group IoChunk Ids
 	TMap<FContainerTargetSpec*, TSet<FShaderInfo*>> ContainerTargetToShaderLibrary;
 	FShaderMapAssetAssociations ShaderMapAssetAssociations;
-
-	//@lh-todo
-	// Move this information into FContainerSourceSpec and FContainerTargetSpec.
-	// This information should be determined in `ParseContainerGenerationArguments()` when the source containers are built,
-	// and then transferred to the target containers in `InitializeContainerTargetsAndPackages()`.
-	// More testing beyond shaders for this is required, so keep this information private to the shader library processor for now.
-	// Same goes for `IsTargetFileGlobalShaderLibrary()`, which should be determined as soon as the source container is generated
-	// instead of parsing the filename later on to extract those flags.
-	struct FContainerChunkDependency
-	{
-		int32 ChunkId = INDEX_NONE;
-		bool bIsExternalContainer = false; // Is the container loaded externally? E.g. pakchunk100iad, pakchunk100ondemand, etc.
-		TSet<int32> ChunkDependencies;
-	};
 	TMap<const FContainerTargetSpec*, FContainerChunkDependency> ContainerTargetToChunkDependencies;
 
 	// Intermediate data while processing a subset of target files
@@ -169,6 +169,9 @@ private:
 	bool ProcessShaderLibraryFromTargetFile(const FIoStoreArguments& Arguments, FContainerTargetFile& TargetFile, FIoStoreShaderLibraryOutput& Output);
 
 #if WITH_EDITORONLY_DATA
+	// Builds the chunk dependency information for the specified container target. See ContainerTargetToChunkDependencies.
+	void BuildContainerTargetChunkDependencies(const FContainerTargetSpec* ContainerTarget);
+
 	// Preloads the *.ushaderbytecode for the specified target file for later chunk assignment.
 	bool PreloadShaderArchive(const FContainerTargetFile& TargetFile);
 
