Fix up patch-in-place failures if builds are not found

#rnx
#Synced-CL 51577459

diff --git a/Engine/Source/Developer/IoStoreUtilities/Private/IoStoreUtilities.cpp b/Engine/Source/Developer/IoStoreUtilities/Private/IoStoreUtilities.cpp
--- a/Engine/Source/Developer/IoStoreUtilities/Private/IoStoreUtilities.cpp
+++ b/Engine/Source/Developer/IoStoreUtilities/Private/IoStoreUtilities.cpp
@@ -9180,8 +9180,21 @@ bool ParseContainerGenerationArguments(FIoStoreArguments& Arguments, FIoStoreWri
 			FString CryptoKeysCacheFilename;
 			if (FParse::Value(FCommandLine::Get(), TEXT("-patchInPlaceReferenceContainerCryptoOverride="), CryptoKeysCacheFilename))
 			{
-				UE_LOG(LogIoStore, Display, TEXT("Parsing reference container crypto keys from a crypto key cache file '%s'"), *CryptoKeysCacheFilename);
-				KeyChainUtilities::LoadKeyChainFromFile(CryptoKeysCacheFilename, Arguments.PatchInPlaceReferenceChunkKeys);
+				if (IFileManager::Get().FileExists(*CryptoKeysCacheFilename))
+				{
+					UE_LOG(LogIoStore, Display, TEXT("Parsing reference container crypto keys from a crypto key cache file '%s'"), *CryptoKeysCacheFilename);
+					KeyChainUtilities::LoadKeyChainFromFile(CryptoKeysCacheFilename, Arguments.PatchInPlaceReferenceChunkKeys);
+				}
+				else
+				{
+					UE_LOG(LogIoStore, Warning, TEXT("Patch-in-place reference container crypto key file '%s' does not exist."), *CryptoKeysCacheFilename);
+				}
+
+				if (Arguments.PatchInPlaceReferenceChunkKeys.GetEncryptionKeys().Num() == 0)
+				{
+					UE_LOG(LogIoStore, Warning, TEXT("No encryption keys available for patch-in-place reference containers. Disabling patch-in-place layout."));
+					WriterSettings.bUsePatchInPlaceLayout = false;
+				}
 			}
 
 			FParse::Value(FCommandLine::Get(), TEXT("-patchInPlaceReferenceAdditionalContainersPath="), Arguments.PatchInPlaceReferenceAdditionalContainersPath);
