[Shaders] Include IoChunkIds in debug output for shader chunk override.

Also include IoChunk stats for packaged shadermaps like this:
```
------------------------------------------
Stats
Minimum IoChunks per shadermap: 1
Average IoChunks per shadermap: 8.41
Maximum IoChunks per shadermap: 53
------------------------------------------
```

This also simplifies the output syntax of the debug information to keep it minimalistic:
```
Shadermap 38DE6B9212C0960A388902F4693EB90EE7CDFA1D
Unchanged: pakchunk0-WindowsClient
IoChunks: 86825bd6167b8aedef8b9c09, 5978afb1552637ee55407109, 6fe01c2eae74d45cc74be609, 529b8905144ae99720971e09
```

#rnx
#Synced-CL 51753475

diff --git a/Engine/Source/Developer/IoStoreUtilities/Private/IoStoreShaderLibraryProcessor.cpp b/Engine/Source/Developer/IoStoreUtilities/Private/IoStoreShaderLibraryProcessor.cpp
--- a/Engine/Source/Developer/IoStoreUtilities/Private/IoStoreShaderLibraryProcessor.cpp
+++ b/Engine/Source/Developer/IoStoreUtilities/Private/IoStoreShaderLibraryProcessor.cpp
@@ -1276,20 +1276,24 @@ bool FIoStoreShaderLibraryProcessor::DumpChunkOverrideDebugInfo(FName FormatAndP
 	}
 
 	// Generate debug output text
-	FString DebugOutput;
+	FUtf8String DebugOutput;
 
-	auto AppendPakChunkNames = [&DebugOutput](const TCHAR* InListName, const TArray<FName>& InPakChunks) -> void
+	auto AppendPakChunkNames = [&DebugOutput](const UTF8CHAR* InListName, const TArray<FName>& InPakChunks) -> void
 		{
-			DebugOutput.Appendf(TEXT("\n  %s:"), InListName);
+			DebugOutput.Appendf(UTF8TEXT("\n  %s: "), InListName);
 			if (InPakChunks.IsEmpty())
 			{
-				DebugOutput.Append(TEXT(" <None>"));
+				DebugOutput.Append(UTF8TEXT("<None>"));
 			}
 			else
 			{
-				for (const FName& PakChunk : InPakChunks)
+				for (int32 PakChunkIndex = 0; PakChunkIndex < InPakChunks.Num(); ++PakChunkIndex)
 				{
-					DebugOutput.Appendf(TEXT(" '%s'"), *PakChunk.ToString());
+					if (PakChunkIndex > 0)
+					{
+						DebugOutput.Append(UTF8TEXT(", "));
+					}
+					DebugOutput.Append(TCHAR_TO_UTF8(*InPakChunks[PakChunkIndex].ToString()));
 				}
 			}
 		};
@@ -1310,24 +1314,70 @@ bool FIoStoreShaderLibraryProcessor::DumpChunkOverrideDebugInfo(FName FormatAndP
 			return true;
 		};
 
-	for (const TPair<FSHAHash, FChunkAssignmentInfo>& Pair : ChunkAssignmentDebugInfo[FormatAndPlatformName].ShaderMaps)
+	TOptional<int32> MaxNumIoChunksPerShadermap, MinNumIoChunksPerShadermap;
+	uint64 TotalNumberOfIoChunksReferences = 0; // Used to determine average number of IoChunks per shadermap
+
+	const FMonolithicArchiveDebugInfo& DebugInfo = ChunkAssignmentDebugInfo[FormatAndPlatformName];
+	for (const TPair<FSHAHash, FChunkAssignmentInfo>& Pair : DebugInfo.ShaderMaps)
 	{
-		DebugOutput.Appendf(TEXT("%s {"), *Pair.Key.ToString());
-		if (IsAssignmentUnchanged(Pair.Value))
+		DebugOutput.Appendf("Shadermap %s", TCHAR_TO_UTF8(*Pair.Key.ToString()));
 		{
-			AppendPakChunkNames(TEXT("Unchanged"), Pair.Value.OldPakChunks);
-		}
-		else
-		{
-			AppendPakChunkNames(TEXT("Old"), Pair.Value.OldPakChunks);
-			AppendPakChunkNames(TEXT("New"), Pair.Value.NewPakChunks);
+			// Print old and new containers this shadermap resides in
+			if (IsAssignmentUnchanged(Pair.Value))
+			{
+				AppendPakChunkNames(UTF8TEXT("Unchanged"), Pair.Value.OldPakChunks);
+			}
+			else
+			{
+				AppendPakChunkNames(UTF8TEXT("Old"), Pair.Value.OldPakChunks);
+				AppendPakChunkNames(UTF8TEXT("New"), Pair.Value.NewPakChunks);
+			}
+
+			// Print IoChunks used by current shadermap
+			DebugOutput.Append(UTF8TEXT("\n  IoChunks: "));
+			if (const TArray<FIoChunkId>* IoChunkIds = ShaderMapToGroupChunkIds.Find(Pair.Key))
+			{
+				TotalNumberOfIoChunksReferences += static_cast<uint64>(IoChunkIds->Num());
+				MaxNumIoChunksPerShadermap = FMath::Max<int32>(MaxNumIoChunksPerShadermap.Get(0), IoChunkIds->Num());
+				MinNumIoChunksPerShadermap = FMath::Min<int32>(MinNumIoChunksPerShadermap.Get(MAX_int32), IoChunkIds->Num());
+
+				for (int32 IoChunkIndex = 0; IoChunkIndex < IoChunkIds->Num(); ++IoChunkIndex)
+				{
+					if (IoChunkIndex > 0)
+					{
+						DebugOutput.Append(UTF8TEXT(", "));
+					}
+					DebugOutput.Append(TCHAR_TO_UTF8(*LexToString((*IoChunkIds)[IoChunkIndex])));
+				}
+			}
+			else
+			{
+				DebugOutput.Append(UTF8TEXT("<None>"));
+			}
 		}
-		DebugOutput.Append(TEXT("\n}\n"));
+		DebugOutput.Append("\n\n");
 	}
 
+	const double AvgNumIoChunksPerShadermap = DebugInfo.ShaderMaps.Num() > 0 ? static_cast<double>(TotalNumberOfIoChunksReferences) / static_cast<double>(DebugInfo.ShaderMaps.Num()) : 0.0;
+
+	// Write stats
+	FUtf8String StatsOutput;
+
+	StatsOutput.Appendf(
+		UTF8TEXT("------------------------------------------\n")
+		UTF8TEXT("Stats\n")
+		UTF8TEXT("  Minimum IoChunks per shadermap: %d\n")
+		UTF8TEXT("  Average IoChunks per shadermap: %.2f\n")
+		UTF8TEXT("  Maximum IoChunks per shadermap: %d\n")
+		UTF8TEXT("------------------------------------------\n\n"),
+		MinNumIoChunksPerShadermap.Get(0),
+		AvgNumIoChunksPerShadermap,
+		MaxNumIoChunksPerShadermap.Get(0)
+	);
+
 	// Write debug output to file
-	FTCHARToUTF8 DebugOutputUtf8Converter(*DebugOutput);
-	Ar.Serialize(const_cast<UTF8CHAR*>(reinterpret_cast<const UTF8CHAR*>(DebugOutputUtf8Converter.Get())), DebugOutputUtf8Converter.Length());
+	Ar.Serialize(StatsOutput.GetCharArray().GetData(), StatsOutput.Len() * sizeof(UTF8CHAR));
+	Ar.Serialize(DebugOutput.GetCharArray().GetData(), DebugOutput.Len() * sizeof(UTF8CHAR));
 
 	return true;
 }
