[Shaders] Allow filtering container chunks for shader group compression per platform.

This allows to filter container target chunks per platform that determines whether their shader groups are compressed or not.
This is used to allow enabling compression only for IAD containers on Metal and continue to use uncompressed metallibs in conjunction with memory-mappinged files for non-IAD containers.

#rnx
#Synced-CL 51883210

diff --git a/Engine/Source/Developer/IoStoreUtilities/Private/IoStoreShaderLibraryProcessor.cpp b/Engine/Source/Developer/IoStoreUtilities/Private/IoStoreShaderLibraryProcessor.cpp
--- a/Engine/Source/Developer/IoStoreUtilities/Private/IoStoreShaderLibraryProcessor.cpp
+++ b/Engine/Source/Developer/IoStoreUtilities/Private/IoStoreShaderLibraryProcessor.cpp
@@ -36,6 +36,8 @@ static TAutoConsoleVariable<bool> CVarShaderLibraryPakChunkOverrideDumpDebugInfo
 	TEXT("Also requires `r.ShaderCodeLibrary.PakChunkOverride=true` or command line argument `-OverrideShaderChunks`.")
 );
 
+static const TCHAR* GProjectPackagingSettingConfigSection = TEXT("/Script/UnrealEd.ProjectPackagingSettings");
+
 // Returns true if the project being packaged is meant to generate chunks for the specified platform.
 // Without chunks being generated, there is nothing to re-assign shader chunks to.
 // This might even fail if attempting to re-assign shaders to a single large (> 4 GB) consecutive buffer.
@@ -47,7 +49,7 @@ static bool IsGenerateChunksEnabled()
 		FConfigFile PlatformIniFile;
 		FConfigCacheIni::LoadLocalIniFile(PlatformIniFile, TEXT("Game"), true, *PlatformIniName);
 		bool bGenerateChunks = false;
-		if (PlatformIniFile.GetBool(TEXT("/Script/UnrealEd.ProjectPackagingSettings"), TEXT("bGenerateChunks"), bGenerateChunks))
+		if (PlatformIniFile.GetBool(GProjectPackagingSettingConfigSection, TEXT("bGenerateChunks"), bGenerateChunks))
 		{
 			return bGenerateChunks;
 		}
@@ -84,20 +86,33 @@ static bool IsChunkOverrideDebugEnabled()
 	return CVarShaderLibraryPakChunkOverrideDumpDebugInfo.GetValueOnAnyThread();
 }
 
-static bool ShouldCompressShaderGroups()
+// Returns whether shader group compression is enabled for the specified container target chunk.
+// Defaults to true, but some platforms only allow it for selected chunks.
+static bool ShouldCompressShaderGroups(FStringView ContainerTargetChunkName)
 {
-	FString PlatformIniName;
-	if (FParse::Value(FCommandLine::Get(), TEXT("platform="), PlatformIniName, false))
+	struct FShaderGroupCompressionConfig
 	{
-		FConfigFile PlatformIniFile;
-		FConfigCacheIni::LoadLocalIniFile(PlatformIniFile, TEXT("Game"), true, *PlatformIniName);
-		bool bDisallowShaderGroupCompression = false;
-		if (PlatformIniFile.GetBool(TEXT("/Script/UnrealEd.ProjectPackagingSettings"), TEXT("bDisallowShaderGroupCompression"), bDisallowShaderGroupCompression))
+		FShaderGroupCompressionConfig()
 		{
-			return !bDisallowShaderGroupCompression;
+			FString PlatformIniName;
+			if (FParse::Value(FCommandLine::Get(), TEXT("platform="), PlatformIniName, false))
+			{
+				FConfigFile PlatformIniFile;
+				FConfigCacheIni::LoadLocalIniFile(PlatformIniFile, TEXT("Game"), true, *PlatformIniName);
+				if (PlatformIniFile.GetBool(GProjectPackagingSettingConfigSection, TEXT("bFilterShaderGroupCompression"), bIsFiltered) && bIsFiltered)
+				{
+					PlatformIniFile.GetArray(GProjectPackagingSettingConfigSection, TEXT("FilteredShaderGroupCompressionContainers"), FilteredContainerNames);
+				}
+			}
 		}
-	}
-	return true;
+
+		bool bIsFiltered = false;
+		TArray<FString> FilteredContainerNames;
+	};
+
+	// Compress shader groups if filtering is disabled or if the container chunk name is included in the list of filtered containers
+	static FShaderGroupCompressionConfig ShaderGroupCompressionConfig;
+	return !ShaderGroupCompressionConfig.bIsFiltered || ShaderGroupCompressionConfig.FilteredContainerNames.Contains(ContainerTargetChunkName);
 }
 
 // Loads an *.assetinfo.json text file either from local file system or Zen store.
@@ -545,7 +560,8 @@ bool FIoStoreShaderLibraryProcessor::ConvertToIoStoreShaderLibrary(FContainerTar
 	const int32 GroupCount = IoStoreLibraryHeader.ShaderGroupEntries.Num();
 	CodeIoChunks.SetNum(GroupCount);
 	
-	const bool bCompressShaderGroup = ShouldCompressShaderGroups();
+	// Is shader group compression for this container target file enabled?
+	const bool bCompressShaderGroup = ShouldCompressShaderGroups(LibraryNameInfo.ChunkName);
 
 	ParallelFor(CodeIoChunks.Num(),
 		[this, &IoStoreLibraryHeader, &SerializedShaders, bCompressShaderGroup](int32 GroupIndex)
