[ZenStoreWriter/MTSAVE]
- Fix some thread-safety issue
- Use tasks instead of future so that retraction can be triggered to avoid deadlock when all workers are busy
- Make the commit multithreaded to avoid OOM when adding faster than the single thread could process

#Synced-CL 51868793

diff --git a/Engine/Source/Developer/IoStoreUtilities/Private/ZenStoreWriter.cpp b/Engine/Source/Developer/IoStoreUtilities/Private/ZenStoreWriter.cpp
--- a/Engine/Source/Developer/IoStoreUtilities/Private/ZenStoreWriter.cpp
+++ b/Engine/Source/Developer/IoStoreUtilities/Private/ZenStoreWriter.cpp
@@ -130,67 +130,6 @@ struct FZenStoreWriter::FZenCommitInfo
 	TUniquePtr<FPendingPackageState> PackageState;
 };
 
-class FZenStoreWriter::FCommitQueue
-{
-public:
-	FCommitQueue()
-		: NewCommitEvent(EEventMode::AutoReset) {}
-	
-	void Enqueue(FZenCommitInfo&& Commit)
-	{
-		bool TriggerEvent = false;
-		{
-			FScopeLock _(&QueueCriticalSection);
-			TriggerEvent = Queue.IsEmpty();
-			Queue.Enqueue(MoveTemp(Commit));
-			
-		}
-
-		if (TriggerEvent)
-		{
-			NewCommitEvent->Trigger();
-		}
-	}
-
-	bool BlockAndDequeue(FZenCommitInfo& OutCommit)
-	{
-		for (;;)
-		{
-			{
-				FScopeLock _(&QueueCriticalSection);
-				if (Queue.Dequeue(OutCommit))
-				{
-					return true;
-				}
-			}
-
-			if (bCompleteAdding)
-			{
-				return false;
-			}
-
-			NewCommitEvent->Wait();
-		}
-	}
-
-	void CompleteAdding()
-	{
-		bCompleteAdding = true;
-		NewCommitEvent->Trigger();
-	}
-
-	void ResetAdding()
-	{
-		bCompleteAdding = false;
-	}
-
-private:
-	FEventRef					NewCommitEvent;
-	FCriticalSection			QueueCriticalSection;
-	TQueue<FZenCommitInfo>		Queue;
-	TAtomic<bool>				bCompleteAdding{false};
-};
-
 TArray<const UTF8CHAR*> FZenStoreWriter::ReservedOplogKeys;
 
 void FZenStoreWriter::StaticInit()
@@ -259,19 +198,17 @@ FZenStoreWriter::FZenStoreWriter(
 
 	ZenFileSystemManifest = MakeUnique<FZenFileSystemManifest>(TargetPlatform, OutputPath);
 
-	CommitQueue = MakeUnique<FCommitQueue>();
-
 	Compressor = FOodleDataCompression::ECompressor::Mermaid;
 	CompressionLevel = FOodleDataCompression::ECompressionLevel::VeryFast;
 }
 
 FZenStoreWriter::~FZenStoreWriter()
 {
-	if (CommitThread.IsValid())
+	FEventCountToken Token = CommitPendingEvent.PrepareWait();
+	if (NumCommitPending.load())
 	{
-		UE_LOG(LogZenStoreWriter, Display, TEXT("Aborted, flushing..."));
-		CommitQueue->CompleteAdding();
-		CommitThread.Wait();
+		UE_LOG(LogZenStoreWriter, Display, TEXT("Flushing pending commits..."));
+		CommitPendingEvent.Wait(Token);
 	}
 
 	FScopeLock _(&PackagesCriticalSection);
@@ -323,7 +260,7 @@ void FZenStoreWriter::WritePackageData(const FPackageInfo& Info, FLargeMemoryWri
 
 	FCbObjectId ChunkOid = ToObjectId(Info.ChunkId);
 
-	Entry.CompressedPayload = Async(EAsyncExecution::TaskGraph, [this, PackageBuffer]()
+	Entry.CompressedPayload = UE::Tasks::Launch(TEXT("CompressPayload"), [this, PackageBuffer]()
 	{ 
 		return FCompressedBuffer::Compress(FSharedBuffer::MakeView(PackageBuffer.GetView()), Compressor, CompressionLevel);
 	});
@@ -349,7 +286,7 @@ void FZenStoreWriter::WriteIoStorePackageData(const FPackageInfo& Info, const FI
 
 	PackageData.EnsureOwned();
 
-	Entry.CompressedPayload = Async(EAsyncExecution::TaskGraph, [this, PackageData]()
+	Entry.CompressedPayload = UE::Tasks::Launch(TEXT("CompressPayload"), [this, PackageData]()
 	{ 
 		return FCompressedBuffer::Compress(FSharedBuffer::MakeView(PackageData.GetView()), Compressor, CompressionLevel);
 	});
@@ -371,7 +308,7 @@ void FZenStoreWriter::WriteBulkData(const FBulkDataInfo& Info, const FIoBuffer&
 
 	BulkData.EnsureOwned();
 
-	BulkEntry.CompressedPayload = Async(EAsyncExecution::TaskGraph, [this, BulkData]()
+	BulkEntry.CompressedPayload = UE::Tasks::Launch(TEXT("CompressPayload"), [this, BulkData]()
 	{ 
 		return FCompressedBuffer::Compress(FSharedBuffer::MakeView(BulkData.GetView()), Compressor, CompressionLevel);
 	});
@@ -399,7 +336,7 @@ void FZenStoreWriter::WriteAdditionalFile(const FAdditionalFileInfo& Info, const
 	checkNoEntry();
 #else
 	UE::SavePackageUtilities::IncrementOutstandingAsyncWrites();
-	FileEntry.CompressedPayload = Async(EAsyncExecution::TaskGraph,
+	FileEntry.CompressedPayload = UE::Tasks::Launch(TEXT("CompressPayLoad"),
 		[this, FileDataSharedBuffer, Filename=Info.Filename]()
 	{
 		ON_SCOPE_EXIT
@@ -434,11 +371,14 @@ void FZenStoreWriter::WriteAdditionalFile(const FAdditionalFileInfo& Info, const
 	});
 #endif
 
-	const FZenFileSystemManifestEntry& ManifestEntry = ZenFileSystemManifest->CreateManifestEntry(Info.Filename);
-	FileEntry.Info					= Info;
-	FileEntry.Info.ChunkId			= ManifestEntry.FileChunkId;
-	FileEntry.ZenManifestServerPath = ManifestEntry.ServerPath;
-	FileEntry.ZenManifestClientPath = ManifestEntry.ClientPath;
+	{
+		UE::TUniqueLock ScopeLock(ZenFileSystemManifestMutex);
+		const FZenFileSystemManifestEntry& ManifestEntry = ZenFileSystemManifest->CreateManifestEntry(Info.Filename);
+		FileEntry.Info					= Info;
+		FileEntry.Info.ChunkId			= ManifestEntry.FileChunkId;
+		FileEntry.ZenManifestServerPath = ManifestEntry.ServerPath;
+		FileEntry.ZenManifestClientPath = ManifestEntry.ClientPath;
+	}
 
 	if (bProvidePerPackageResults)
 	{
@@ -619,6 +559,7 @@ void FZenStoreWriter::Initialize(const FCookInfo& Info)
 							if (!CopyFromPreviousCookDenyListExtensions.ContainsByHash(
 								GetTypeHash(Extension), Extension))
 							{
+								UE::TUniqueLock LockScope(ZenFileSystemManifestMutex);
 								ZenFileSystemManifest->AddManifestEntry(FileChunkId, MoveTemp(ServerPath),
 									MoveTemp(ClientPath));
 							}
@@ -663,41 +604,31 @@ void FZenStoreWriter::BeginCook(const FCookInfo& Info)
 
 	if (CookMode == ICookedPackageWriter::FCookInfo::CookOnTheFlyMode)
 	{
-		ZenFileSystemManifest->Generate(MetadataDirectoryPath, Info.ReferencedPlugins);
+		{
+			UE::TUniqueLock LockScope(ZenFileSystemManifestMutex);
+			ZenFileSystemManifest->Generate(MetadataDirectoryPath, Info.ReferencedPlugins);
+		}
 		TIoStatusOr<uint64> Status = HttpClient->AppendOp(CreateProjectMetaDataOpPackage("CookOnTheFly"));
 		UE_CLOG(!Status.IsOk(), LogZenStoreWriter, Error, TEXT("Failed to append OpLog. Reason: %s"), *Status.Status().ToString());
 	}
-
-	if (FPlatformProcess::SupportsMultithreading())
-	{
-		CommitQueue->ResetAdding();
-		CommitThread = AsyncThread([this]()
-		{ 
-			for(;;)
-			{
-				TRACE_CPUPROFILER_EVENT_SCOPE(FZenStoreWriter::WaitingOnCooker);
-				FZenCommitInfo Commit;
-				if (!CommitQueue->BlockAndDequeue(Commit))
-				{
-					break;
-				}
-				CommitPackageInternal(MoveTemp(Commit));
-			}
-		});
-	}
 }
 
 void FZenStoreWriter::EndCook(const FCookInfo& Info)
 {
 	UE_LOG(LogZenStoreWriter, Display, TEXT("Flushing..."));
 	
-	CommitQueue->CompleteAdding();
-	CommitThread.Wait();
-	CommitThread.Reset();
+	FEventCountToken Token = CommitPendingEvent.PrepareWait();
+	if (NumCommitPending.load())
+	{
+		CommitPendingEvent.Wait(Token);
+	}
 
 	if (!Info.bWorkerOnSharedSandbox)
 	{
-		ZenFileSystemManifest->Generate(MetadataDirectoryPath, Info.ReferencedPlugins);
+		{
+			UE::TUniqueLock LockScope(ZenFileSystemManifestMutex);
+			ZenFileSystemManifest->Generate(MetadataDirectoryPath, Info.ReferencedPlugins);
+		}
 		if (FCbPackage ReferencedSetPackage = CreateReferencedSetOpPackage())
 		{
 			TIoStatusOr<uint64> Status = HttpClient->AppendOp(MoveTemp(ReferencedSetPackage));
@@ -793,7 +724,18 @@ void FZenStoreWriter::CommitPackage(FCommitPackageInfo&& Info)
 	FZenCommitInfo ZenCommitInfo{ Forward<FCommitPackageInfo>(Info), MoveTemp(PackageState) };
 	if (FPlatformProcess::SupportsMultithreading())
 	{
-		CommitQueue->Enqueue(MoveTemp(ZenCommitInfo));
+		NumCommitPending++;
+		UE::Tasks::Launch(
+			TEXT("CommitPackageTask"),
+			[this, ZenCommitInfo = MoveTemp(ZenCommitInfo)]() mutable
+			{
+				CommitPackageInternal(MoveTemp(ZenCommitInfo));
+				if (--NumCommitPending == 0)
+				{
+					CommitPendingEvent.Notify();
+				}
+			}
+		);
 	}
 	else
 	{
@@ -928,7 +870,7 @@ void FZenStoreWriter::CommitPackageInternal(FZenCommitInfo&& ZenCommitInfo)
 
 			for (FPackageDataEntry& PkgData : PackageState->PackageData)
 			{
-				FCompressedBuffer Payload = PkgData.CompressedPayload.Get();
+				FCompressedBuffer Payload = PkgData.CompressedPayload.GetResult();
 				if (bComputeHash)
 				{
 					PackageState->PackageHashes->ChunkHashes.Add(PkgData.Info.ChunkId, Payload.GetRawHash());
@@ -956,7 +898,7 @@ void FZenStoreWriter::CommitPackageInternal(FZenCommitInfo&& ZenCommitInfo)
 
 			for (FBulkDataEntry& Bulk : PackageState->BulkData)
 			{
-				FCompressedBuffer Payload = Bulk.CompressedPayload.Get();
+				FCompressedBuffer Payload = Bulk.CompressedPayload.GetResult();
 				if (bComputeHash)
 				{
 					PackageState->PackageHashes->ChunkHashes.Add(Bulk.Info.ChunkId, Payload.GetRawHash());
@@ -987,10 +929,10 @@ void FZenStoreWriter::CommitPackageInternal(FZenCommitInfo&& ZenCommitInfo)
 			{
 				if (bComputeHash)
 				{
-					PackageState->PackageHashes->ChunkHashes.Add(File.Info.ChunkId, File.CompressedPayload.Get().GetRawHash());
+					PackageState->PackageHashes->ChunkHashes.Add(File.Info.ChunkId, File.CompressedPayload.GetResult().GetRawHash());
 				}
 
-				FCbAttachment FileDataAttachment(File.CompressedPayload.Get());
+				FCbAttachment FileDataAttachment(File.CompressedPayload.GetResult());
 				PkgHashGen.Update(FileDataAttachment.GetHash().GetBytes(), sizeof(FIoHash::ByteArray));
 				OplogEntry.AddAttachment(FileDataAttachment);
 
@@ -1079,14 +1021,14 @@ void FZenStoreWriter::CommitPackageInternal(FZenCommitInfo&& ZenCommitInfo)
 		
 		for (FPackageDataEntry& PkgData : PackageState->PackageData)
 		{
-			FCompressedBuffer Payload = PkgData.CompressedPayload.Get();
+			FCompressedBuffer Payload = PkgData.CompressedPayload.GetResult();
 			FIoHash IoHash = Payload.GetRawHash();
 			PkgHashGen.Update(IoHash.GetBytes(), sizeof(FIoHash::ByteArray));
 		}
 		
 		for (FBulkDataEntry& Bulk : PackageState->BulkData)
 		{
-			FCompressedBuffer Payload = Bulk.CompressedPayload.Get();
+			FCompressedBuffer Payload = Bulk.CompressedPayload.GetResult();
 			FIoHash IoHash = Payload.GetRawHash();
 			PackageState->PackageHashes->ChunkHashes.Add(Bulk.Info.ChunkId, IoHash);
 			PkgHashGen.Update(IoHash.GetBytes(), sizeof(FIoHash::ByteArray));
@@ -1094,7 +1036,7 @@ void FZenStoreWriter::CommitPackageInternal(FZenCommitInfo&& ZenCommitInfo)
 		
 		for (FFileDataEntry& File : PackageState->FileData)
 		{
-			FCompressedBuffer Payload = File.CompressedPayload.Get();
+			FCompressedBuffer Payload = File.CompressedPayload.GetResult();
 			FIoHash IoHash = Payload.GetRawHash();
 			PackageState->PackageHashes->ChunkHashes.Add(File.Info.ChunkId, IoHash);
 			PkgHashGen.Update(IoHash.GetBytes(), sizeof(FIoHash::ByteArray));
@@ -1684,6 +1626,7 @@ bool FZenStoreWriter::TryReadMPCookMessageForPackage(FName PackageName, FCbObjec
 	TArray<FString> AdditionalFiles;
 	if (LoadFromCompactBinary(Message["AdditionalFiles"], AdditionalFiles))
 	{
+		UE::TUniqueLock LockScope(ZenFileSystemManifestMutex);
 		for (const FString& Filename : AdditionalFiles)
 		{
 			ZenFileSystemManifest->CreateManifestEntry(Filename);
@@ -1729,6 +1672,7 @@ FZenStoreWriter::FPendingPackageState& FZenStoreWriter::AddPendingPackage(const
 FCbPackage FZenStoreWriter::CreateReferencedSetOpPackage()
 {
 #if WITH_EDITOR
+	UE::TUniqueLock LockScope(ZenFileSystemManifestMutex);
 	TOptional<FZenFileSystemManifestEntry> ReferencedSet = ZenFileSystemManifest->GetReferencedSet();
 	if (!ReferencedSet)
 	{
@@ -1770,6 +1714,7 @@ void FZenStoreWriter::CreateProjectMetaData(FCbPackage& Pkg, FCbWriter& PackageO
 {
 	// File Manifest
 	{
+		UE::TUniqueLock LockScope(ZenFileSystemManifestMutex);
 		if (ZenFileSystemManifest->NumEntries() > 0)
 		{
 			TArrayView<FZenFileSystemManifestEntry const> Entries = ZenFileSystemManifest->ManifestEntries();
diff --git a/Engine/Source/Developer/IoStoreUtilities/Public/PackageStoreWriter.h b/Engine/Source/Developer/IoStoreUtilities/Public/PackageStoreWriter.h
--- a/Engine/Source/Developer/IoStoreUtilities/Public/PackageStoreWriter.h
+++ b/Engine/Source/Developer/IoStoreUtilities/Public/PackageStoreWriter.h
@@ -46,7 +46,7 @@ public:
 		const FPackageStoreEntryResource& Entry;
 	};
 
-	DECLARE_EVENT_OneParam(IPackageStoreWriter, FEntryCreatedEvent, const FEntryCreatedEventArgs&);
+	DECLARE_TS_MULTICAST_DELEGATE_OneParam(FEntryCreatedEvent, const FEntryCreatedEventArgs&);
 	virtual FEntryCreatedEvent& OnEntryCreated() = 0;
 
 	/**
@@ -65,7 +65,7 @@ public:
 	/**
 	 * Broadcasted after a package has been committed, i.e cooked.
 	 */
-	DECLARE_EVENT_OneParam(IPackageStoreWriter, FCommitEvent, const FCommitEventArgs&);
+	DECLARE_TS_MULTICAST_DELEGATE_OneParam(FCommitEvent, const FCommitEventArgs&);
 	virtual FCommitEvent& OnCommit() = 0;
 
 	struct FMarkUpToDateEventArgs
@@ -79,6 +79,6 @@ public:
 	/**
 	 * Broadcasted after a set of packages have been found to be up to date.
 	 */
-	DECLARE_EVENT_OneParam(IPackageStoreWriter, FMarkUpToDateEvent, const FMarkUpToDateEventArgs&);
+	DECLARE_TS_MULTICAST_DELEGATE_OneParam(FMarkUpToDateEvent, const FMarkUpToDateEventArgs&);
 	virtual FMarkUpToDateEvent& OnMarkUpToDate() = 0;
 };
\ No newline at end of file
diff --git a/Engine/Source/Developer/IoStoreUtilities/Public/ZenStoreWriter.h b/Engine/Source/Developer/IoStoreUtilities/Public/ZenStoreWriter.h
--- a/Engine/Source/Developer/IoStoreUtilities/Public/ZenStoreWriter.h
+++ b/Engine/Source/Developer/IoStoreUtilities/Public/ZenStoreWriter.h
@@ -4,6 +4,7 @@
 
 #include "Async/Async.h"
 #include "Async/Future.h"
+#include "Async/EventCount.h"
 #include "Compression/CompressedBuffer.h"
 #include "Compression/OodleDataCompression.h"
 #include "Containers/Array.h"
@@ -24,6 +25,7 @@
 #include "Serialization/AsyncLoading2.h"
 #include "Serialization/CompactBinary.h"
 #include "Serialization/PackageWriter.h"
+#include "Tasks/Task.h"
 #include "Templates/RefCounting.h"
 #include "Templates/SharedPointer.h"
 #include "Templates/UniquePtr.h"
@@ -182,7 +184,7 @@ public:
 private:
 	struct FBulkDataEntry
 	{
-		TFuture<FCompressedBuffer> CompressedPayload;
+		UE::Tasks::TTask<FCompressedBuffer> CompressedPayload;
 		FBulkDataInfo Info;
 		FCbObjectId ChunkId;
 		TArray<FFileRegion> FileRegions;
@@ -193,7 +195,7 @@ private:
 	{
 		~FPackageDataEntry();
 
-		TFuture<FCompressedBuffer> CompressedPayload;
+		UE::Tasks::TTask<FCompressedBuffer> CompressedPayload;
 		FPackageInfo Info;
 		FCbObjectId ChunkId;
 		TUniquePtr<FPackageStorePackage> OptimizedPackage;
@@ -203,7 +205,7 @@ private:
 
 	struct FFileDataEntry
 	{
-		TFuture<FCompressedBuffer> CompressedPayload;
+		UE::Tasks::TTask<FCompressedBuffer> CompressedPayload;
 		FAdditionalFileInfo Info;
 		FString ZenManifestServerPath;
 		FString ZenManifestClientPath;
@@ -278,11 +280,14 @@ private:
 	TArray<FOplogCookInfo>				CookedPackagesInfo;
 	TMap<FName, int32>					PackageNameToIndex;
 
+	UE::FMutex				ZenFileSystemManifestMutex;
 	TUniquePtr<FZenFileSystemManifest>	ZenFileSystemManifest;
 	TMap<FName, TArray<FString>>		PackageAdditionalFiles;
 
 	UE::PackageWriter::Private::ICookerInterface* Cooker = nullptr;
 
+	std::atomic<int32>					NumCommitPending;
+	UE::FEventCount						CommitPendingEvent;
 	FEntryCreatedEvent					EntryCreatedEvent;
 	FCriticalSection					CommitEventCriticalSection;
 	FCommitEvent						CommitEvent;
@@ -293,11 +298,6 @@ private:
 	FOodleDataCompression::ECompressor			Compressor;				
 	FOodleDataCompression::ECompressionLevel	CompressionLevel;	
 	
-	class FCommitQueue;
-
-	TUniquePtr<FCommitQueue>			CommitQueue;
-	TFuture<void>						CommitThread;
-
 	bool								bInitialized;
 	bool								bInitializedConnection;
 	bool								bProvidePerPackageResults;
