Make UnrealPak tolerant to invalid reference databases

#rnx
#Synced-CL 51937299

diff --git a/Engine/Source/Developer/IoStoreUtilities/Private/IoStoreUtilities.cpp b/Engine/Source/Developer/IoStoreUtilities/Private/IoStoreUtilities.cpp
--- a/Engine/Source/Developer/IoStoreUtilities/Private/IoStoreUtilities.cpp
+++ b/Engine/Source/Developer/IoStoreUtilities/Private/IoStoreUtilities.cpp
@@ -235,7 +235,7 @@ public:
 		TUniquePtr<FIoStoreReader> GlobalReader = CreateIoStoreReader(*GlobalFileName, InDecryptionKeychain);
 		if (GlobalReader.IsValid() == false)
 		{
-			UE_LOG(LogIoStore, Error, TEXT("Failed to open reference chunk container %s"), *GlobalFileName);
+			UE_LOG(LogIoStore, Warning, TEXT("Failed to open reference chunk container %s"), *GlobalFileName);
 			return false;
 		}
 
@@ -269,7 +269,7 @@ public:
 			TUniquePtr<FIoStoreReader> Reader = CreateIoStoreReader(*ContainerFilePath, InDecryptionKeychain);
 			if (Reader.IsValid() == false)
 			{
-				UE_LOG(LogIoStore, Error, TEXT("Failed to open reference chunk container %s"), *ContainerFilePath);
+				UE_LOG(LogIoStore, Warning, TEXT("Failed to open reference chunk container %s"), *ContainerFilePath);
 				return false;
 			}
 
@@ -322,7 +322,7 @@ public:
 			}
 			else if (Reader->GetCompressionBlockSize() != CompressionBlockSize)
 			{
-				UE_LOG(LogIoStore, Error, TEXT("Reference chunk containers had different compression block sizes, failing to init reference db (%u and %u)"), CompressionBlockSize, Reader->GetCompressionBlockSize());
+				UE_LOG(LogIoStore, Warning, TEXT("Reference chunk containers had different compression block sizes, failing to init reference db (%u and %u)"), CompressionBlockSize, Reader->GetCompressionBlockSize());
 				return false;
 			}
 
@@ -4574,7 +4574,8 @@ static int32 CreateTarget(const FIoStoreArguments& Arguments, const FIoStoreWrit
 
 		if (ChunkDbPtr->Init(Arguments.ReferenceChunkGlobalContainerFileName, Arguments.ReferenceChunkAdditionalContainersPath, Arguments.ReferenceChunkKeys) == false)
 		{
-			UE_LOG(LogIoStore, Warning, TEXT("Failed to initialize reference chunk store. Pak will continue."));
+			// broken reference database will likely result in a large patch size, treat it as error
+			UE_LOG(LogIoStore, Error, TEXT("Failed to initialize reference chunk store. Pak will continue."));
 			ChunkDatabase.Reset();
 		}
 	}
@@ -4587,7 +4588,7 @@ static int32 CreateTarget(const FIoStoreArguments& Arguments, const FIoStoreWrit
 
 		if (ChunkDbPtr->Init(Arguments.PatchInPlaceReferenceContainerPathOverride, Arguments.PatchInPlaceReferenceAdditionalContainersPath, Arguments.PatchInPlaceReferenceChunkKeys) == false)
 		{
-			UE_LOG(LogIoStore, Warning, TEXT("Failed to initialize reference chunk store. Pak will continue."));
+			UE_LOG(LogIoStore, Warning, TEXT("Failed to initialize patch-in-place reference chunk store. Pak will continue."));
 			PatchInPlaceChunkDatabase.Reset();
 		}
 	}
