Fix patch-in-place bytes not matching the reference due to partition padding. Base pinned metric on the bytes matching.
#rnx
#Synced-CL 52463327

diff --git a/Engine/Source/Developer/IoStoreUtilities/Private/IoStoreWriter.cpp b/Engine/Source/Developer/IoStoreUtilities/Private/IoStoreWriter.cpp
--- a/Engine/Source/Developer/IoStoreUtilities/Private/IoStoreWriter.cpp
+++ b/Engine/Source/Developer/IoStoreUtilities/Private/IoStoreWriter.cpp
@@ -808,11 +808,18 @@ public:
 	bool PrepareForLayout(
 		const uint64 SizeBeforeLayout,
 		uint64& OutPinnedSizeBeforeDefrag,
-		uint64& OutPinnedSizeAfterDefrag
+		int32& OutTargetPartitionNum,
+		uint64& OutTotalPadding,
+		int64& OutLocalityChange
 	)
 	{
 		TRACE_CPUPROFILER_EVENT_SCOPE(PrepareForLayout);
 
+		OutPinnedSizeBeforeDefrag = 0;
+		OutTargetPartitionNum = 0;
+		OutTotalPadding = 0;
+		OutLocalityChange = 0;
+
 		const FIoStatus ReferenceTocValid = FIoStoreTocResourceView::Read(*ReferenceContainerPath, EIoStoreTocReadOptions::ReadAll, ReferenceToc, ReferenceTocStorage);
 		if (!ReferenceTocValid.IsOk())
 		{
@@ -826,7 +833,9 @@ public:
 		// e.g. 0% would mean nothing is stays in place, 100% would mean everything stays in place
 		// this is needed to let data move around to shrink containers between big releases
 
-		PinChunksEntries();
+		TArray<FTocEntry> EntriesToKeep; // sorted by Offset increasing
+		TArray<FTocEntry> EntriesToInsert; // sorted by OriginalOffset increasing
+		PinChunksEntries(SourceToc, ReferenceToc, ReferenceTocChunkIdToIndex, EntriesToKeep, EntriesToInsert);
 
 		// only do the layout if there are entries to be kept in place
 		if (EntriesToKeep.IsEmpty())
@@ -839,127 +848,106 @@ public:
 		{
 			OutPinnedSizeBeforeDefrag += Entry.Size;
 		}
-		OutPinnedSizeAfterDefrag = OutPinnedSizeBeforeDefrag;
+
+		// we rely on Settings.MaxPartitionSize
+		ensure(Settings.MaxPartitionSize == 0 || Settings.MaxPartitionSize == UINT64_MAX || Settings.MaxPartitionSize == SourceToc.Header.PartitionSize);
 
 		const FBFDLayoutResult ResultBeforeDefrag = BestFitDecreasing(EntriesToKeep, EntriesToInsert, Settings);
 
-		// continue if we fit into growth limit
-		if (Settings.PatchInPlaceMaximumPerContainerGrowth == 0 ||
-			ResultBeforeDefrag.SizeAfter <= SizeBeforeLayout + Settings.PatchInPlaceMaximumPerContainerGrowth)
+		// attempt to defrag if container grew above size limit 
+		if ((Settings.PatchInPlaceMaximumPerContainerGrowth != 0) &&
+			(ResultBeforeDefrag.SizeAfter > SizeBeforeLayout + Settings.PatchInPlaceMaximumPerContainerGrowth))
 		{
-			return true;
-		}
+			TMap<FIoChunkId, int32> UnpinCandidates;
+			BuildUnpinCandidates(UnpinCandidates, EntriesToKeep, EntriesToInsert, Settings);
 
-		TMap<FIoChunkId, int32> UnpinCandidates;
-		BuildUnpinCandidates(UnpinCandidates, EntriesToKeep, EntriesToInsert, Settings);
+			if (bParanoidChecks)
+			{
+				// unpinning all should result in valid pareto front that doesn't go over existing size
+				TArray<FParetoPoint> ParetoFrontTemp;
+				ProbeParetoFront(
+					EntriesToKeep,
+					EntriesToInsert,
+					UnpinCandidates,
+					SizeBeforeLayout,
+					Settings,
+					ParetoFrontTemp,
+					UnpinCandidates.Num() - 1,
+					UnpinCandidates.Num(),
+					1);
+				ensure(!ParetoFrontTemp.IsEmpty());
+				ensure(ParetoFrontTemp.Last().UnpinCount == UnpinCandidates.Num());
+				ensure(ParetoFrontTemp.Last().LayoutResult.SizeAfter <= SizeBeforeLayout);
+			}
 
-		if (bParanoidChecks)
-		{
-			// unpinning all should result in valid pareto front that doesn't go over existing size
-			TArray<FParetoPoint> ParetoFrontTemp;
+			TArray<FParetoPoint> ParetoFront;
+
+			// probe whole unpin set with coarse steps to collect initial pareto front
 			ProbeParetoFront(
 				EntriesToKeep,
 				EntriesToInsert,
 				UnpinCandidates,
-				SizeBeforeLayout,
+				SizeBeforeLayout + Settings.PatchInPlaceMaximumPerContainerGrowth,
 				Settings,
-				ParetoFrontTemp,
-				UnpinCandidates.Num() - 1,
+				ParetoFront,
+				0,
 				UnpinCandidates.Num(),
-				1);
-			ensure(!ParetoFrontTemp.IsEmpty());
-			ensure(ParetoFrontTemp.Last().UnpinCount == UnpinCandidates.Num());
-			ensure(ParetoFrontTemp.Last().LayoutResult.SizeAfter <= SizeBeforeLayout);
-		}
-
-		TArray<FParetoPoint> ParetoFront;
-
-		// probe whole unpin set with coarse steps to collect initial pareto front
-		ProbeParetoFront(
-			EntriesToKeep,
-			EntriesToInsert,
-			UnpinCandidates,
-			SizeBeforeLayout + Settings.PatchInPlaceMaximumPerContainerGrowth,
-			Settings,
-			ParetoFront,
-			0,
-			UnpinCandidates.Num(),
-			Settings.PatchInPlaceParetoFrontCoarseSamples);
-
-		// refine pareto front if found, re-sample whole set otherwise
-		ProbeParetoFront(
-			EntriesToKeep,
-			EntriesToInsert,
-			UnpinCandidates,
-			SizeBeforeLayout + Settings.PatchInPlaceMaximumPerContainerGrowth,
-			Settings,
-			ParetoFront,
-			ParetoFront.Num() > 1 ? ParetoFront[0].UnpinCount : 0,
-			ParetoFront.Num() > 1 ? ParetoFront.Last().UnpinCount : UnpinCandidates.Num(),
-			Settings.PatchInPlaceParetoFrontFineSamples);
-
-		// chose smallest unpin penalty size
-		const FParetoPoint* BestPoint = Algo::MinElementBy(ParetoFront, [](const FParetoPoint& X){return X.LayoutResult.BytesWritten;});
-
-		if (!BestPoint)
-		{
-			OutPinnedSizeBeforeDefrag = 0;
-			OutPinnedSizeAfterDefrag = 0;
-			UE_LOGF(LogIoStore, Warning, "Patch-in-place '%ls' container grew above limit (%.2f MiB before pip -> %.2f MiB after pip), failed to fit into growth limit, disabling patch-in-place",
-				*ReferenceContainerPath,
-				(double)SizeBeforeLayout / 1024.0 / 1024.0,
-				(double)ResultBeforeDefrag.SizeAfter / 1024.0 / 1024.0);
-			return false;
-		}
+				Settings.PatchInPlaceParetoFrontCoarseSamples);
 
-		// preserve EntriesToKeep order
-		EntriesToKeep.RemoveAll([&UnpinCandidates, &BestPoint, this](const FTocEntry& Entry)
-		{
-			const int32* UnpinIndex = UnpinCandidates.Find(Entry.Id);
-			if (!UnpinIndex || *UnpinIndex >= BestPoint->UnpinCount)
+			// refine pareto front if found, re-sample whole set otherwise
+			ProbeParetoFront(
+				EntriesToKeep,
+				EntriesToInsert,
+				UnpinCandidates,
+				SizeBeforeLayout + Settings.PatchInPlaceMaximumPerContainerGrowth,
+				Settings,
+				ParetoFront,
+				ParetoFront.Num() > 1 ? ParetoFront[0].UnpinCount : 0,
+				ParetoFront.Num() > 1 ? ParetoFront.Last().UnpinCount : UnpinCandidates.Num(),
+				Settings.PatchInPlaceParetoFrontFineSamples);
+
+			// chose smallest unpin penalty size
+			const FParetoPoint* BestPoint = Algo::MinElementBy(ParetoFront, [](const FParetoPoint& X){return X.LayoutResult.BytesWritten;});
+
+			if (!BestPoint)
 			{
+				UE_LOGF(LogIoStore, Warning, "Patch-in-place '%ls' container grew above limit (%.2f MiB before pip -> %.2f MiB after pip), failed to fit into growth limit, disabling patch-in-place",
+					*ReferenceContainerPath,
+					(double)SizeBeforeLayout / 1024.0 / 1024.0,
+					(double)ResultBeforeDefrag.SizeAfter / 1024.0 / 1024.0);
 				return false;
 			}
 
-			EntriesToInsert.Add(Entry.Unpin());
-			return true;
-		});
-
-		Algo::SortBy(EntriesToInsert, &FTocEntry::OriginalOffset);
-
-		OutPinnedSizeAfterDefrag = 0;
-		for (const FTocEntry& Entry : EntriesToKeep)
-		{
-			OutPinnedSizeAfterDefrag += Entry.Size;
-		}
-
-		// redo BFD on the new set
-		const FBFDLayoutResult ResultAfterDefrag = BestFitDecreasing(EntriesToKeep, EntriesToInsert, Settings);
+			// preserve EntriesToKeep order
+			EntriesToKeep.RemoveAll([&EntriesToInsert, &UnpinCandidates, &BestPoint, this](const FTocEntry& Entry)
+			{
+				const int32* UnpinIndex = UnpinCandidates.Find(Entry.Id);
+				if (!UnpinIndex || *UnpinIndex >= BestPoint->UnpinCount)
+				{
+					return false;
+				}
 
-		UE_LOGF(LogIoStore, Display,
-			"Patch-in-place '%ls' container grew above limit (%.2f MiB before pip -> %.2f MiB after pip), defragmented by unpinning %.2f MiB, new container size %.2f MiB",
-			*ReferenceContainerPath,
-			(double)SizeBeforeLayout / 1024.0 / 1024.0,
-			(double)ResultBeforeDefrag.SizeAfter / 1024.0 / 1024.0,
-			(double)((int64)ResultAfterDefrag.BytesWritten - (int64)ResultBeforeDefrag.BytesWritten) / 1024.0 / 1024.0,
-			(double)ResultAfterDefrag.SizeAfter / 1024.0 / 1024.0
-		);
+				EntriesToInsert.Add(Entry.Unpin());
+				return true;
+			});
 
-		return true;
-	}
+			Algo::SortBy(EntriesToInsert, &FTocEntry::OriginalOffset);
 
-	void RunLayout(
-		FIoStoreTocResource& TargetToc,
-		int32& OutTargetPartitionNum,
-		uint64& OutTotalPadding,
-		int64& OutLocalityChange
-	)
-	{
-		TRACE_CPUPROFILER_EVENT_SCOPE(PatchInPlaceLayout);
+			// redo BFD on the new set
+			const FBFDLayoutResult ResultAfterDefrag = BestFitDecreasing(EntriesToKeep, EntriesToInsert, Settings);
 
-		UE_LOGF(LogIoStore, Display, "Patch-in-place '%ls' -> '%ls'", *ReferenceContainerPath, *TargetContainerPath);
+			UE_LOGF(LogIoStore, Display,
+				"Patch-in-place '%ls' container grew above limit (%.2f MiB before pip -> %.2f MiB after pip), defragmented by unpinning %.2f MiB, new container size %.2f MiB",
+				*ReferenceContainerPath,
+				(double)SizeBeforeLayout / 1024.0 / 1024.0,
+				(double)ResultBeforeDefrag.SizeAfter / 1024.0 / 1024.0,
+				(double)((int64)ResultAfterDefrag.BytesWritten - (int64)ResultBeforeDefrag.BytesWritten) / 1024.0 / 1024.0,
+				(double)ResultAfterDefrag.SizeAfter / 1024.0 / 1024.0
+			);
+		}
 
-		TArray<FTocEntry> TargetEntries(MoveTemp(EntriesToKeep));
+		// finalize TargetEntries
+		TargetEntries = MoveTemp(EntriesToKeep);
 		TargetEntries.Append(MoveTemp(EntriesToInsert));
 		Algo::SortBy(TargetEntries, [](const FTocEntry& X) {return X.Offset;});
 		if (bParanoidChecks)
@@ -967,7 +955,32 @@ public:
 			EnsureSortedAndNonOverlapping(TargetEntries);
 		}
 
-		FitIntoPartitions(TargetEntries, SourceToc.Header.PartitionSize, OutTotalPadding);
+		// finalize resulting values
+		uint64 PreviousEntryEnd = 0;
+		for (const FTocEntry& Entry : TargetEntries)
+		{
+			OutTargetPartitionNum = FMath::Max(Entry.PartitionIndexBegin(Settings.MaxPartitionSize) + 1, OutTargetPartitionNum);
+
+			OutTotalPadding += Entry.Offset - PreviousEntryEnd;
+			PreviousEntryEnd = Entry.End();
+
+			// at this point no chunk shall cross partition boundary
+			if (!ensure(Entry.PartitionIndexBegin(Settings.MaxPartitionSize) == Entry.PartitionIndexEnd(Settings.MaxPartitionSize)))
+			{
+				UE_LOGF(LogIoStore, Warning, "Patch-in-place '%ls' container failed to partition chunk '%ls', disabling patch-in-place",
+					*ReferenceContainerPath,
+					*LexToString(Entry.Id));
+				return false;
+			}
+
+			if (Entry.bMemoryMapped && !IsAligned(Entry.Offset, Settings.MemoryMappingAlignment))
+			{
+				UE_LOGF(LogIoStore, Warning, "Patch-in-place '%ls' container failed to align chunk '%ls', disabling patch-in-place",
+					*ReferenceContainerPath,
+					*LexToString(Entry.Id));
+				return false;
+			}
+		}
 
 		uint64 LocalityBefore = 0;
 		uint64 LocalityAfter = 0;
@@ -981,12 +994,27 @@ public:
 			(double)LocalityAfter / 1024.0 / 1024.0,
 			(double)OutLocalityChange / 1024.0 / 1024.0);
 
+		return true;
+	}
+
+	void RunLayout(
+		FIoStoreTocResource& TargetToc,
+		const int32 TargetPartitionNum,
+		uint64& OutActualPinnedSize
+	)
+	{
+		TRACE_CPUPROFILER_EVENT_SCOPE(PatchInPlaceLayout);
+
+		UE_LOGF(LogIoStore, Display, "Patch-in-place '%ls' -> '%ls'", *ReferenceContainerPath, *TargetContainerPath);
+
 		TArray<FFileOp> FileOps = {};
+		int32 TargetPartitionNumCheck = 0;
 		int32 SourcePartitionNum = 0;
-		BuildTargetToc(TargetToc, FileOps, TargetEntries, SourcePartitionNum, OutTargetPartitionNum);
-		WriteChunksToTarget(TargetEntries, FileOps, SourcePartitionNum, OutTargetPartitionNum);
+		BuildTargetToc(TargetToc, FileOps, TargetEntries, SourcePartitionNum, TargetPartitionNumCheck);
+		verify(TargetPartitionNumCheck == TargetPartitionNum);
+		WriteChunksToTarget(TargetEntries, FileOps, SourcePartitionNum, TargetPartitionNum, OutActualPinnedSize);
 
-		for (int32 Index = 0; Index < OutTargetPartitionNum; ++Index)
+		for (int32 Index = 0; Index < TargetPartitionNum; ++Index)
 		{
 			OverwriteFile(GetPartitionFilePath(SourceContainerPath, Index), GetPartitionFilePath(TargetContainerPath, Index));
 		}
@@ -1006,8 +1034,7 @@ private:
 	TMap<FIoChunkId, int32> SourceTocChunkIdToIndex;
 
 	struct FTocEntry;
-	TArray<FTocEntry> EntriesToKeep; // sorted by Offset increasing
-	TArray<FTocEntry> EntriesToInsert; // sorted by OriginalOffset increasing
+	TArray<FTocEntry> TargetEntries; // sorted by Offset increasing
 
 	static void BuildTocMap(const FIoStoreTocResourceView& Toc, TMap<FIoChunkId, int32>& OutResult)
 	{
@@ -1041,6 +1068,16 @@ private:
 			return Offset + Size;
 		}
 
+		int32 PartitionIndexBegin(const uint64 PartitionSize) const
+		{
+			return PartitionSize > 0 && PartitionSize < UINT64_MAX ? Offset / PartitionSize : 0;
+		}
+
+		int32 PartitionIndexEnd(const uint64 PartitionSize) const
+		{
+			return PartitionSize > 0 && PartitionSize < UINT64_MAX ? (Size > 0 ? Offset + Size - 1 : Offset) / PartitionSize : 0;
+		}
+
 		FTocEntry Unpin() const
 		{
 			return FTocEntry(Id, UINT64_MAX, OriginalOffset, Size, bMemoryMapped);
@@ -1075,21 +1112,23 @@ private:
 			Size += InSize;
 		}
 
-		uint64 CalculateEnd(uint64 AtOffset, const TConstArrayView<FTocEntry>& InEntriesToInsert, const uint64 MemoryMappingAlignment) const
+		uint64 CalculateEnd(uint64 AtOffset, const TConstArrayView<FTocEntry>& InEntriesToInsert, const uint64 MemoryMappingAlignment, const uint64 PartitionSize) const
 		{
 			for (int32 I = Index; I < Index + Count; ++I)
 			{
-				AtOffset = AlignOffset(AtOffset, InEntriesToInsert[I].bMemoryMapped, MemoryMappingAlignment) + InEntriesToInsert[I].Size;
+				AtOffset = AlignOffset(AtOffset, InEntriesToInsert[I].Size, InEntriesToInsert[I].bMemoryMapped, MemoryMappingAlignment, PartitionSize);
+				AtOffset += InEntriesToInsert[I].Size;
 			}
 
 			return AtOffset;
 		}
 
-		uint64 InsertEntries(uint64 AtOffset, const TArrayView<FTocEntry>& InEntriesToInsert, const uint64 MemoryMappingAlignment) const
+		uint64 InsertEntries(uint64 AtOffset, const TArrayView<FTocEntry>& InEntriesToInsert, const uint64 MemoryMappingAlignment, const uint64 PartitionSize) const
 		{
 			for (int32 I = Index; I < Index + Count; ++I)
 			{
-				InEntriesToInsert[I].Offset = AlignOffset(AtOffset, InEntriesToInsert[I].bMemoryMapped, MemoryMappingAlignment);
+				AtOffset = AlignOffset(AtOffset, InEntriesToInsert[I].Size, InEntriesToInsert[I].bMemoryMapped, MemoryMappingAlignment, PartitionSize);
+				InEntriesToInsert[I].Offset = AtOffset;
 				AtOffset = InEntriesToInsert[I].End();
 			}
 
@@ -1113,15 +1152,15 @@ private:
 			return End >= Begin ? End - Begin : 0;
 		}
 
-		bool Fits(const FTocEntry& Entry, const uint64 MemoryMappingAlignment) const
+		bool Fits(const FTocEntry& Entry, const uint64 MemoryMappingAlignment, const uint64 PartitionSize) const
 		{
-			const uint64 NewBegin = AlignOffset(Begin, Entry.bMemoryMapped, MemoryMappingAlignment);
+			const uint64 NewBegin = AlignOffset(Begin, Entry.Size, Entry.bMemoryMapped, MemoryMappingAlignment, PartitionSize);
 			return NewBegin + Entry.Size <= End;
 		}
 
-		uint64 Place(const FTocEntry& Entry, const uint64 MemoryMappingAlignment)
+		uint64 Place(const FTocEntry& Entry, const uint64 MemoryMappingAlignment, const uint64 PartitionSize)
 		{
-			const uint64 NewBegin = AlignOffset(Begin, Entry.bMemoryMapped, MemoryMappingAlignment);
+			const uint64 NewBegin = AlignOffset(Begin, Entry.Size, Entry.bMemoryMapped, MemoryMappingAlignment, PartitionSize);
 			Begin = NewBegin + Entry.Size;
 			return NewBegin;
 		}
@@ -1144,22 +1183,22 @@ private:
 
 		int32 SourcePartitionIndex(const uint64 PartitionSize) const
 		{
-			return PartitionSize > 0 ? SourceOffset / PartitionSize : 0;
+			return PartitionSize > 0 && PartitionSize < UINT64_MAX ? SourceOffset / PartitionSize : 0;
 		}
 
 		int32 TargetPartitionIndex(const uint64 PartitionSize) const
 		{
-			return PartitionSize > 0 ? TargetOffset / PartitionSize : 0;
+			return PartitionSize > 0 && PartitionSize < UINT64_MAX ? TargetOffset / PartitionSize : 0;
 		}
 
 		uint64 SourcePartitionOffset(const uint64 PartitionSize) const
 		{
-			return PartitionSize > 0 ? SourceOffset % PartitionSize : SourceOffset;
+			return PartitionSize > 0 && PartitionSize < UINT64_MAX ? SourceOffset % PartitionSize : SourceOffset;
 		}
 
 		uint64 TargetPartitionOffset(const uint64 PartitionSize) const
 		{
-			return PartitionSize > 0 ? TargetOffset % PartitionSize : TargetOffset;
+			return PartitionSize > 0 && PartitionSize < UINT64_MAX ? TargetOffset % PartitionSize : TargetOffset;
 		}
 
 		bool CanCoalesce(const uint64 SourceReadOffset, const uint64 TargetWriteOffset, const uint64 PartitionSize) const
@@ -1211,7 +1250,13 @@ private:
 		}
 	};
 
-	void PinChunksEntries()
+	static void PinChunksEntries(
+		const FIoStoreTocResourceView& SourceToc,
+		const FIoStoreTocResourceView& ReferenceToc,
+		const TMap<FIoChunkId, int32>& ReferenceTocChunkIdToIndex,
+		TArray<FTocEntry>& OutEntriesToKeep,
+		TArray<FTocEntry>& OutEntriesToInsert
+	)
 	{
 		for (int32 ChunkIndex = 0; ChunkIndex < SourceToc.ChunkIds.Num(); ChunkIndex++)
 		{
@@ -1220,24 +1265,24 @@ private:
 			const int32* ReferenceTocIndex = ReferenceTocChunkIdToIndex.Find(SourceChunk.Id);
 			if (!ReferenceTocIndex)
 			{
-				EntriesToInsert.Add(FTocEntry::Unpinned(SourceChunk));
+				OutEntriesToInsert.Add(FTocEntry::Unpinned(SourceChunk));
 				continue;
 			}
 
-			const FIoStoreTocChunkInfo ReferenceChunk = ReferenceToc.GetTocChunkInfo(*ReferenceTocIndex);
+			const FIoStoreTocChunkInfo& ReferenceChunk = ReferenceToc.GetTocChunkInfo(*ReferenceTocIndex);
 			if (SourceChunk.ChunkHash == ReferenceChunk.ChunkHash &&
 				SourceChunk.SizeOnDisk == ReferenceChunk.SizeOnDisk)
 			{
-				EntriesToKeep.Add(FTocEntry::Pinned(SourceChunk, ReferenceChunk.OffsetOnDisk));
+				OutEntriesToKeep.Add(FTocEntry::Pinned(SourceChunk, ReferenceChunk.OffsetOnDisk));
 			}
 			else
 			{
-				EntriesToInsert.Add(FTocEntry::Unpinned(SourceChunk));
+				OutEntriesToInsert.Add(FTocEntry::Unpinned(SourceChunk));
 			}
 		}
 
-		Algo::SortBy(EntriesToKeep, &FTocEntry::Offset);
-		Algo::SortBy(EntriesToInsert, &FTocEntry::OriginalOffset);
+		Algo::SortBy(OutEntriesToKeep, &FTocEntry::Offset);
+		Algo::SortBy(OutEntriesToInsert, &FTocEntry::OriginalOffset);
 	}
 
 	static uint64 GapsBinarySearchProjection(const FLayoutGap& Gap)
@@ -1245,13 +1290,21 @@ private:
 		return Gap.Size();
 	}
 
-	static uint64 AlignOffset(const uint64 Offset, const bool bMemoryMapped, const uint64 MemoryMappingAlignment)
+	static uint64 AlignOffset(const uint64 Offset, const uint64 Size, const bool bMemoryMapped, const uint64 MemoryMappingAlignment, const uint64 PartitionSize)
 	{
 		uint64 Result = Offset;
 		if (bMemoryMapped && MemoryMappingAlignment > 0)
 		{
 			Result = Align(Result, MemoryMappingAlignment);
 		}
+		if ((PartitionSize > 0) && (PartitionSize < UINT64_MAX) && (Size > 0) && (Result / PartitionSize) != ((Result + Size - 1) / PartitionSize))
+		{
+			Result = ((Result / PartitionSize) + 1) * PartitionSize;
+			if (bMemoryMapped && MemoryMappingAlignment > 0) // align again if partition size is not a modulo of memory mapping alignment 
+			{
+				Result = Align(Result, MemoryMappingAlignment);
+			}
+		}
 		return Result;
 	}
 
@@ -1352,13 +1405,13 @@ private:
 
 					for (const FTocEntry& Insert : EntriesToInsert)
 					{
-						if (!Gap.Fits(Insert, Settings.MemoryMappingAlignment))
+						if (!Gap.Fits(Insert, Settings.MemoryMappingAlignment, Settings.MaxPartitionSize))
 						{
 							continue;
 						}
 
 						UsableSpace += Insert.Size;
-						Gap.Place(Insert, Settings.MemoryMappingAlignment);
+						Gap.Place(Insert, Settings.MemoryMappingAlignment, Settings.MaxPartitionSize);
 
 						if (Gap.Size() == 0)
 						{
@@ -1421,6 +1474,7 @@ private:
 		Algo::SortBy(Gaps, GapsBinarySearchProjection);
 
 		const uint64 MemoryMappingAlignment = Settings.MemoryMappingAlignment;
+		const uint64 PartitionSize = Settings.MaxPartitionSize;
 		const uint64 Granularity = Settings.PatchInPlaceReadGranularity;
 		const uint64 MaxChainSize = Settings.PatchInPlaceMaxChainSize;
 		const int32 MaxChainLength = Settings.PatchInPlaceMaxChainLength;
@@ -1469,7 +1523,7 @@ private:
 				.RightChop(StartIndex)
 				.IndexOfByPredicate([&](const FLayoutGap& Gap)
 			{
-				return Chain.CalculateEnd(Gap.Begin, EntriesToInsert, MemoryMappingAlignment) <= Gap.End;
+				return Chain.CalculateEnd(Gap.Begin, EntriesToInsert, MemoryMappingAlignment, PartitionSize) <= Gap.End;
 			});
 			verify(RelativeIndex != INDEX_NONE);
 
@@ -1477,7 +1531,7 @@ private:
 
 			FLayoutGap Gap = Gaps[BestGapIndex];
 			Gaps.RemoveAt(BestGapIndex, EAllowShrinking::No);
-			Gap.Begin = Chain.InsertEntries(Gap.Begin, EntriesToInsert, MemoryMappingAlignment);
+			Gap.Begin = Chain.InsertEntries(Gap.Begin, EntriesToInsert, MemoryMappingAlignment, PartitionSize);
 
 			if (Gap.Size() > 0)
 			{
@@ -1576,32 +1630,6 @@ private:
 		return true;
 	}
 
-	static void FitIntoPartitions(TArray<FTocEntry>& Chunks, const uint64 PartitionSize, uint64& OutTotalPadding)
-	{
-		OutTotalPadding = 0;
-
-		uint64 Padding = 0;
-		uint64 PreviousEnd = 0;
-
-		for (FTocEntry& Chunk : Chunks)
-		{
-			ensure(Chunk.Size < PartitionSize);
-			const uint64 BeginIndex = (Chunk.Offset + Padding) / PartitionSize;
-			const uint64 EndIndex = (Chunk.Offset + Chunk.Size + Padding) / PartitionSize;
-
-			if (EndIndex != BeginIndex)
-			{
-				// bump chunk to the next partition
-				Padding += EndIndex * PartitionSize - (Chunk.Offset + Padding);
-			}
-
-			Chunk.Offset += Padding;
-
-			OutTotalPadding += Chunk.Offset - PreviousEnd;
-			PreviousEnd = Chunk.Offset + Chunk.Size;
-		}
-	}
-
 	// returns amount of bytes that are prefetched when reading all entries in their provided order
 	static void CalculateLocalityChange(
 		TConstArrayView<FTocEntry> Chunks,
@@ -1812,10 +1840,12 @@ private:
 		}
 	}
 
-	void WriteChunksToTarget(TConstArrayView<FTocEntry> Chunks, TConstArrayView<FFileOp> FileOps, const int32 SourcePartitionNum, const int32 TargetPartitionNum)
+	void WriteChunksToTarget(TConstArrayView<FTocEntry> Chunks, TConstArrayView<FFileOp> FileOps, const int32 SourcePartitionNum, const int32 TargetPartitionNum, uint64& OutActualPinnedSize)
 	{
 		TRACE_CPUPROFILER_EVENT_SCOPE(WriteChunksToTarget);
 
+		OutActualPinnedSize = 0;
+
 		if (Chunks.IsEmpty() || SourcePartitionNum == 0 || TargetPartitionNum == 0 || FileOps.IsEmpty())
 		{
 			return;
@@ -1828,7 +1858,7 @@ private:
 			return;
 		}
 
-		const uint64 LastPartitionSize = 
+		const uint64 LastPartitionSize =
 			MaxPartitionSize > 0
 			? (LastPartitionSizeLogical % MaxPartitionSize == 0
 				? MaxPartitionSize // edge case, LastPartitionSizeLogical precisely matches partition size
@@ -1861,7 +1891,7 @@ private:
 			FFileOpenResult Result = Ipf.OpenWrite(*PartitionFilePath, IPlatformFile::EOpenWriteFlags::AllowRead | IPlatformFile::EOpenWriteFlags::Append);
 			if (!Result.IsValid())
 			{
-				UE_LOGF(LogIoStore, Fatal, "Failed to read partition '%ls'", *PartitionFilePath);
+				UE_LOGF(LogIoStore, Fatal, "Failed to write partition '%ls'", *PartitionFilePath);
 				return;
 			}
 			TargetPartitions[I] = MoveTemp(Result.GetValue());
@@ -1871,10 +1901,10 @@ private:
 			}
 		}
 
-		constexpr uint64 TempBufferSize = 4 * 1024 * 1024; 
+		constexpr uint64 TempBufferSize = 4 * 1024 * 1024;
 		FIoBuffer TempBuffer(TempBufferSize);
 
-		for (const FFileOp& Op : FileOps) 
+		for (const FFileOp& Op : FileOps)
 		{
 			IFileHandle* Source = Op.bZeroFill ? nullptr : SourcePartitions[Op.SourcePartitionIndex(MaxPartitionSize)].Get();
 			IFileHandle* Target = TargetPartitions[Op.TargetPartitionIndex(MaxPartitionSize)].Get();
@@ -1913,7 +1943,7 @@ private:
 		{
 			FIoBuffer TempBuffer2(TempBufferSize);
 
-			for (const FFileOp& Op : FileOps) 
+			for (const FFileOp& Op : FileOps)
 			{
 				IFileHandle* Source = Op.bZeroFill ? nullptr : SourcePartitions[Op.SourcePartitionIndex(MaxPartitionSize)].Get();
 				IFileHandle* Target = TargetPartitions[Op.TargetPartitionIndex(MaxPartitionSize)].Get();
@@ -1939,6 +1969,104 @@ private:
 				}
 			}
 		}
+
+		// verifying new container against the reference to ensure correct pinned byte count
+		TRACE_CPUPROFILER_EVENT_SCOPE(VerifyPinnedChunks);
+
+		const uint64 ReferencePartitionSize = ReferenceToc.Header.PartitionSize;
+
+		TArray<TUniquePtr<IFileHandle>> ReferencePartitions;
+		ReferencePartitions.SetNum((int32)ReferenceToc.Header.PartitionCount);
+		for (int32 I = 0; I < ReferencePartitions.Num(); ++I)
+		{
+			const FString Path = GetPartitionFilePath(ReferenceContainerPath, I);
+			FFileOpenResult Result = Ipf.OpenRead(*Path, IPlatformFile::EOpenReadFlags::None);
+			if (Result.IsValid())
+			{
+				ReferencePartitions[I] = MoveTemp(Result.GetValue());
+			}
+			else
+			{
+				UE_LOGF(LogIoStore, Warning,
+					"Patch-in-place verify '%ls': failed to read '%ls', PIP pinned bytes might be wrong",
+					*SourceContainerPath,
+					*Path);
+				break;
+			}
+		}
+
+		FIoBuffer ReferenceReadBuffer(TempBufferSize);
+		FIoBuffer TargetReadBuffer(TempBufferSize);
+
+		uint64 PinnedChunksMismatched = 0;
+
+		for (const FTocEntry& Chunk : Chunks)
+		{
+			const int32* ReferenceIndex = ReferenceTocChunkIdToIndex.Find(Chunk.Id);
+			if (!ReferenceIndex || !ensure(SourceTocChunkIdToIndex.Contains(Chunk.Id)))
+			{
+				continue;
+			}
+
+			const FIoStoreTocChunkInfo& ReferenceChunk = ReferenceToc.GetTocChunkInfo(*ReferenceIndex);
+			const FIoStoreTocChunkInfo& SourceChunk = SourceToc.GetTocChunkInfo(*SourceTocChunkIdToIndex.Find(Chunk.Id));
+
+			if (Chunk.Offset != ReferenceChunk.OffsetOnDisk ||
+				Chunk.Size != ReferenceChunk.SizeOnDisk ||
+				SourceChunk.ChunkHash != ReferenceChunk.ChunkHash)
+			{
+				continue;
+			}
+
+			const int32 TargetPartitionIndex = Chunk.PartitionIndexBegin(MaxPartitionSize);
+			const uint64 TargetPartitionOffset = MaxPartitionSize > 0 && MaxPartitionSize < UINT64_MAX ? Chunk.Offset % MaxPartitionSize : Chunk.Offset;
+			const int32 ReferencePartitionIndex = ReferencePartitionSize > 0 && ReferencePartitionSize < UINT64_MAX ? ReferenceChunk.OffsetOnDisk / ReferencePartitionSize : 0;
+			const uint64 ReferencePartitionOffset = ReferencePartitionSize > 0 && ReferencePartitionSize < UINT64_MAX ? ReferenceChunk.OffsetOnDisk % ReferencePartitionSize : ReferenceChunk.OffsetOnDisk;
+
+			if (!ensure(TargetPartitionIndex < TargetPartitions.Num()))
+			{
+				continue;
+			}
+
+			if (ReferencePartitionIndex >= ReferencePartitions.Num() || !ReferencePartitions[ReferencePartitionIndex].IsValid())
+			{
+				// assume it's valid as we can't compare the bytes against the reference
+				OutActualPinnedSize += Chunk.Size;
+				continue;
+			}
+
+			bool bMatch = true;
+			for (uint64 Offset = 0; Offset < Chunk.Size; Offset += TempBufferSize)
+			{
+				const uint64 ReadSize = FMath::Min(Chunk.Size - Offset, TempBufferSize);
+
+				TargetPartitions[TargetPartitionIndex]->ReadAt(TargetReadBuffer.GetData(), ReadSize, TargetPartitionOffset + Offset);
+				ReferencePartitions[ReferencePartitionIndex]->ReadAt(ReferenceReadBuffer.GetData(), ReadSize, ReferencePartitionOffset + Offset);
+
+				if (FMemory::Memcmp(TargetReadBuffer.GetData(), ReferenceReadBuffer.GetData(), ReadSize) != 0)
+				{
+					bMatch = false;
+					break;
+				}
+			}
+
+			if (bMatch)
+			{
+				OutActualPinnedSize += Chunk.Size;
+			}
+			else
+			{
+				PinnedChunksMismatched++;
+			}
+		}
+
+		if (PinnedChunksMismatched > 0)
+		{
+			UE_LOGF(LogIoStore, Display,
+				"Patch-in-place verify '%ls': %llu of pinned chunks didn't binary-match reference",
+				*SourceContainerPath,
+				PinnedChunksMismatched);
+		}
 	}
 
 	static FString GetPartitionFilePath(FString ContainerFilePath, int32 PartitionIndex)
@@ -2546,7 +2674,13 @@ public:
 			WriterSettings
 		);
 
-		if (!Layout.PrepareForLayout(PatchInPlaceSizeBeforeLayout, PatchInPlacePinnedSizeBeforeDefrag, PatchInPlacePinnedSizeAfterDefrag))
+		if (!Layout.PrepareForLayout(
+			PatchInPlaceSizeBeforeLayout,
+			PatchInPlacePinnedSizeBeforeDefrag,
+			TargetPartitionNum,
+			TotalPaddingSize,
+			PatchInPlaceLocalityChange
+		))
 		{
 			return;
 		}
@@ -2554,7 +2688,7 @@ public:
 		// release file handles for patch-in-place
 		Partitions.Empty();
 
-		Layout.RunLayout(TargetToc, TargetPartitionNum, TotalPaddingSize, PatchInPlaceLocalityChange);
+		Layout.RunLayout(TargetToc, TargetPartitionNum, PatchInPlacePinnedSizeAfterDefrag);
 
 		SourceToc = MoveTemp(TargetToc);
 
