UnrealPak: Fix unsafe path traversal outside the extract dir in the extract command for iostore containers.

#rnx
#Synced-CL 52611489

diff --git a/Engine/Source/Developer/IoStoreUtilities/Private/IoStoreUtilities.cpp b/Engine/Source/Developer/IoStoreUtilities/Private/IoStoreUtilities.cpp
--- a/Engine/Source/Developer/IoStoreUtilities/Private/IoStoreUtilities.cpp
+++ b/Engine/Source/Developer/IoStoreUtilities/Private/IoStoreUtilities.cpp
@@ -8430,6 +8430,22 @@ bool ExtractFilesFromIoStoreContainer(
 	return ProcessFilesFromIoStoreContainer(InContainerFilename, InDestPath, InKeyChain, InFilter, ExtractFilesWriter, OutOrderMap, OutUsedEncryptionKeys, bOutIsSigned, -1);
 }
 
+FString MakeSafeDestEntry(const FString& NormalizedDestRoot, const FString& EntryFileName)
+{
+	FString DestFilename = FPaths::ConvertRelativePathToFull(NormalizedDestRoot / EntryFileName);
+	if (!FPaths::IsUnderDirectory(DestFilename, NormalizedDestRoot))
+	{
+		FString EntryPath = FPaths::GetPath(EntryFileName);
+		uint64 EntryHash = FXxHash64::HashBuffer(*EntryPath, EntryPath.Len() * sizeof(TCHAR)).Hash;
+		FString HashDir = FString::Printf(TEXT("%016llx"), EntryHash);
+		FString SafeDestFilename = FPaths::ConvertRelativePathToFull(NormalizedDestRoot / HashDir / FPaths::GetCleanFilename(EntryFileName));
+		UE_LOG(LogPakFile, Display, TEXT("Fixup extract path for \"%s\": from \"%s\" to \"%s\""),
+				*EntryFileName, *DestFilename, *SafeDestFilename);
+		DestFilename = SafeDestFilename;
+	}
+	return DestFilename;
+}
+
 bool ProcessFilesFromIoStoreContainer(
 	const TCHAR* InContainerFilename,
 	const TCHAR* InDestPath,
@@ -8479,8 +8495,8 @@ bool ProcessFilesFromIoStoreContainer(
 	};
 	TArray<FEntry> Entries;
 	const FIoDirectoryIndexReader& IndexReader = Reader->GetDirectoryIndexReader();
-	FString DestPath(InDestPath);
-	Reader->EnumerateChunks([&Entries, InFilter, &DestPath](const FIoStoreTocChunkInfo& ChunkInfo)
+	FString NormalizedDestRoot = FPaths::ConvertRelativePathToFull(InDestPath);
+	Reader->EnumerateChunks([&Entries, InFilter, &NormalizedDestRoot](const FIoStoreTocChunkInfo& ChunkInfo)
 		{
 			if (!ChunkInfo.bHasValidFileName)
 			{
@@ -8492,10 +8508,15 @@ bool ProcessFilesFromIoStoreContainer(
 				return true;
 			}
 
+			FString DestFileName = ChunkInfo.FileName;
+			FPaths::NormalizeFilename(DestFileName);
+			DestFileName.ReplaceInline(TEXT("../../../"), TEXT(""));
+			DestFileName = MakeSafeDestEntry(NormalizedDestRoot, DestFileName);
+
 			FEntry& Entry = Entries.AddDefaulted_GetRef();
 			Entry.ChunkId = ChunkInfo.Id;
 			Entry.SourceFileName = ChunkInfo.FileName;
-			Entry.DestFileName = DestPath / ChunkInfo.FileName.Replace(TEXT("../../../"), TEXT(""));
+			Entry.DestFileName = DestFileName;
 			Entry.Offset = ChunkInfo.Offset;
 			Entry.bIsCompressed = ChunkInfo.bIsCompressed;
 
